Horde: Generalize the way includes are processed in configuration files.

The [ConfigInclude] attribute can be applied to any string property. When parsed, the referenced file will be parsed into the first outer object with a [ConfigIncludeRoot] attribute.

diff --git a/Engine/Source/Programs/Horde/Horde.Build.Tests/ConfigTests.cs b/Engine/Source/Programs/Horde/Horde.Build.Tests/ConfigTests.cs
--- a/Engine/Source/Programs/Horde/Horde.Build.Tests/ConfigTests.cs
+++ b/Engine/Source/Programs/Horde/Horde.Build.Tests/ConfigTests.cs
@@ -27,6 +27,7 @@ namespace Horde.Build.Tests
 			public SubObject? ValueC { get; set; }
 		}
 
+		[ConfigIncludeRoot]
 		class ConfigObject
 		{
 			public List<ConfigInclude> Include { get; set; } = new List<ConfigInclude>();
diff --git a/Engine/Source/Programs/Horde/Horde.Build/Configuration/ConfigType.cs b/Engine/Source/Programs/Horde/Horde.Build/Configuration/ConfigType.cs
--- a/Engine/Source/Programs/Horde/Horde.Build/Configuration/ConfigType.cs
+++ b/Engine/Source/Programs/Horde/Horde.Build/Configuration/ConfigType.cs
@@ -19,23 +19,27 @@ using JsonObject = System.Text.Json.Nodes.JsonObject;
 namespace Horde.Build.Configuration
 {
 	/// <summary>
-	/// Attribute used to mark <see cref="Uri"/> properties that are relative to their containing file
+	/// Attribute used to mark <see cref="Uri"/> properties that include other config files
 	/// </summary>
 	[AttributeUsage(AttributeTargets.Property)]
-	public sealed class ConfigPathAttribute : Attribute
+	public sealed class ConfigIncludeAttribute : Attribute
 	{
 	}
 
 	/// <summary>
-	/// Directive to merge config data from another source
+	/// Specifies that a class is the root for including other files
 	/// </summary>
-	public class ConfigInclude
+	[AttributeUsage(AttributeTargets.Class)]
+	public sealed class ConfigIncludeRootAttribute : Attribute
+	{
+	}
+
+	/// <summary>
+	/// Attribute used to mark <see cref="Uri"/> properties that are relative to their containing file
+	/// </summary>
+	[AttributeUsage(AttributeTargets.Property)]
+	public sealed class ConfigRelativePathAttribute : Attribute
 	{
-		/// <summary>
-		/// Path to the config data to be included. May be relative to the including file's location.
-		/// </summary>
-		[Required]
-		public string Path { get; set; } = null!;
 	}
 
 	/// <summary>
@@ -81,7 +85,7 @@ namespace Horde.Build.Configuration
 				}
 				else
 				{
-					value = ObjectConfigType.FindOrAdd(type);
+					value = ClassConfigType.FindOrAdd(type);
 				}
 
 				lock (s_typeToValueType)
@@ -102,7 +106,7 @@ namespace Horde.Build.Configuration
 		/// <returns></returns>
 		public static async Task<T> ReadAsync<T>(Uri uri, ConfigContext context, CancellationToken cancellationToken) where T : class, new()
 		{
-			ObjectConfigType type = ObjectConfigType.FindOrAdd(typeof(T));
+			ClassConfigType type = ClassConfigType.FindOrAdd(typeof(T));
 
 			T target = new T();
 			await type.MergeObjectAsync(target, uri, context, cancellationToken);
@@ -133,7 +137,7 @@ namespace Horde.Build.Configuration
 	}
 
 	/// <summary>
-	/// Implementation of <see cref="ConfigType"/> for scalar values
+	/// Implementation of <see cref="ConfigType"/> for scalar types
 	/// </summary>
 	class ScalarConfigType : ConfigType
 	{
@@ -162,9 +166,9 @@ namespace Horde.Build.Configuration
 	}
 
 	/// <summary>
-	/// Implementation of <see cref="ConfigType"/> to handle object variables
+	/// Implementation of <see cref="ConfigType"/> to handle class types
 	/// </summary>
-	class ObjectConfigType : ConfigType
+	class ClassConfigType : ConfigType
 	{
 		abstract class Property
 		{
@@ -177,7 +181,11 @@ namespace Horde.Build.Configuration
 				PropertyInfo = propertyInfo;
 			}
 
+			public abstract bool HasIncludes();
+
 			public abstract Task MergeAsync(object target, JsonNode? node, ConfigContext context, CancellationToken cancellationToken);
+
+			public abstract Task ParseIncludesAsync(JsonNode jsonNode, object targetObject, ClassConfigType targetType, ConfigContext context, CancellationToken cancellationToken);
 		}
 
 		class ScalarProperty : Property
@@ -187,12 +195,14 @@ namespace Horde.Build.Configuration
 			{
 			}
 
+			public override bool HasIncludes() => PropertyInfo.GetCustomAttribute<ConfigIncludeAttribute>() != null;
+
 			public override Task MergeAsync(object target, JsonNode? node, ConfigContext context, CancellationToken cancellationToken)
 			{
 				context.AddProperty(Name);
 
 				object? value;
-				if (PropertyInfo.GetCustomAttribute<ConfigPathAttribute>() != null)
+				if (PropertyInfo.GetCustomAttribute<ConfigRelativePathAttribute>() != null)
 				{
 					value = CombinePaths(context.CurrentFile, JsonSerializer.Deserialize<string>(node, context.JsonOptions) ?? String.Empty).AbsoluteUri;
 				}
@@ -204,6 +214,21 @@ namespace Horde.Build.Configuration
 
 				return Task.CompletedTask;
 			}
+
+			public override async Task ParseIncludesAsync(JsonNode jsonNode, object targetObject, ClassConfigType targetType, ConfigContext context, CancellationToken cancellationToken)
+			{
+				string? path = (string?)jsonNode;
+
+				Uri uri = ConfigType.CombinePaths(context.CurrentFile, path!);
+				IConfigFile file = await ReadFileAsync(uri, context, cancellationToken);
+
+				context.IncludeStack.Push(file);
+
+				JsonObject includedJsonObject = await ParseFileAsync(file, context, cancellationToken);
+				await targetType.MergeObjectAsync(targetObject, includedJsonObject, context, cancellationToken);
+
+				context.IncludeStack.Pop();
+			}
 		}
 
 		class ListProperty : Property
@@ -216,6 +241,8 @@ namespace Horde.Build.Configuration
 				_elementType = elementType;
 			}
 
+			public override bool HasIncludes() => _elementType is ClassConfigType elementType && elementType.HasIncludes();
+
 			public override async Task MergeAsync(object target, JsonNode? node, ConfigContext context, CancellationToken cancellationToken)
 			{
 				IList? list = (IList?)PropertyInfo.GetValue(target);
@@ -237,18 +264,32 @@ namespace Horde.Build.Configuration
 					context.LeaveScope();
 				}
 			}
+
+			public override async Task ParseIncludesAsync(JsonNode jsonNode, object targetObject, ClassConfigType targetType, ConfigContext context, CancellationToken cancellationToken)
+			{
+				if (jsonNode is JsonArray jsonArrayValue)
+				{
+					ClassConfigType classElementType = (ClassConfigType)_elementType;
+					foreach (JsonObject jsonObjectElement in jsonArrayValue.OfType<JsonObject>())
+					{
+						await classElementType.ParseIncludesAsync(jsonObjectElement, targetObject, targetType, context, cancellationToken);
+					}
+				}
+			}
 		}
 
 		class ObjectProperty : Property
 		{
-			readonly ObjectConfigType _classConfigType;
+			readonly ClassConfigType _classConfigType;
 
-			public ObjectProperty(string name, PropertyInfo propertyInfo, ObjectConfigType classConfigType)
+			public ObjectProperty(string name, PropertyInfo propertyInfo, ClassConfigType classConfigType)
 				: base(name, propertyInfo)
 			{
 				_classConfigType = classConfigType;
 			}
 
+			public override bool HasIncludes() => _classConfigType.HasIncludes();
+
 			public override async Task MergeAsync(object target, JsonNode? node, ConfigContext context, CancellationToken cancellationToken)
 			{
 				if (node is JsonObject obj)
@@ -282,43 +323,59 @@ namespace Horde.Build.Configuration
 					PropertyInfo.SetValue(target, JsonSerializer.Deserialize(node, PropertyInfo.PropertyType, context.JsonOptions));
 				}
 			}
+
+			public override async Task ParseIncludesAsync(JsonNode jsonNode, object targetObject, ClassConfigType targetType, ConfigContext context, CancellationToken cancellationToken)
+			{
+				if (jsonNode is JsonObject jsonObjectValue)
+				{
+					await _classConfigType.ParseIncludesAsync(jsonObjectValue, targetObject, targetType, context, cancellationToken);
+				}
+			}
 		}
 
 		readonly Type _type;
-		readonly Dictionary<string, PropertyInfo> _nameToIncludeProperty = new Dictionary<string, PropertyInfo>(StringComparer.OrdinalIgnoreCase);
+		readonly bool _isIncludeRoot;
 		readonly Dictionary<string, Property> _nameToProperty = new Dictionary<string, Property>(StringComparer.OrdinalIgnoreCase);
-		readonly Dictionary<string, ObjectConfigType>? _knownTypes;
+		readonly Dictionary<string, Property> _nameToIncludeProperty = new Dictionary<string, Property>(StringComparer.OrdinalIgnoreCase);
+		readonly Dictionary<string, ClassConfigType>? _knownTypes;
 
-		static readonly ConcurrentDictionary<Type, ObjectConfigType> s_typeToObjectValueType = new ConcurrentDictionary<Type, ObjectConfigType>();
+		static readonly ConcurrentDictionary<Type, ClassConfigType> s_typeToObjectValueType = new ConcurrentDictionary<Type, ClassConfigType>();
 
-		public ObjectConfigType(Type type)
+		public ClassConfigType(Type type)
 		{
 			s_typeToObjectValueType.TryAdd(type, this);
 
 			_type = type;
+			_isIncludeRoot = type.GetCustomAttribute<ConfigIncludeRootAttribute>() != null;
 
-			foreach (PropertyInfo propertyInfo in type.GetProperties(BindingFlags.Instance | BindingFlags.Public | BindingFlags.GetProperty))
+			// Find all the direct include properties
+			PropertyInfo[] propertyInfos = type.GetProperties(BindingFlags.Instance | BindingFlags.Public | BindingFlags.GetProperty);
+			foreach (PropertyInfo propertyInfo in propertyInfos)
 			{
-				JsonPropertyNameAttribute? attribute = propertyInfo.GetCustomAttribute<JsonPropertyNameAttribute>();
-				string name = attribute?.Name ?? propertyInfo.Name;
-
-				if (propertyInfo.PropertyType == typeof(List<ConfigInclude>))
+				if (propertyInfo.GetCustomAttribute<JsonIgnoreAttribute>() == null)
 				{
-					_nameToIncludeProperty.Add(name, propertyInfo);
+					string name = propertyInfo.GetCustomAttribute<JsonPropertyNameAttribute>()?.Name ?? propertyInfo.Name;
+					_nameToProperty.Add(name, CreateProperty(name, propertyInfo));
 				}
-				else
+			}
+
+			// Build a map of all the properties which can include other files
+			foreach (Property property in _nameToProperty.Values)
+			{
+				if (property.HasIncludes())
 				{
-					_nameToProperty.Add(name, CreateProperty(name, propertyInfo));
+					_nameToIncludeProperty.Add(property.Name, property);
 				}
 			}
 
+			// Build up a list of possible types for this object
 			JsonKnownTypesAttribute? knownTypes = _type.GetCustomAttribute<JsonKnownTypesAttribute>();
 			if (knownTypes != null)
 			{
-				_knownTypes = new Dictionary<string, ObjectConfigType>(StringComparer.Ordinal);
+				_knownTypes = new Dictionary<string, ClassConfigType>(StringComparer.Ordinal);
 				foreach (Type knownType in knownTypes.Types)
 				{
-					ObjectConfigType knownConfigType = FindOrAdd(knownType);
+					ClassConfigType knownConfigType = FindOrAdd(knownType);
 					foreach (JsonDiscriminatorAttribute discriminatorAttribute in knownType.GetCustomAttributes(typeof(JsonDiscriminatorAttribute), true))
 					{
 						_knownTypes.Add(discriminatorAttribute.Name, knownConfigType);
@@ -327,16 +384,18 @@ namespace Horde.Build.Configuration
 			}
 		}
 
-		public static ObjectConfigType FindOrAdd(Type type)
+		bool HasIncludes() => !_isIncludeRoot && _nameToIncludeProperty.Count > 0;
+
+		public static ClassConfigType FindOrAdd(Type type)
 		{
-			ObjectConfigType? value;
+			ClassConfigType? value;
 			if (!s_typeToObjectValueType.TryGetValue(type, out value))
 			{
 				lock (s_typeToObjectValueType)
 				{
 					if (!s_typeToObjectValueType.TryGetValue(type, out value))
 					{
-						value = new ObjectConfigType(type);
+						value = new ClassConfigType(type);
 					}
 				}
 			}
@@ -379,7 +438,7 @@ namespace Horde.Build.Configuration
 
 		public ValueTask<object> ReadAsync(JsonObject obj, ConfigContext context, CancellationToken cancellationToken)
 		{
-			ObjectConfigType targetType = this;
+			ClassConfigType targetType = this;
 			if (_knownTypes != null && obj.TryGetPropertyValue("Type", out JsonNode? knownTypeNode) && knownTypeNode != null)
 			{
 				targetType = _knownTypes[knownTypeNode.ToString()];
@@ -444,30 +503,10 @@ namespace Horde.Build.Configuration
 
 		async Task MergeObjectAsync(object target, JsonObject obj, ConfigContext context, CancellationToken cancellationToken)
 		{
-			// If it's a BaseConfig object, read the parent config files recursively first
-			foreach ((string name, JsonNode? node) in obj)
+			// Before parsing properties into this object, read all the includes recursively
+			if (_isIncludeRoot)
 			{
-				if (_nameToIncludeProperty.TryGetValue(name, out PropertyInfo? propertyInfo))
-				{
-					List<ConfigInclude>? includes = node.Deserialize<List<ConfigInclude>>(context.JsonOptions);
-					if (includes != null)
-					{
-						List<ConfigInclude> targetIncludes = (List<ConfigInclude>)propertyInfo.GetValue(target)!;
-						foreach (ConfigInclude include in includes)
-						{
-							Uri uri = ConfigType.CombinePaths(context.CurrentFile, include.Path);
-							IConfigFile file = await ReadFileAsync(uri, context, cancellationToken);
-
-							context.IncludeStack.Push(file);
-
-							JsonObject nextObj = await ParseFileAsync(file, context, cancellationToken);
-							await MergeObjectAsync(target, nextObj, context, cancellationToken);
-
-							context.IncludeStack.Pop();
-							targetIncludes.Add(include);
-						}
-					}
-				}
+				await ParseIncludesAsync(obj, target, this, context, cancellationToken);
 			}
 
 			// Parse all the properties into this object
@@ -479,5 +518,16 @@ namespace Horde.Build.Configuration
 				}
 			}
 		}
+
+		async Task ParseIncludesAsync(JsonObject jsonObject, object targetObject, ClassConfigType targetType, ConfigContext context, CancellationToken cancellationToken)
+		{
+			foreach ((string name, JsonNode? node) in jsonObject)
+			{
+				if (_nameToIncludeProperty.TryGetValue(name, out Property? property) && node != null)
+				{
+					await property.ParseIncludesAsync(node, targetObject, targetType, context, cancellationToken);
+				}
+			}
+		}
 	}
 }
diff --git a/Engine/Source/Programs/Horde/Horde.Build/Projects/ProjectConfig.cs b/Engine/Source/Programs/Horde/Horde.Build/Projects/ProjectConfig.cs
--- a/Engine/Source/Programs/Horde/Horde.Build/Projects/ProjectConfig.cs
+++ b/Engine/Source/Programs/Horde/Horde.Build/Projects/ProjectConfig.cs
@@ -70,7 +70,7 @@ namespace Horde.Build.Projects
 		/// <summary>
 		/// Path to the configuration file
 		/// </summary>
-		[Required, ConfigPath]
+		[Required, ConfigRelativePath]
 		public string Path { get; set; } = null!;
 
 		/// <summary>
diff --git a/Engine/Source/Programs/Horde/Horde.Build/Server/GlobalConfig.cs b/Engine/Source/Programs/Horde/Horde.Build/Server/GlobalConfig.cs
--- a/Engine/Source/Programs/Horde/Horde.Build/Server/GlobalConfig.cs
+++ b/Engine/Source/Programs/Horde/Horde.Build/Server/GlobalConfig.cs
@@ -24,11 +24,24 @@ namespace Horde.Build.Server
 	using ProjectId = StringId<IProject>;
 	using AgentSoftwareChannelName = StringId<AgentSoftwareChannels>;
 
+	/// <summary>
+	/// Directive to merge config data from another source
+	/// </summary>
+	public class ConfigInclude
+	{
+		/// <summary>
+		/// Path to the config data to be included. May be relative to the including file's location.
+		/// </summary>
+		[Required, ConfigInclude, ConfigRelativePath]
+		public string Path { get; set; } = null!;
+	}
+
 	/// <summary>
 	/// Global configuration
 	/// </summary>
 	[JsonSchema("https://unrealengine.com/horde/global")]
 	[JsonSchemaCatalog("Horde Globals", "Horde global configuration file", "globals.json")]
+	[ConfigIncludeRoot]
 	public class GlobalConfig
 	{
 		/// <summary>
@@ -189,7 +202,7 @@ namespace Horde.Build.Server
 		/// <summary>
 		/// Config path for the project
 		/// </summary>
-		[Required, ConfigPath]
+		[Required, ConfigRelativePath]
 		public string Path { get; set; } = null!;
 
 		/// <summary>
