Horde: Missing using directive.

diff --git a/Engine/Source/Programs/Horde/Horde.Build.Tests/ConfigTests.cs b/Engine/Source/Programs/Horde/Horde.Build.Tests/ConfigTests.cs
--- a/Engine/Source/Programs/Horde/Horde.Build.Tests/ConfigTests.cs
+++ b/Engine/Source/Programs/Horde/Horde.Build.Tests/ConfigTests.cs
@@ -14,6 +14,7 @@ using Horde.Build.Perforce;
 using EpicGames.Perforce;
 using Microsoft.Extensions.Logging.Abstractions;
 using System.IO;
+using Horde.Build.Server;
 
 namespace Horde.Build.Tests
 {
