Horde: Refactor codebase to use new config system.

* IOptions<GlobalConfig> is now authoritative for all configuration data.
* ProjectCollection/ProjectService have been removed, since they only contained a redundant copy of the config data. StreamService has been removed due to no longer providing any functionality over IStreamCollection.
* Global configuration file can now include project/stream definitions directly, or can include them in other files. Each file can include other files using the "Include" property at top level scope.

diff --git a/Engine/Source/Programs/Horde/Horde.Build.Tests/BlobStoreTests.cs b/Engine/Source/Programs/Horde/Horde.Build.Tests/BlobStoreTests.cs
--- a/Engine/Source/Programs/Horde/Horde.Build.Tests/BlobStoreTests.cs
+++ b/Engine/Source/Programs/Horde/Horde.Build.Tests/BlobStoreTests.cs
@@ -26,8 +26,7 @@ namespace Horde.Build.Tests
 			GlobalConfig globalConfig = new GlobalConfig();
 			globalConfig.Storage.Backends.Add(new BackendConfig { Id = new BackendId("default-backend"), Type = StorageBackendType.Memory });
 			globalConfig.Storage.Namespaces.Add(new NamespaceConfig { Id = new NamespaceId("default"), Backend = new BackendId("default-backend"), GcDelayHrs = 0.0 });
-			await ConfigCollection.AddConfigAsync("globals", globalConfig);
-			Assert.IsNotNull(await GlobalsService.TryUpdateAsync(await GlobalsService.GetAsync(), "globals"));
+			SetConfig(globalConfig);
 
 			return await StorageService.GetClientAsync(new NamespaceId("default"), CancellationToken.None);
 		}
diff --git a/Engine/Source/Programs/Horde/Horde.Build.Tests/Compute/ComputeServiceTest.cs b/Engine/Source/Programs/Horde/Horde.Build.Tests/Compute/ComputeServiceTest.cs
--- a/Engine/Source/Programs/Horde/Horde.Build.Tests/Compute/ComputeServiceTest.cs
+++ b/Engine/Source/Programs/Horde/Horde.Build.Tests/Compute/ComputeServiceTest.cs
@@ -35,13 +35,11 @@ namespace Horde.Build.Tests.Compute
 		}
 
 		[TestInitialize]
-		public async Task Setup()
+		public void Setup()
 		{
 			GlobalConfig globalConfig = new GlobalConfig();
 			globalConfig.Compute.Add(_clusterConfig);
-
-			await ConfigCollection.AddConfigAsync("globals", globalConfig);
-			Assert.IsNotNull(await GlobalsService.TryUpdateAsync(await GlobalsService.GetAsync(), "globals"));
+			SetConfig(globalConfig);
 		}
 		
 		[TestMethod]
diff --git a/Engine/Source/Programs/Horde/Horde.Build.Tests/ControllerIntegrationTest.cs b/Engine/Source/Programs/Horde/Horde.Build.Tests/ControllerIntegrationTest.cs
--- a/Engine/Source/Programs/Horde/Horde.Build.Tests/ControllerIntegrationTest.cs
+++ b/Engine/Source/Programs/Horde/Horde.Build.Tests/ControllerIntegrationTest.cs
@@ -21,6 +21,7 @@ using Horde.Build.Configuration;
 using Horde.Build.Jobs.Graphs;
 using Horde.Build.Jobs.Artifacts;
 using Moq;
+using Microsoft.Extensions.Options;
 
 namespace Horde.Build.Tests;
 
@@ -91,16 +92,15 @@ public class ControllerIntegrationTest : IDisposable
 	private async Task<Fixture> CreateFixture()
 	{
 		IServiceProvider services = Factory.Services;
-		ConfigCollection configCollection = services.GetRequiredService<ConfigCollection>();
+		ConfigService configService = services.GetRequiredService<ConfigService>();
 		MongoService mongoService = services.GetRequiredService<MongoService>();
 		ITemplateCollection templateService = services.GetRequiredService<ITemplateCollection>();
 		JobService jobService = services.GetRequiredService<JobService>();
 		IArtifactCollection artifactCollection = services.GetRequiredService<IArtifactCollection>();
-		StreamService streamService = services.GetRequiredService<StreamService>();
 		AgentService agentService = services.GetRequiredService<AgentService>();
 		GraphCollection graphCollection = new(mongoService);
+		IOptions<ServerSettings> serverSettings = services.GetRequiredService<IOptions<ServerSettings>>();
 
-		return await Fixture.Create(configCollection, graphCollection, templateService, jobService, artifactCollection, streamService,
-			agentService);
+		return await Fixture.Create(configService, graphCollection, templateService, jobService, artifactCollection, agentService, serverSettings.Value);
 	}
 }
\ No newline at end of file
diff --git a/Engine/Source/Programs/Horde/Horde.Build.Tests/DeviceServiceTests.cs b/Engine/Source/Programs/Horde/Horde.Build.Tests/DeviceServiceTests.cs
--- a/Engine/Source/Programs/Horde/Horde.Build.Tests/DeviceServiceTests.cs
+++ b/Engine/Source/Programs/Horde/Horde.Build.Tests/DeviceServiceTests.cs
@@ -57,10 +57,6 @@ namespace Horde.Build.Tests
 		{
 			DeviceConfig Devices = new DeviceConfig();
 
-			const string ConfigRevision = "device-test-config";
-			await ConfigCollection.AddConfigAsync(ConfigRevision, new ProjectConfig { Name = "UE5" });
-			await ProjectCollection.AddOrUpdateAsync(new ProjectId("ue5"), ConfigRevision, 0);
-
 			// create 2 pools
 			for (int i = 1; i < 3; i++)
 			{
@@ -97,11 +93,7 @@ namespace Horde.Build.Tests
 				}
 			}
 
-			GlobalConfig globalConfig = new GlobalConfig();
-			globalConfig.Devices = Devices;
-
-			await ConfigCollection.AddConfigAsync("globals", globalConfig);
-			Assert.IsNotNull(await GlobalsService.TryUpdateAsync(await GlobalsService.GetAsync(), "globals"));
+			UpdateConfig(x => x.Devices = Devices);
 
 			// add 4 devices to each platform, split between 2 pools
 			for (int i = 1; i < 4; i++)
diff --git a/Engine/Source/Programs/Horde/Horde.Build.Tests/Fixture.cs b/Engine/Source/Programs/Horde/Horde.Build.Tests/Fixture.cs
--- a/Engine/Source/Programs/Horde/Horde.Build.Tests/Fixture.cs
+++ b/Engine/Source/Programs/Horde/Horde.Build.Tests/Fixture.cs
@@ -32,7 +32,8 @@ namespace Horde.Build.Tests
 		public IJob Job2 { get; private set; } = null!;
 		public ITemplate Template { get; private set; } = null!;
 		public IGraph Graph { get; private set; } = null!;
-		public IStream? Stream { get; private set; }
+		public StreamId StreamId { get; private set; }
+		public StreamConfig? StreamConfig { get; private set; }
 		public TemplateId TemplateRefId1 { get; private set; }
 		public TemplateId TemplateRefId2 { get; private set; }
 		public IArtifact Job1Artifact { get; private set; } = null!;
@@ -41,10 +42,10 @@ namespace Horde.Build.Tests
 		public string Agent1Name { get; private set; } = null!;
 		public const string PoolName = "TestingPool";
 
-		public static async Task<Fixture> Create(ConfigCollection configCollection, IGraphCollection graphCollection, ITemplateCollection templateCollection, JobService jobService, IArtifactCollection artifactCollection, StreamService streamService, AgentService agentService)
+		public static async Task<Fixture> Create(ConfigService configService, IGraphCollection graphCollection, ITemplateCollection templateCollection, JobService jobService, IArtifactCollection artifactCollection, AgentService agentService, ServerSettings serverSettings)
 		{
 			Fixture fixture = new Fixture();
-			await fixture.Populate(configCollection, graphCollection, templateCollection, jobService, artifactCollection, streamService, agentService);
+			await fixture.Populate(configService, graphCollection, templateCollection, jobService, artifactCollection, agentService, serverSettings);
 
 //			(PerforceService as PerforceServiceStub)?.AddChange("//UE5/Main", 112233, "leet.coder", "Did stuff", new []{"file.cpp"});
 //			(PerforceService as PerforceServiceStub)?.AddChange("//UE5/Main", 1111, "swarm", "A shelved CL here", new []{"renderer.cpp"});
@@ -52,7 +53,7 @@ namespace Horde.Build.Tests
 			return fixture;
 		}
 
-		private async Task Populate(ConfigCollection configCollection, IGraphCollection graphCollection, ITemplateCollection templateCollection, JobService jobService, IArtifactCollection artifactCollection, StreamService streamService, AgentService agentService)
+		private async Task Populate(ConfigService configService, IGraphCollection graphCollection, ITemplateCollection templateCollection, JobService jobService, IArtifactCollection artifactCollection, AgentService agentService, ServerSettings serverSettings)
 		{
 			FixtureGraph fg = new FixtureGraph();
 			fg.Id = ContentHash.Empty;
@@ -61,7 +62,7 @@ namespace Horde.Build.Tests
 			fg.Aggregates = new List<IAggregate>();
 			fg.Labels = new List<ILabel>();
 
-			Template = await templateCollection.AddAsync("Test template");
+			Template = await templateCollection.GetOrAddAsync(new TemplateConfig { Name = "Test template" });
 			Graph = await graphCollection.AddAsync(Template, null);
 
 			TemplateRefId1 = new TemplateId("template1");
@@ -79,22 +80,24 @@ namespace Horde.Build.Tests
 				{ "Win64", new() { Pool = new PoolId(PoolName) } }
 			};
 
-			StreamConfig config = new StreamConfig { Name = "//UE5/Main", Tabs = tabs, Templates = templates, AgentTypes = agentTypes };
-			await configCollection.AddConfigAsync("rev1", config);
+			StreamId streamId = new StreamId("ue5-main");
+			StreamConfig streamConfig = new StreamConfig { Id = streamId, Name = "//UE5/Main", Tabs = tabs, Templates = templates, AgentTypes = agentTypes };
+
+			ProjectId projectId = new ProjectId("ue5");
+			ProjectConfig projectConfig = new ProjectConfig { Id = projectId, Name = "UE5", Streams = new List<StreamConfig> { streamConfig } };
+
+			GlobalConfig globalConfig = new GlobalConfig { Projects = new List<ProjectConfig> { projectConfig } };
+			globalConfig.PostLoad(serverSettings);
+			configService.Set(IoHash.Zero, globalConfig);
+
+			StreamId = streamId;
+			StreamConfig = streamConfig;
 
-			Stream = await streamService.StreamCollection.GetAsync(new StreamId("ue5-main"));
-			Stream = await streamService.StreamCollection.TryCreateOrReplaceAsync(
-				new StreamId("ue5-main"),
-				Stream,
-				"rev1",
-				new ProjectId("does-not-exist")
-			);
-			
 			Job1 = await jobService.CreateJobAsync(
 				jobId: new JobId("5f283932841e7fdbcafb6ab5"),
-				stream: Stream!,
+				streamConfig: streamConfig,
 				templateRefId: TemplateRefId1,
-				templateHash: Template.Id,
+				templateHash: Template.Hash,
 				graph: Graph,
 				name: "hello1",
 				change: 1000001,
@@ -105,7 +108,7 @@ namespace Horde.Build.Tests
 
 			Job2 = await jobService.CreateJobAsync(
 				jobId: new JobId("5f69ea1b68423e921b035106"),
-				stream: Stream!,
+				streamConfig: streamConfig,
 				templateRefId: new TemplateId("template-id-1"),
 				templateHash: ContentHash.MD5("made-up-template-hash"),
 				graph: fg,
diff --git a/Engine/Source/Programs/Horde/Horde.Build.Tests/Fleet/FleetServiceTest.cs b/Engine/Source/Programs/Horde/Horde.Build.Tests/Fleet/FleetServiceTest.cs
--- a/Engine/Source/Programs/Horde/Horde.Build.Tests/Fleet/FleetServiceTest.cs
+++ b/Engine/Source/Programs/Horde/Horde.Build.Tests/Fleet/FleetServiceTest.cs
@@ -171,8 +171,8 @@ namespace Horde.Build.Tests.Fleet
 			serverSettingsOpt.Value.FleetManagerV2 = FleetManagerType.AwsReuse;
 			
 			FleetService service = new(
-				AgentCollection, GraphCollection, JobCollection, LeaseCollection, PoolCollection, new DowntimeServiceStub(isDowntimeActive), StreamService, _dogStatsD,
-				new StubFleetManagerFactory(fleetManager), Clock, Cache, serverSettingsOpt, loggerFactory.CreateLogger<FleetService>());
+				AgentCollection, GraphCollection, JobCollection, LeaseCollection, PoolCollection, new DowntimeServiceStub(isDowntimeActive), StreamCollection, _dogStatsD,
+				new StubFleetManagerFactory(fleetManager), Clock, Cache, serverSettingsOpt, GlobalConfig, loggerFactory.CreateLogger<FleetService>());
 				
 			return service;
 		}
@@ -434,4 +434,4 @@ namespace Horde.Build.Tests.Fleet
 			return FleetService.CreateFleetManager(pool);
 		}
 	}
-}
\ No newline at end of file
+}
diff --git a/Engine/Source/Programs/Horde/Horde.Build.Tests/Fleet/JobQueueStrategyTest.cs b/Engine/Source/Programs/Horde/Horde.Build.Tests/Fleet/JobQueueStrategyTest.cs
--- a/Engine/Source/Programs/Horde/Horde.Build.Tests/Fleet/JobQueueStrategyTest.cs
+++ b/Engine/Source/Programs/Horde/Horde.Build.Tests/Fleet/JobQueueStrategyTest.cs
@@ -114,12 +114,16 @@ namespace Horde.Build.Tests.Fleet
 			string agentTypeName1 = "bogusAgentType1";
 			Dictionary<string, AgentConfig> agentTypes = new() { {agentTypeName1, new() { Pool = new PoolId(pool.Name) } }, };
 
-			IStream stream = (await CreateOrReplaceStreamAsync(
-				new StreamId("ue5-main"),
-				null,
-				new ProjectId("does-not-exist"),
-				new StreamConfig { Name = "//UE5/Main", AgentTypes = agentTypes }
-			))!;
+			StreamConfig streamConfig = new StreamConfig { Id = new StreamId("ue5"), Name = "//UE5/Main", AgentTypes = agentTypes };
+
+			ProjectConfig projectConfig = new ProjectConfig();
+			projectConfig.Id = new ProjectId("ue5");
+			projectConfig.Streams.Add(streamConfig);
+
+			GlobalConfig globalConfig = new GlobalConfig();
+			globalConfig.Projects.Add(projectConfig);
+
+			SetConfig(globalConfig);
 
 			string nodeForAgentType1 = "bogusNodeOnAgentType1";
 			IGraph graph = await GraphCollection.AppendAsync(null, new()
@@ -132,18 +136,17 @@ namespace Horde.Build.Tests.Fleet
 
 			for (int i = 0; i < numBatchesRunning; i++)
 			{
-				IJob job = await AddPlaceholderJob(graph, stream.Id, nodeForAgentType1);
+				IJob job = await AddPlaceholderJob(graph, streamConfig.Id, nodeForAgentType1);
 				await JobCollection.TryUpdateBatchAsync(job, graph, job.Batches[0].Id, null, JobStepBatchState.Running, null);
 			}
 			
 			for (int i = 0; i < numBatchesReady; i++)
 			{
-				IJob job = await AddPlaceholderJob(graph, stream.Id, nodeForAgentType1);
+				IJob job = await AddPlaceholderJob(graph, streamConfig.Id, nodeForAgentType1);
 				await JobCollection.TryUpdateBatchAsync(job, graph, job.Batches[0].Id, null, JobStepBatchState.Ready, null);
 			}
 
-			JobQueueSettings settings = new ();
-			return (new (JobCollection, GraphCollection, StreamService, Clock, Cache, settings), poolSize, pool, agents);
+			return (new (JobCollection, GraphCollection, StreamCollection, Clock, Cache, GlobalConfig), poolSize, pool, agents);
 		}
 		
 		private async Task<IJob> AddPlaceholderJob(IGraph graph, StreamId streamId, string nodeNameToExecute)
@@ -158,4 +161,4 @@ namespace Horde.Build.Tests.Fleet
 			return job;
 		}
 	}
-}
\ No newline at end of file
+}
diff --git a/Engine/Source/Programs/Horde/Horde.Build.Tests/GcServiceTests.cs b/Engine/Source/Programs/Horde/Horde.Build.Tests/GcServiceTests.cs
--- a/Engine/Source/Programs/Horde/Horde.Build.Tests/GcServiceTests.cs
+++ b/Engine/Source/Programs/Horde/Horde.Build.Tests/GcServiceTests.cs
@@ -19,13 +19,12 @@ namespace Horde.Build.Tests
 	[TestClass]
 	public sealed class GcServiceTests : TestSetup
 	{
-		async Task SetupNamespaceAsync()
+		void SetupNamespace()
 		{
 			GlobalConfig globalConfig = new GlobalConfig();
 			globalConfig.Storage.Backends.Add(new BackendConfig { Id = new BackendId("default-backend"), Type = StorageBackendType.Memory });
 			globalConfig.Storage.Namespaces.Add(new NamespaceConfig { Id = new NamespaceId("default"), Backend = new BackendId("default-backend"), GcDelayHrs = 0.0 });
-			await ConfigCollection.AddConfigAsync("globals", globalConfig);
-			Assert.IsNotNull(await GlobalsService.TryUpdateAsync(await GlobalsService.GetAsync(), "globals"));
+			SetConfig(globalConfig);
 		}
 
 		[TestMethod]
@@ -34,7 +33,7 @@ namespace Horde.Build.Tests
 			await StorageService.StartAsync(CancellationToken.None);
 
 			BlobId.UseDeterministicIds();
-			await SetupNamespaceAsync();
+			SetupNamespace();
 			IStorageClientImpl store = await StorageService.GetClientAsync(new NamespaceId("default"), CancellationToken.None);
 
 			Random random = new Random(0);
diff --git a/Engine/Source/Programs/Horde/Horde.Build.Tests/IssueServiceTests.cs b/Engine/Source/Programs/Horde/Horde.Build.Tests/IssueServiceTests.cs
--- a/Engine/Source/Programs/Horde/Horde.Build.Tests/IssueServiceTests.cs
+++ b/Engine/Source/Programs/Horde/Horde.Build.Tests/IssueServiceTests.cs
@@ -117,18 +117,6 @@ namespace Horde.Build.Tests
 		readonly DirectoryReference _autoSdkDir;
 		readonly DirectoryReference _workspaceDir;
 
-		async Task<IStream> CreateStreamAsync(ProjectId projectId, StreamId streamId, string streamName)
-		{
-			string revision = $"config:{streamId}";
-
-			StreamConfig streamConfig = new StreamConfig { Name = streamName };
-			streamConfig.Tabs.Add(new JobsTabConfig { Title = "General", Templates = new List<TemplateId> { new TemplateId("test-template") } });
-			streamConfig.Templates.Add(new TemplateRefConfig { Id = new TemplateId("test-template") });
-			await ConfigCollection.AddConfigAsync(revision, streamConfig);
-
-			return Deref(await StreamCollection.TryCreateOrReplaceAsync(streamId, null, revision, projectId));
-		}
-
 		public IssueServiceTests()
 		{
 			if (RuntimeInformation.IsOSPlatform(OSPlatform.Windows))
@@ -142,13 +130,29 @@ namespace Horde.Build.Tests
 				_workspaceDir = new DirectoryReference("/Horde");
 			}
 
-			const string ProjectConfigRevision = "projectconfig";
-			ConfigCollection.AddConfigAsync(ProjectConfigRevision, new ProjectConfig { Name = "UE4" }).Wait();
-			IProject project = ProjectCollection.AddOrUpdateAsync(new ProjectId("ue4"), ProjectConfigRevision, 0).Result!;
+			ProjectId projectId = new ProjectId("ue5");
+
+			ProjectConfig projectConfig = new ProjectConfig();
+			projectConfig.Id = projectId;
+			projectConfig.Streams.Add(CreateStream(_mainStreamId, MainStreamName));
+			projectConfig.Streams.Add(CreateStream(_releaseStreamId, ReleaseStreamName));
+			projectConfig.Streams.Add(CreateStream(_devStreamId, DevStreamName));
 
-			IStream mainStream = CreateStreamAsync(project.Id, _mainStreamId, MainStreamName).Result;
-			IStream releaseStream = CreateStreamAsync(project.Id, _releaseStreamId, ReleaseStreamName).Result;
-			IStream devStream = CreateStreamAsync(project.Id, _devStreamId, DevStreamName).Result;
+			GlobalConfig globalConfig = new GlobalConfig();
+			globalConfig.Projects.Add(projectConfig);
+
+			SetConfig(globalConfig);
+
+			static StreamConfig CreateStream(StreamId streamId, string streamName)
+			{
+				return new StreamConfig
+				{
+					Id = streamId,
+					Name = streamName,
+					Tabs = new List<TabConfig> { new JobsTabConfig { Title = "General", Templates = new List<TemplateId> { new TemplateId("test-template") } } },
+					Templates = new List<TemplateRefConfig> { new TemplateRefConfig { Id = new TemplateId("test-template") } }
+				};
+			}
 
 			IUser bill = UserCollection.FindOrAddUserByLoginAsync("Bill").Result;
 			IUser anne = UserCollection.FindOrAddUserByLoginAsync("Anne").Result;
@@ -330,7 +334,7 @@ namespace Horde.Build.Tests
 			// Scenario: Stream is deleted
 			// Expected: Issue is closed
 			{
-				await StreamCollection.DeleteAsync(_mainStreamId);
+				UpdateConfig(x => x.Projects.Clear());
 				await Clock.AdvanceAsync(TimeSpan.FromHours(1.0));
 
 				List<IIssue> issues = await IssueCollection.FindIssuesAsync();
diff --git a/Engine/Source/Programs/Horde/Horde.Build.Tests/JobServiceTests.cs b/Engine/Source/Programs/Horde/Horde.Build.Tests/JobServiceTests.cs
--- a/Engine/Source/Programs/Horde/Horde.Build.Tests/JobServiceTests.cs
+++ b/Engine/Source/Programs/Horde/Horde.Build.Tests/JobServiceTests.cs
@@ -39,37 +39,37 @@ namespace Horde.Build.Tests
 		public async Task TestChainedJobs()
 		{
 			ProjectId projectId = new ProjectId("ue5");
-
-			const string ProjectConfigRevision = "projectconfig";
-			await ConfigCollection.AddConfigAsync(ProjectConfigRevision, new ProjectConfig { Name = "UE4" });
-			IProject? project = await ProjectService.Collection.AddOrUpdateAsync(projectId, ProjectConfigRevision, 0);
-			Assert.IsNotNull(project);
-
-			ITemplate template = await TemplateCollection.AddAsync("Test template");
-			IGraph graph = await GraphCollection.AddAsync(template, null);
-
+			StreamId streamId = new StreamId("ue5-main");
 			TemplateId templateRefId1 = new TemplateId("template1");
 			TemplateId templateRefId2 = new TemplateId("template2");
 
-			StreamConfig streamConfig = new StreamConfig();
+			StreamConfig streamConfig = new StreamConfig { Id = streamId };
 			streamConfig.Templates.Add(new TemplateRefConfig { Id = templateRefId1, Name = "Test Template", ChainedJobs = new List<ChainedJobTemplateConfig> { new ChainedJobTemplateConfig { TemplateId = templateRefId2, Trigger = "Setup Build" } } });
 			streamConfig.Templates.Add(new TemplateRefConfig { Id = templateRefId2, Name = "Test Template" });
 			streamConfig.Tabs.Add(new JobsTabConfig { Title = "foo", Templates = new List<TemplateId> { templateRefId1, templateRefId2 } });
 
-			StreamId streamId = new StreamId("ue5-main");
-			IStream? stream = await StreamService.GetStreamAsync(streamId);
-			stream = await CreateOrReplaceStreamAsync(new StreamId("ue5-main"), stream, projectId, streamConfig);
+			ProjectConfig projectConfig = new ProjectConfig { Id = projectId };
+			projectConfig.Streams.Add(streamConfig);
+
+			GlobalConfig globalConfig = new GlobalConfig();
+			globalConfig.Projects.Add(projectConfig);
+
+			SetConfig(globalConfig);
 
 			CreateJobOptions options = new CreateJobOptions();
 			options.PreflightChange = 999;
-			options.JobTriggers.AddRange(stream!.Templates[templateRefId1].Config.ChainedJobs!);
+			options.JobTriggers.AddRange(streamConfig.Templates[0].ChainedJobs!);
+
+			ITemplate template = await TemplateCollection.GetOrAddAsync(streamConfig.Templates[0]);
 
-			IJob job = await JobService.CreateJobAsync(null, stream!, templateRefId1, template.Id, graph, "Hello", 1234, 1233, options);
+			IGraph graph = await GraphCollection.AddAsync(template, null);
+
+			IJob job = await JobService.CreateJobAsync(null, streamConfig, templateRefId1, template.Hash, graph, "Hello", 1234, 1233, options);
 			Assert.AreEqual(1, job.ChainedJobs.Count);
 
-			job = Deref(await JobService.UpdateBatchAsync(job, job.Batches[0].Id, LogId.GenerateNewId(), JobStepBatchState.Running));
-			job = Deref(await JobService.UpdateStepAsync(job, job.Batches[0].Id, job.Batches[0].Steps[0].Id, JobStepState.Running));
-			job = Deref(await JobService.UpdateStepAsync(job, job.Batches[0].Id, job.Batches[0].Steps[0].Id, JobStepState.Completed, JobStepOutcome.Success));
+			job = Deref(await JobService.UpdateBatchAsync(job, job.Batches[0].Id, streamConfig, LogId.GenerateNewId(), JobStepBatchState.Running));
+			job = Deref(await JobService.UpdateStepAsync(job, job.Batches[0].Id, job.Batches[0].Steps[0].Id, streamConfig, JobStepState.Running));
+			job = Deref(await JobService.UpdateStepAsync(job, job.Batches[0].Id, job.Batches[0].Steps[0].Id, streamConfig, JobStepState.Completed, JobStepOutcome.Success));
 
 			Assert.IsNotNull(job.ChainedJobs[0].JobId);
 
@@ -86,17 +86,13 @@ namespace Horde.Build.Tests
 		[TestMethod]
 		public async Task ChangeQueryTestAsync()
 		{
+			StreamId streamId = new StreamId("ue5-main");
+			StreamConfig streamConfig = new StreamConfig { Id = streamId };
+
 			ProjectId projectId = new ProjectId("ue5");
-			const string ProjectConfigRevision = "projectconfig";
-			await ConfigCollection.AddConfigAsync(ProjectConfigRevision, new ProjectConfig { Name = "UE4" });
-			IProject? project = await ProjectService.Collection.AddOrUpdateAsync(projectId, ProjectConfigRevision, 0);
-			Assert.IsNotNull(project);
+			ProjectConfig projectConfig = new ProjectConfig { Id = projectId, Name = "UE4", Streams = new List<StreamConfig> { streamConfig } };
 
-			StreamId streamId = new StreamId("ue5-main");
-			const string StreamConfigRevision = "streamconfig";
-			await ConfigCollection.AddConfigAsync(StreamConfigRevision, new StreamConfig { });
-			IStream? stream = await CreateOrReplaceStreamAsync(streamId, null, projectId, new StreamConfig());
-			Assert.IsNotNull(stream);
+			SetConfig(new GlobalConfig { Projects = new List<ProjectConfig> { projectConfig } });
 
 			IUser user = await UserCollection.FindOrAddUserByLoginAsync("Bob");
 
@@ -106,22 +102,22 @@ namespace Horde.Build.Tests
 			PerforceService.AddChange(streamId, 1003, user, "", new[] { "Foo.uasset" });
 			PerforceService.AddChange(streamId, 1004, user, "", new[] { "Foo.uasset" });
 
-			ICommitCollection commits = PerforceService.GetCommits(stream);
+			ICommitCollection commits = PerforceService.GetCommits(streamConfig);
 
 			{
-				int? change = await JobService.EvaluateChangeQueryAsync(stream, new ChangeQueryConfig { CommitTag = CommitTag.Code }, null, commits, CancellationToken.None);
+				int? change = await JobService.EvaluateChangeQueryAsync(streamId, new ChangeQueryConfig { CommitTag = CommitTag.Code }, null, commits, CancellationToken.None);
 				Assert.AreEqual(1002, change);
 			}
 			{
-				int? change = await JobService.EvaluateChangeQueryAsync(stream, new ChangeQueryConfig { CommitTag = CommitTag.Content }, null, commits, CancellationToken.None);
+				int? change = await JobService.EvaluateChangeQueryAsync(streamId, new ChangeQueryConfig { CommitTag = CommitTag.Content }, null, commits, CancellationToken.None);
 				Assert.AreEqual(1004, change);
 			}
 			{
-				int? change = await JobService.EvaluateChangeQueryAsync(stream, new ChangeQueryConfig { CommitTag = CommitTag.Content, Condition = "tag.code == 1" }, new List<CommitTag> { CommitTag.Code }, commits, CancellationToken.None);
+				int? change = await JobService.EvaluateChangeQueryAsync(streamId, new ChangeQueryConfig { CommitTag = CommitTag.Content, Condition = "tag.code == 1" }, new List<CommitTag> { CommitTag.Code }, commits, CancellationToken.None);
 				Assert.AreEqual(1004, change);
 			}
 			{
-				int? change = await JobService.EvaluateChangeQueryAsync(stream, new ChangeQueryConfig { CommitTag = CommitTag.Content, Condition = "tag.content == 1" }, new List<CommitTag> { CommitTag.Code }, commits, CancellationToken.None);
+				int? change = await JobService.EvaluateChangeQueryAsync(streamId, new ChangeQueryConfig { CommitTag = CommitTag.Content, Condition = "tag.content == 1" }, new List<CommitTag> { CommitTag.Code }, commits, CancellationToken.None);
 				Assert.IsNull(change);
 			}
 		}
@@ -165,7 +161,7 @@ namespace Horde.Build.Tests
 
 			return await JobService.CreateJobAsync(
 				jobId: JobId.GenerateNewId(),
-				stream: fixture!.Stream!,
+				streamConfig: fixture!.StreamConfig!,
 				templateRefId: new TemplateId(templateRefId),
 				templateHash: new ContentHash(Encoding.ASCII.GetBytes(templateHash)),
 				graph: fixture!.Graph,
@@ -203,19 +199,20 @@ namespace Horde.Build.Tests
 		[TestMethod]
 		public async Task TestRunEarly()
 		{
-			IAgent? agent = await AgentService.CreateAgentAsync("TestAgent", true, new List<StringId<IPool>> { new StringId<IPool>("win") });
-			await AgentService.CreateSessionAsync(agent, AgentStatus.Ok, new List<string>(), new Dictionary<string, int>(), null);
+			StreamId streamId = new StreamId("ue5-main");
+			StreamConfig streamConfig = new StreamConfig { Id = streamId };
 
-			const string ProjectConfigRevision = "projectconfig";
-			await ConfigCollection.AddConfigAsync(ProjectConfigRevision, new ProjectConfig { Name = "UE4" });
-			IProject? project = await ProjectService.Collection.AddOrUpdateAsync(new ProjectId("ue5"), ProjectConfigRevision, 0);
-			Assert.IsNotNull(project);
+			ProjectId projectId = new ProjectId("ue5");
+			ProjectConfig projectConfig = new ProjectConfig { Id = projectId, Name = "UE5", Streams = new List<StreamConfig> { streamConfig } };
 
-			StreamId streamId = new StreamId("ue5-main");
-			IStream? stream = await StreamCollection.GetAsync(streamId);
-			stream = await CreateOrReplaceStreamAsync(streamId, stream, project!.Id, new StreamConfig { Name = "//UE5/Main" });
+			SetConfig(new GlobalConfig { Projects = new List<ProjectConfig> { projectConfig } });
+
+			// ----
+
+			IAgent? agent = await AgentService.CreateAgentAsync("TestAgent", true, new List<StringId<IPool>> { new StringId<IPool>("win") });
+			await AgentService.CreateSessionAsync(agent, AgentStatus.Ok, new List<string>(), new Dictionary<string, int>(), null);
 
-			ITemplate template = await TemplateCollection.AddAsync("Test template");
+			ITemplate template = await TemplateCollection.GetOrAddAsync(new TemplateConfig { Name = "Test template" });
 			IGraph graph = await GraphCollection.AddAsync(template, null);
 
 			NewGroup groupA = new NewGroup("win", new List<NewNode>());
@@ -232,18 +229,18 @@ namespace Horde.Build.Tests
 			options.PreflightChange = 999;
 			options.Arguments.Add("-Target=Pak");
 
-			IJob job = await JobService.CreateJobAsync(null, stream!, new TemplateId("temp"), template.Id, graph, "Hello", 1234, 1233, options);
+			IJob job = await JobService.CreateJobAsync(null, streamConfig!, new TemplateId("temp"), new ContentHash(new byte[] { 1, 2, 3 }), graph, "Hello", 1234, 1233, options);
 
-			job = Deref(await JobService.UpdateBatchAsync(job, job.Batches[0].Id, LogId.GenerateNewId(), JobStepBatchState.Running));
-			job = Deref(await JobService.UpdateStepAsync(job, job.Batches[0].Id, job.Batches[0].Steps[0].Id, JobStepState.Running));
-			job = Deref(await JobService.UpdateStepAsync(job, job.Batches[0].Id, job.Batches[0].Steps[0].Id, JobStepState.Completed, JobStepOutcome.Success));
+			job = Deref(await JobService.UpdateBatchAsync(job, job.Batches[0].Id, streamConfig, LogId.GenerateNewId(), JobStepBatchState.Running));
+			job = Deref(await JobService.UpdateStepAsync(job, job.Batches[0].Id, job.Batches[0].Steps[0].Id, streamConfig, JobStepState.Running));
+			job = Deref(await JobService.UpdateStepAsync(job, job.Batches[0].Id, job.Batches[0].Steps[0].Id, streamConfig, JobStepState.Completed, JobStepOutcome.Success));
 
-			job = Deref(await JobService.UpdateBatchAsync(job, job.Batches[1].Id, LogId.GenerateNewId(), JobStepBatchState.Running));
-			job = Deref(await JobService.UpdateStepAsync(job, job.Batches[1].Id, job.Batches[1].Steps[0].Id, JobStepState.Running));
+			job = Deref(await JobService.UpdateBatchAsync(job, job.Batches[1].Id, streamConfig, LogId.GenerateNewId(), JobStepBatchState.Running));
+			job = Deref(await JobService.UpdateStepAsync(job, job.Batches[1].Id, job.Batches[1].Steps[0].Id, streamConfig, JobStepState.Running));
 
-			job = Deref(await JobService.UpdateBatchAsync(job, job.Batches[2].Id, LogId.GenerateNewId(), JobStepBatchState.Running));
-			job = Deref(await JobService.UpdateStepAsync(job, job.Batches[2].Id, job.Batches[2].Steps[0].Id, JobStepState.Running));
-			job = Deref(await JobService.UpdateStepAsync(job, job.Batches[2].Id, job.Batches[2].Steps[0].Id, JobStepState.Completed, JobStepOutcome.Success));
+			job = Deref(await JobService.UpdateBatchAsync(job, job.Batches[2].Id, streamConfig, LogId.GenerateNewId(), JobStepBatchState.Running));
+			job = Deref(await JobService.UpdateStepAsync(job, job.Batches[2].Id, job.Batches[2].Steps[0].Id, streamConfig, JobStepState.Running));
+			job = Deref(await JobService.UpdateStepAsync(job, job.Batches[2].Id, job.Batches[2].Steps[0].Id, streamConfig, JobStepState.Completed, JobStepOutcome.Success));
 
 			Assert.AreEqual(JobStepState.Waiting, job.Batches[2].Steps[1].State);
 		}
diff --git a/Engine/Source/Programs/Horde/Horde.Build.Tests/JobTaskSourceTests.cs b/Engine/Source/Programs/Horde/Horde.Build.Tests/JobTaskSourceTests.cs
--- a/Engine/Source/Programs/Horde/Horde.Build.Tests/JobTaskSourceTests.cs
+++ b/Engine/Source/Programs/Horde/Horde.Build.Tests/JobTaskSourceTests.cs
@@ -49,7 +49,6 @@ namespace Horde.Build.Tests
 		{
 			Fixture fixture = await SetupPoolWithAgentAsync(isPoolAutoScaled: true, shouldCreateAgent: true, isAgentEnabled: true);
 
-			Assert.AreEqual(0, JobTaskSource.GetQueueForTesting().Count);
 			await JobTaskSource.TickAsync(CancellationToken.None);
 			Assert.AreEqual(1, JobTaskSource.GetQueueForTesting().Count);
 			Assert.AreEqual(fixture.Job1.Id, JobTaskSource.GetQueueForTesting().Min!.Id.Item1);
@@ -64,7 +63,6 @@ namespace Horde.Build.Tests
 		{
 			Fixture fixture = await SetupPoolWithAgentAsync(isPoolAutoScaled: true, shouldCreateAgent: false, isAgentEnabled: false);
 			
-			Assert.AreEqual(0, JobTaskSource.GetQueueForTesting().Count);
 			await JobTaskSource.TickAsync(CancellationToken.None);
 			Assert.AreEqual(0, JobTaskSource.GetQueueForTesting().Count);
 
@@ -79,7 +77,6 @@ namespace Horde.Build.Tests
 		{
 			Fixture fixture = await SetupPoolWithAgentAsync(isPoolAutoScaled: false, shouldCreateAgent: true, isAgentEnabled: false);
 			
-			Assert.AreEqual(0, JobTaskSource.GetQueueForTesting().Count);
 			await JobTaskSource.TickAsync(CancellationToken.None);
 			Assert.AreEqual(1, JobTaskSource.GetQueueForTesting().Count);
 
@@ -95,7 +92,6 @@ namespace Horde.Build.Tests
 		{
 			Fixture fixture = await SetupPoolWithAgentAsync(isPoolAutoScaled: true, shouldCreateAgent: true, isAgentEnabled: false);
 			
-			Assert.AreEqual(0, JobTaskSource.GetQueueForTesting().Count);
 			await JobTaskSource.TickAsync(CancellationToken.None);
 			Assert.AreEqual(1, JobTaskSource.GetQueueForTesting().Count);
 
@@ -112,7 +108,8 @@ namespace Horde.Build.Tests
 			Fixture fixture = await SetupPoolWithAgentAsync(isPoolAutoScaled: true, shouldCreateAgent: true, isAgentEnabled: true);
 
 			// update template with some step states
-			IStream Stream = Deref(await StreamService.TryUpdateTemplateRefAsync(fixture.Stream!, fixture.TemplateRefId1, new List<UpdateStepStateRequest>() { new UpdateStepStateRequest() { Name = "Paused Step", PausedByUserId = UserId.GenerateNewId().ToString() } }));
+			IStream Stream = await StreamCollection.GetAsync(fixture.StreamConfig!);
+			Stream = Deref(await StreamCollection.TryUpdateTemplateRefAsync(Stream, fixture.TemplateRefId1, new List<UpdateStepStateRequest>() { new UpdateStepStateRequest() { Name = "Paused Step", PausedByUserId = UserId.GenerateNewId().ToString() } }));
 
 			// create a new graph with the associated nodes
 			List<NewGroup> newGroups = new List<NewGroup>();
@@ -137,11 +134,10 @@ namespace Horde.Build.Tests
 			options.Arguments.Add("-Target=Step That Depends on Paused Step;Step That Depends on Update Version Files");
 
 			IJob job = await JobCollection.AddAsync(JobId.GenerateNewId(), Stream.Id,
-				fixture.TemplateRefId1, fixture.Template.Id, graph, "Test Paused Step Job",
+				fixture.TemplateRefId1, fixture.Template.Hash, graph, "Test Paused Step Job",
 				1000, 1000, options);
 
 			// validate
-			Assert.AreEqual(0, JobTaskSource.GetQueueForTesting().Count);
 			await JobTaskSource.TickAsync(CancellationToken.None);
 			Assert.AreEqual(1, JobTaskSource.GetQueueForTesting().Count);
 			Assert.AreEqual(job.Id, JobTaskSource.GetQueueForTesting().Min!.Id.Item1);
diff --git a/Engine/Source/Programs/Horde/Horde.Build.Tests/JobsControllerDbTest.cs b/Engine/Source/Programs/Horde/Horde.Build.Tests/JobsControllerDbTest.cs
--- a/Engine/Source/Programs/Horde/Horde.Build.Tests/JobsControllerDbTest.cs
+++ b/Engine/Source/Programs/Horde/Horde.Build.Tests/JobsControllerDbTest.cs
@@ -67,7 +67,7 @@ namespace Horde.Build.Tests
 	        Fixture fixture = await CreateFixtureAsync();
 	        IJob job = fixture.Job1;
 	        string[] templates = { job.TemplateId.ToString() };
-	        object obj = (await JobsController.FindJobTimingsAsync(fixture.Stream!.Id.ToString(), templates)).Value!;
+	        object obj = (await JobsController.FindJobTimingsAsync(fixture.StreamConfig!.Id.ToString(), templates)).Value!;
 	        FindJobTimingsResponse res = (obj as FindJobTimingsResponse)!;
 	        Assert.AreEqual(1, res.Timings.Count);
 	        GetJobTimingResponse timingResponse = res.Timings[job.Id.ToString()];
diff --git a/Engine/Source/Programs/Horde/Horde.Build.Tests/NotificationServiceTests.cs b/Engine/Source/Programs/Horde/Horde.Build.Tests/NotificationServiceTests.cs
--- a/Engine/Source/Programs/Horde/Horde.Build.Tests/NotificationServiceTests.cs
+++ b/Engine/Source/Programs/Horde/Horde.Build.Tests/NotificationServiceTests.cs
@@ -36,14 +36,14 @@ namespace Horde.Build.Tests
 			return Task.CompletedTask;
 		}
 
-		public Task NotifyJobCompleteAsync(IStream jobStream, IJob job, IGraph graph, LabelOutcome outcome) { throw new NotImplementedException(); }
-		public Task NotifyJobCompleteAsync(IUser user, IStream jobStream, IJob job, IGraph graph, LabelOutcome outcome) { throw new NotImplementedException(); }
-		public Task NotifyJobStepCompleteAsync(IUser user, IStream jobStream, IJob job, IJobStepBatch batch, IJobStep step, INode node, List<ILogEventData> jobStepEventData) { throw new NotImplementedException(); }
-		public Task NotifyLabelCompleteAsync(IUser user, IJob job, IStream stream, ILabel label, int labelIdx, LabelOutcome outcome, List<(string, JobStepOutcome, Uri)> stepData) { throw new NotImplementedException(); }
+		public Task NotifyJobCompleteAsync(IJob job, IGraph graph, LabelOutcome outcome) { throw new NotImplementedException(); }
+		public Task NotifyJobCompleteAsync(IUser user, IJob job, IGraph graph, LabelOutcome outcome) { throw new NotImplementedException(); }
+		public Task NotifyJobStepCompleteAsync(IUser user, IJob job, IJobStepBatch batch, IJobStep step, INode node, List<ILogEventData> jobStepEventData) { throw new NotImplementedException(); }
+		public Task NotifyLabelCompleteAsync(IUser user, IJob job, ILabel label, int labelIdx, LabelOutcome outcome, List<(string, JobStepOutcome, Uri)> stepData) { throw new NotImplementedException(); }
 		public Task NotifyIssueUpdatedAsync(IIssue issue) { throw new NotImplementedException(); }
 		public Task NotifyConfigUpdateAsync(Exception? ex) => Task.CompletedTask;
 		public Task NotifyConfigUpdateFailureAsync(string errorMessage, string fileName, int? change = null, IUser? author = null, string? description = null) { throw new NotImplementedException(); }
-		public Task NotifyDeviceServiceAsync(string message, IDevice? device = null, IDevicePool? pool = null, IStream? stream = null, IJob? job = null, IJobStep? step = null, INode? node = null, IUser? user = null) { throw new NotImplementedException(); }
+		public Task NotifyDeviceServiceAsync(string message, IDevice? device = null, IDevicePool? pool = null, StreamConfig? stream = null, IJob? job = null, IJobStep? step = null, INode? node = null, IUser? user = null) { throw new NotImplementedException(); }
 		public Task SendIssueReportAsync(IssueReportGroup report) => throw new NotImplementedException();
 	}
 
diff --git a/Engine/Source/Programs/Horde/Horde.Build.Tests/SchedulerTests.cs b/Engine/Source/Programs/Horde/Horde.Build.Tests/SchedulerTests.cs
--- a/Engine/Source/Programs/Horde/Horde.Build.Tests/SchedulerTests.cs
+++ b/Engine/Source/Programs/Horde/Horde.Build.Tests/SchedulerTests.cs
@@ -41,12 +41,11 @@ namespace Horde.Build.Tests
 		{
 			IUser bob = UserCollection.FindOrAddUserByLoginAsync("Bob").Result;
 
-			const string ProjectConfigRevision = "projectconfig";
-			ConfigCollection.AddConfigAsync(ProjectConfigRevision, new ProjectConfig { Name = "UE4" }).Wait();
-			IProject? project = ProjectService.Collection.AddOrUpdateAsync(ProjectId, ProjectConfigRevision, 0).Result;
-			Assert.IsNotNull(project);
+			GlobalConfig globalConfig = new GlobalConfig();
+			globalConfig.Projects.Add(new ProjectConfig { Id = ProjectId, Name = "UE4" });
+			SetConfig(globalConfig);
 
-			_template = TemplateCollection.AddAsync("Test template").Result;
+			_template = TemplateCollection.GetOrAddAsync(new TemplateConfig { Name = "Test template" }).Result;
 
 			_initialJobIds = new HashSet<JobId>(JobCollection.FindAsync().Result.Select(x => x.Id));
 
@@ -60,14 +59,15 @@ namespace Horde.Build.Tests
 		{
 			await ScheduleService.ResetAsync();
 
-			IStream? stream = await StreamService.GetStreamAsync(StreamId);
+			StreamConfig streamConfig = new StreamConfig();
+			streamConfig.Id = StreamId;
+			streamConfig.Name = "//UE5/Main";
+			streamConfig.Tabs.Add(new JobsTabConfig { Title = "foo", Templates = new List<TemplateId> { TemplateId } });
+			streamConfig.Templates.Add(new TemplateRefConfig { Id = TemplateId, Name = "Test", Schedule = schedule });
 
-			StreamConfig config = new StreamConfig();
-			config.Name = "//UE5/Main";
-			config.Tabs.Add(new JobsTabConfig { Title = "foo", Templates = new List<TemplateId> { TemplateId } });
-			config.Templates.Add(new TemplateRefConfig { Id = TemplateId, Name = "Test", Schedule = schedule });
+			UpdateConfig(x => x.Projects[0].Streams = new List<StreamConfig> { streamConfig });
 
-			return (await CreateOrReplaceStreamAsync(StreamId, stream, ProjectId, config))!; 
+			return await StreamCollection.GetAsync(streamConfig);
 		}
 
 		public async Task<List<IJob>> FileTestHelperAsync(params string[] files)
@@ -234,7 +234,7 @@ namespace Horde.Build.Tests
 			ScheduleConfig schedule = new ScheduleConfig();
 			schedule.Enabled = true;
 			schedule.Patterns.Add(new SchedulePatternConfig { MinTime = 13 * 60, MaxTime = 14 * 60, Interval = 15 });
-			IStream stream = await SetScheduleAsync(schedule);
+			await SetScheduleAsync(schedule);
 
 			// Initial tick
 			await ScheduleService.TickForTestingAsync();
@@ -251,7 +251,10 @@ namespace Horde.Build.Tests
 			Assert.AreEqual(102, jobs2[0].Change);
 			Assert.AreEqual(100, jobs2[0].CodeChange);
 
-			IStream stream2 = (await StreamCollection.GetAsync(stream.Id))!;
+			StreamConfig? streamConfig;
+			GlobalConfig.CurrentValue.TryGetStream(StreamId, out streamConfig);
+
+			IStream stream2 = (await StreamCollection.GetAsync(streamConfig!))!;
 			ITemplateSchedule schedule2 = stream2.Templates.First().Value.Schedule!;
 			Assert.AreEqual(102, schedule2.LastTriggerChange);
 			Assert.AreEqual(Clock.UtcNow, schedule2.LastTriggerTimeUtc);
@@ -419,11 +422,11 @@ namespace Horde.Build.Tests
 			Clock.UtcNow = startTime;
 
 			// Create two templates, the second dependent on the first
-			ITemplate? newTemplate1 = await TemplateCollection.AddAsync("Test template 1");
+			ITemplate? newTemplate1 = await TemplateCollection.GetOrAddAsync(new TemplateConfig { Name = "Test template 1" });
 			//TemplateRef newTemplateRef1 = new TemplateRef(newTemplate1);
 			TemplateId newTemplateRefId1 = new TemplateId("new-template-1");
 
-			ITemplate? newTemplate2 = await TemplateCollection.AddAsync("Test template 2");
+			ITemplate? newTemplate2 = await TemplateCollection.GetOrAddAsync(new TemplateConfig { Name = "Test template 2" });
 //			TemplateRef newTemplateRef2 = new TemplateRef(newTemplate2);
 //			newTemplateRef2.Schedule = new Schedule(Clock.UtcNow);
 //			newTemplateRef2.Schedule.Gate = new ScheduleGate(newTemplateRefId1, "TriggerNext");
@@ -431,14 +434,14 @@ namespace Horde.Build.Tests
 //			newTemplateRef2.Schedule.LastTriggerTime = startTime;
 			TemplateId newTemplateRefId2 = new TemplateId("new-template-2");
 
-			IStream? stream = await StreamService.GetStreamAsync(StreamId);
-
 			StreamConfig config = new StreamConfig();
+			config.Id = StreamId;
 			config.Name = "//UE5/Main";
 			config.Tabs.Add(new JobsTabConfig { Title = "foo", Templates = new List<TemplateId> { newTemplateRefId1, newTemplateRefId2 } });
 			config.Templates = new() { new TemplateRefConfig { Id = newTemplateRefId1 }, new TemplateRefConfig { Id = newTemplateRefId2 }  };
+			UpdateConfig(x => x.Projects[0].Streams = new List<StreamConfig> { config });
 
-			stream = (await CreateOrReplaceStreamAsync(StreamId, stream, ProjectId, config))!;
+			IStream stream = await StreamCollection.GetAsync(config);
 
 			// Create the TriggerNext step and mark it as complete
 			IGraph graphA = await GraphCollection.AddAsync(newTemplate1, null);
@@ -455,12 +458,12 @@ namespace Horde.Build.Tests
 			options1.PreflightChange = 999;
 			options1.Arguments.Add("-Target=TriggerNext");
 
-			IJob job1 = await JobService.CreateJobAsync(null, stream, newTemplateRefId1, _template.Id, graphA, "Hello", 1234, 1233, options1);
+			IJob job1 = await JobService.CreateJobAsync(null, config, newTemplateRefId1, newTemplate1.Hash, graphA, "Hello", 1234, 1233, options1);
 			SubResourceId batchId1 = job1.Batches[0].Id;
 			SubResourceId stepId1 = job1.Batches[0].Steps[0].Id;
-			job1 = Deref(await JobService.UpdateBatchAsync(job1, batchId1, LogId.GenerateNewId(), JobStepBatchState.Running));
-			job1 = Deref(await JobService.UpdateStepAsync(job1, batchId1, stepId1, JobStepState.Completed, JobStepOutcome.Failure));
-			job1 = Deref(await JobService.UpdateBatchAsync(job1, batchId1, LogId.GenerateNewId(), JobStepBatchState.Complete));
+			job1 = Deref(await JobService.UpdateBatchAsync(job1, batchId1, config, LogId.GenerateNewId(), JobStepBatchState.Running));
+			job1 = Deref(await JobService.UpdateStepAsync(job1, batchId1, stepId1, config, JobStepState.Completed, JobStepOutcome.Failure));
+			job1 = Deref(await JobService.UpdateBatchAsync(job1, batchId1, config, LogId.GenerateNewId(), JobStepBatchState.Complete));
 			Assert.IsNotNull(job1);
 			await GetNewJobs();
 
@@ -475,11 +478,11 @@ namespace Horde.Build.Tests
 			options2.PreflightChange = 999;
 			options2.Arguments.Add("-Target=TriggerNext");
 
-			IJob job2 = await JobService.CreateJobAsync(null, stream, newTemplateRefId1, _template.Id, graphA, "Hello", 1234, 1233, options2);
+			IJob job2 = await JobService.CreateJobAsync(null, config, newTemplateRefId1, newTemplate1.Hash, graphA, "Hello", 1234, 1233, options2);
 			SubResourceId batchId2 = job2.Batches[0].Id;
 			SubResourceId stepId2 = job2.Batches[0].Steps[0].Id;
-			job2 = Deref(await JobService.UpdateBatchAsync(job2, batchId2, LogId.GenerateNewId(), JobStepBatchState.Running));
-			job2 = Deref(await JobService.UpdateStepAsync(job2, batchId2, stepId2, JobStepState.Completed, JobStepOutcome.Success));
+			job2 = Deref(await JobService.UpdateBatchAsync(job2, batchId2, config, LogId.GenerateNewId(), JobStepBatchState.Running));
+			job2 = Deref(await JobService.UpdateStepAsync(job2, batchId2, stepId2, config, JobStepState.Completed, JobStepOutcome.Success));
 			Assert.IsNotNull(job2);
 
 			// Tick the schedule and make sure it does trigger
@@ -520,19 +523,23 @@ namespace Horde.Build.Tests
 			newTemplate2.Schedule.Patterns.Add(new SchedulePatternConfig { Interval = 10 });// (null, 0, null, 10));
 //			NewTemplate2.Schedule.LastTriggerTime = StartTime;
 
-			IStream? stream = await StreamService.GetStreamAsync(StreamId);
+//			IStream? stream = await StreamService.GetStreamAsync(StreamId);
 
 			StreamConfig config = new StreamConfig();
+			config.Id = StreamId;
 			config.Name = "//UE5/Main";
 			config.Tabs.Add(new JobsTabConfig { Title = "foo", Templates = new List<TemplateId> { newTemplateRefId1, newTemplateRefId2 } });
 			config.Templates.Add(newTemplate1);
 			config.Templates.Add(newTemplate2);
+			UpdateConfig(x => x.Projects[0].Streams = new List<StreamConfig> { config });
+
+			IStream stream = await StreamCollection.GetAsync(config);
 
-			stream = (await CreateOrReplaceStreamAsync(StreamId, stream, ProjectId, config))!;
+//			stream = (await CreateOrReplaceStreamAsync(StreamId, stream, ProjectId, config))!;
 
-			ITemplate template1 = (await TemplateCollection.GetAsync(stream.Templates[newTemplateRefId1].Hash))!;
+			ITemplate template1 = (await TemplateCollection.GetOrAddAsync(config.Templates[0]))!;
 			Assert.IsNotNull(template1);
-			ITemplate template2 = (await TemplateCollection.GetAsync(stream.Templates[newTemplateRefId2].Hash))!;
+			ITemplate template2 = (await TemplateCollection.GetOrAddAsync(config.Templates[1]))!;
 			Assert.IsNotNull(template2);
 
 			// Create the graph
@@ -548,17 +555,17 @@ namespace Horde.Build.Tests
 				CreateJobOptions options1 = new CreateJobOptions();
 				options1.Arguments.Add("-Target=TriggerNext");
 
-				IJob job1 = await JobService.CreateJobAsync(null, stream, newTemplateRefId1, _template.Id, graphA, "Hello", change, codeChange, options1);
+				IJob job1 = await JobService.CreateJobAsync(null, config, newTemplateRefId1, _template.Hash, graphA, "Hello", change, codeChange, options1);
 				for (int batchIdx = 0; batchIdx < job1.Batches.Count; batchIdx++)
 				{
 					SubResourceId batchId1 = job1.Batches[batchIdx].Id;
-					job1 = Deref(await JobService.UpdateBatchAsync(job1, batchId1, LogId.GenerateNewId(), JobStepBatchState.Running));
+					job1 = Deref(await JobService.UpdateBatchAsync(job1, batchId1, config, LogId.GenerateNewId(), JobStepBatchState.Running));
 					for (int stepIdx = 0; stepIdx < job1.Batches[batchIdx].Steps.Count; stepIdx++)
 					{
 						SubResourceId stepId1 = job1.Batches[batchIdx].Steps[stepIdx].Id;
-						job1 = Deref(await JobService.UpdateStepAsync(job1, batchId1, stepId1, JobStepState.Completed, JobStepOutcome.Success, newLogId: LogId.GenerateNewId()));
+						job1 = Deref(await JobService.UpdateStepAsync(job1, batchId1, stepId1, config, JobStepState.Completed, JobStepOutcome.Success, newLogId: LogId.GenerateNewId()));
 					}
-					job1 = Deref(await JobService.UpdateBatchAsync(job1, batchId1, LogId.GenerateNewId(), JobStepBatchState.Complete));
+					job1 = Deref(await JobService.UpdateBatchAsync(job1, batchId1, config, LogId.GenerateNewId(), JobStepBatchState.Complete));
 				}
 			}
 			await GetNewJobs();
@@ -595,7 +602,7 @@ namespace Horde.Build.Tests
 			Assert.AreEqual(100, jobs1[0].CodeChange);
 
 			// Make sure the job is registered
-			IStream? stream1 = await StreamService.GetStreamAsync(StreamId);
+			IStream? stream1 = await StreamCollection.GetAsync(GlobalConfig.CurrentValue.Streams[0]);
 			ITemplateRef templateRef1 = stream1!.Templates.First().Value;
 			Assert.AreEqual(1, templateRef1.Schedule!.ActiveJobs.Count);
 			Assert.AreEqual(jobs1[0].Id, templateRef1.Schedule!.ActiveJobs[0]);
@@ -604,7 +611,7 @@ namespace Horde.Build.Tests
 			await SetScheduleAsync(schedule);
 
 			// Make sure the job is still registered
-			IStream? stream2 = await StreamService.GetStreamAsync(StreamId);
+			IStream? stream2 = await StreamCollection.GetAsync(GlobalConfig.CurrentValue.Streams[0]);
 			ITemplateRef templateRef2 = stream2!.Templates.First().Value;
 			Assert.AreEqual(1, templateRef2.Schedule!.ActiveJobs.Count);
 			Assert.AreEqual(jobs1[0].Id, templateRef2.Schedule!.ActiveJobs[0]);
@@ -648,4 +655,4 @@ namespace Horde.Build.Tests
 			return jobs.OrderBy(x => x.Change).ToList();
 		}
 	}
-}
\ No newline at end of file
+}
diff --git a/Engine/Source/Programs/Horde/Horde.Build.Tests/StorageServiceTests.cs b/Engine/Source/Programs/Horde/Horde.Build.Tests/StorageServiceTests.cs
--- a/Engine/Source/Programs/Horde/Horde.Build.Tests/StorageServiceTests.cs
+++ b/Engine/Source/Programs/Horde/Horde.Build.Tests/StorageServiceTests.cs
@@ -21,9 +21,6 @@ namespace Horde.Build.Tests
 		[TestMethod]
 		public async Task BlobCollectionTest()
 		{
-			ConfigUpdateService configUpdateService = ServiceProvider.GetRequiredService<ConfigUpdateService>();
-			await configUpdateService.StartAsync(default);
-
 			StorageService storageService = ServiceProvider.GetRequiredService<StorageService>();
 			IStorageClient client = await storageService.GetClientAsync(new NamespaceId("memory"), default);
 
diff --git a/Engine/Source/Programs/Horde/Horde.Build.Tests/StreamServiceTest.cs b/Engine/Source/Programs/Horde/Horde.Build.Tests/StreamServiceTest.cs
--- a/Engine/Source/Programs/Horde/Horde.Build.Tests/StreamServiceTest.cs
+++ b/Engine/Source/Programs/Horde/Horde.Build.Tests/StreamServiceTest.cs
@@ -16,24 +16,24 @@ namespace Horde.Build.Tests
         {
 			Fixture fixture = await CreateFixtureAsync();
 
-	        IStream stream = (await StreamService.GetStreamAsync(fixture!.Stream!.Id))!;
+	        IStream stream = (await StreamCollection.GetAsync(fixture!.StreamConfig!))!;
 	        Assert.IsFalse(stream.IsPaused(DateTime.UtcNow));
 	        Assert.IsNull(stream.PausedUntil);
 	        Assert.IsNull(stream.PauseComment);
 
 	        DateTime pausedUntil = DateTime.UtcNow.AddHours(1);
-	        await StreamService.UpdatePauseStateAsync(stream, newPausedUntil: pausedUntil, newPauseComment: "mycomment");
-	        stream = (await StreamService.GetStreamAsync(fixture!.Stream!.Id))!;
+	        await StreamCollection.UpdatePauseStateAsync(stream, newPausedUntil: pausedUntil, newPauseComment: "mycomment");
+	        stream = (await StreamCollection.GetAsync(fixture!.StreamConfig!))!;
 	        // Comparing by string to avoid comparing exact milliseconds as those are not persisted in MongoDB fields
 	        Assert.IsTrue(stream.IsPaused(DateTime.UtcNow));
 	        Assert.AreEqual(pausedUntil.ToString(CultureInfo.InvariantCulture), stream.PausedUntil!.Value.ToString(CultureInfo.InvariantCulture));
 	        Assert.AreEqual("mycomment", stream.PauseComment);
 	        
-	        await StreamService.UpdatePauseStateAsync(stream, newPausedUntil: null, newPauseComment: null);
-	        stream = (await StreamService.GetStreamAsync(fixture!.Stream!.Id))!;
+	        await StreamCollection.UpdatePauseStateAsync(stream, newPausedUntil: null, newPauseComment: null);
+	        stream = (await StreamCollection.GetAsync(fixture!.StreamConfig!))!;
 	        Assert.IsFalse(stream.IsPaused(DateTime.UtcNow));
 	        Assert.IsNull(stream.PausedUntil);
 	        Assert.IsNull(stream.PauseComment);
         }
     }
-}
\ No newline at end of file
+}
diff --git a/Engine/Source/Programs/Horde/Horde.Build.Tests/Stubs/Services/PerforceServiceStub.cs b/Engine/Source/Programs/Horde/Horde.Build.Tests/Stubs/Services/PerforceServiceStub.cs
--- a/Engine/Source/Programs/Horde/Horde.Build.Tests/Stubs/Services/PerforceServiceStub.cs
+++ b/Engine/Source/Programs/Horde/Horde.Build.Tests/Stubs/Services/PerforceServiceStub.cs
@@ -124,7 +124,7 @@ namespace Horde.Build.Tests.Stubs.Services
 			streamChanges.Add(number, new Commit(streamId, number, originalNumber, author.Id, owner.Id, description, files.ToList()));
 		}
 
-		public Task<List<ICommit>> GetChangesAsync(IStream stream, int? minChange, int? maxChange, int? numResults, CancellationToken cancellationToken)
+		public Task<List<ICommit>> GetChangesAsync(StreamConfig stream, int? minChange, int? maxChange, int? numResults, CancellationToken cancellationToken)
 		{
 			List<ICommit> results = new List<ICommit>();
 
@@ -151,12 +151,12 @@ namespace Horde.Build.Tests.Stubs.Services
 			return Task.FromResult(results);
 		}
 
-		public Task<(CheckShelfResult, ShelfInfo?)> CheckShelfAsync(IStream stream, int changeNumber, CancellationToken cancellationToken)
+		public Task<(CheckShelfResult, ShelfInfo?)> CheckShelfAsync(StreamConfig stream, int changeNumber, CancellationToken cancellationToken)
 		{
 			throw new NotImplementedException();
 		}
 
-		public Task<List<ICommit>> GetChangeDetailsAsync(IStream stream, IReadOnlyList<int> changeNumbers, CancellationToken cancellationToken)
+		public Task<List<ICommit>> GetChangeDetailsAsync(StreamConfig stream, IReadOnlyList<int> changeNumbers, CancellationToken cancellationToken)
 		{
 			List<ICommit> results = new List<ICommit>();
 			foreach (int changeNumber in changeNumbers)
@@ -176,7 +176,7 @@ namespace Horde.Build.Tests.Stubs.Services
 			throw new NotImplementedException();
 		}
 
-		public Task<(int? Change, string Message)> SubmitShelvedChangeAsync(IStream stream, int shelvedChange, int originalChange, CancellationToken cancellationToken)
+		public Task<(int? Change, string Message)> SubmitShelvedChangeAsync(StreamConfig stream, int shelvedChange, int originalChange, CancellationToken cancellationToken)
 		{
 			throw new NotImplementedException();
 		}
@@ -209,12 +209,12 @@ namespace Horde.Build.Tests.Stubs.Services
 		class CommitCollection : ICommitCollection
 		{
 			readonly PerforceServiceStub _owner;
-			readonly IStream _stream;
+			readonly StreamConfig _streamConfig;
 
-			public CommitCollection(PerforceServiceStub owner, IStream stream)
+			public CommitCollection(PerforceServiceStub owner, StreamConfig stream)
 			{
 				_owner = owner;
-				_stream = stream;
+				_streamConfig = stream;
 			}
 
 			public Task<int> CreateNewAsync(string path, string description, CancellationToken cancellationToken = default)
@@ -224,7 +224,7 @@ namespace Horde.Build.Tests.Stubs.Services
 
 			public async IAsyncEnumerable<ICommit> FindAsync(int? minChange, int? maxChange, int? maxResults, IReadOnlyList<CommitTag>? tags, [EnumeratorCancellation] CancellationToken cancellationToken = default)
 			{
-				foreach (ICommit commit in await _owner.GetChangesAsync(_stream, minChange, maxChange, null, cancellationToken))
+				foreach (ICommit commit in await _owner.GetChangesAsync(_streamConfig, minChange, maxChange, null, cancellationToken))
 				{
 					if (tags != null && tags.Count > 0)
 					{
@@ -240,13 +240,13 @@ namespace Horde.Build.Tests.Stubs.Services
 
 			public async Task<ICommit> GetAsync(int changeNumber, CancellationToken cancellationToken = default)
 			{
-				List<ICommit> commits = await _owner.GetChangeDetailsAsync(_stream, new[] { changeNumber }, cancellationToken);
+				List<ICommit> commits = await _owner.GetChangeDetailsAsync(_streamConfig, new[] { changeNumber }, cancellationToken);
 				return commits[0];
 			}
 
 			public async Task<int> LatestNumberAsync(CancellationToken cancellationToken = default)
 			{
-				List<ICommit> commits = await _owner.GetChangesAsync(_stream, null, null, 1, cancellationToken);
+				List<ICommit> commits = await _owner.GetChangesAsync(_streamConfig, null, null, 1, cancellationToken);
 				return commits[0].Number;
 			}
 
@@ -256,9 +256,9 @@ namespace Horde.Build.Tests.Stubs.Services
 			}
 		}
 
-		public ICommitCollection GetCommits(IStream stream)
+		public ICommitCollection GetCommits(StreamConfig streamConfig)
 		{
-			return new CommitCollection(this, stream);
+			return new CommitCollection(this, streamConfig);
 		}
 	}
 }
diff --git a/Engine/Source/Programs/Horde/Horde.Build.Tests/TestDataTests.cs b/Engine/Source/Programs/Horde/Horde.Build.Tests/TestDataTests.cs
--- a/Engine/Source/Programs/Horde/Horde.Build.Tests/TestDataTests.cs
+++ b/Engine/Source/Programs/Horde/Horde.Build.Tests/TestDataTests.cs
@@ -19,6 +19,7 @@ using System.Security.Claims;
 using Horde.Build.Logs;
 using System.Linq;
 using System.Threading;
+using Horde.Build.Server;
 
 namespace Horde.Build.Tests
 {
@@ -74,26 +75,24 @@ namespace Horde.Build.Tests
 			return node.Object;
 		}
 
-		async Task<IStream> CreateStreamAsync(ProjectId projectId, StreamId streamId, string streamName)
+		static StreamConfig CreateStream(StreamId streamId, string streamName)
 		{
-			string revision = $"config:{streamId}";
-
-			StreamConfig streamConfig = new StreamConfig { Name = streamName };
+			StreamConfig streamConfig = new StreamConfig { Id = streamId, Name = streamName };
 			streamConfig.Tabs.Add(new JobsTabConfig { Title = "General", Templates = new List<TemplateId> { new TemplateId("test-template") } });
 			streamConfig.Templates.Add(new TemplateRefConfig { Id = new TemplateId("test-template") });
-			await ConfigCollection.AddConfigAsync(revision, streamConfig);
-
-			return Deref(await StreamCollection.TryCreateOrReplaceAsync(streamId, null, revision, projectId));
+			return streamConfig;
 		}
 
 		public TestDataTests()
 		{
-			const string ProjectConfigRevision = "projectconfig";
-			ConfigCollection.AddConfigAsync(ProjectConfigRevision, new ProjectConfig { Name = "UE4" }).Wait();
-			IProject project = ProjectCollection.AddOrUpdateAsync(new ProjectId("ue4"), ProjectConfigRevision, 0).Result!;
+			ProjectConfig projectConfig = new ProjectConfig { Id = new ProjectId("ue5"), Name = "UE5" };
+			projectConfig.Streams.Add(CreateStream(_mainStreamId, MainStreamName));
+			projectConfig.Streams.Add(CreateStream(_releaseStreamId, ReleaseStreamName));
+
+			GlobalConfig globalConfig = new GlobalConfig();
+			globalConfig.Projects.Add(projectConfig);
 
-			IStream mainStream = CreateStreamAsync(project.Id, _mainStreamId, MainStreamName).Result;
-			IStream releaseStream = CreateStreamAsync(project.Id, _releaseStreamId, ReleaseStreamName).Result;
+			SetConfig(globalConfig);
 
 			List<INode> nodes = new List<INode>();
 			nodes.Add(MockNode("Update Version Files", NodeAnnotations.Empty));
diff --git a/Engine/Source/Programs/Horde/Horde.Build.Tests/TestSetup.cs b/Engine/Source/Programs/Horde/Horde.Build.Tests/TestSetup.cs
--- a/Engine/Source/Programs/Horde/Horde.Build.Tests/TestSetup.cs
+++ b/Engine/Source/Programs/Horde/Horde.Build.Tests/TestSetup.cs
@@ -66,9 +66,6 @@ using Horde.Build.Telemetry;
 
 namespace Horde.Build.Tests
 {
-	using StreamId = StringId<IStream>;
-	using ProjectId = StringId<IProject>;
-
 	/// <summary>
 	/// Handles set up of collections, services, fixtures etc during testing
 	///
@@ -79,7 +76,6 @@ namespace Horde.Build.Tests
 		public FakeClock Clock => ServiceProvider.GetRequiredService<FakeClock>();
 		public IMemoryCache Cache => ServiceProvider.GetRequiredService<IMemoryCache>();
 
-		public ConfigCollection ConfigCollection => ServiceProvider.GetRequiredService<ConfigCollection>();
 		public IGraphCollection GraphCollection => ServiceProvider.GetRequiredService<IGraphCollection>();
 		public INotificationTriggerCollection NotificationTriggerCollection => ServiceProvider.GetRequiredService<INotificationTriggerCollection>();
 		public IStreamCollection StreamCollection => ServiceProvider.GetRequiredService<IStreamCollection>();
@@ -109,8 +105,6 @@ namespace Horde.Build.Tests
 		public ITemplateCollection TemplateCollection => ServiceProvider.GetRequiredService<ITemplateCollection>();
 		internal PerforceServiceStub PerforceService => (PerforceServiceStub)ServiceProvider.GetRequiredService<IPerforceService>();
 		public ILogFileService LogFileService => ServiceProvider.GetRequiredService<ILogFileService>();
-		public ProjectService ProjectService => ServiceProvider.GetRequiredService<ProjectService>();
-		public StreamService StreamService => ServiceProvider.GetRequiredService<StreamService>();
 		public ISubscriptionCollection SubscriptionCollection => ServiceProvider.GetRequiredService<ISubscriptionCollection>();
 		public INotificationService NotificationService => ServiceProvider.GetRequiredService<INotificationService>();
 		public IssueService IssueService => ServiceProvider.GetRequiredService<IssueService>();
@@ -137,6 +131,10 @@ namespace Horde.Build.Tests
 		public DevicesController DevicesController => GetDevicesController();
 		public TestDataController TestDataController => GetTestDataController();
 
+		public ConfigService ConfigService => ServiceProvider.GetRequiredService<ConfigService>();
+		public IOptionsMonitor<GlobalConfig> GlobalConfig => ServiceProvider.GetRequiredService<IOptionsMonitor<GlobalConfig>>();
+		public IOptionsSnapshot<GlobalConfig> GlobalConfigSnapshot => ServiceProvider.GetRequiredService<IOptionsSnapshot<GlobalConfig>>();
+
 		private static bool s_datadogWriterPatched;
 
 		public TestSetup()
@@ -144,6 +142,20 @@ namespace Horde.Build.Tests
 			PatchDatadogWriter();
 		}
 
+		protected void SetConfig(GlobalConfig globalConfig)
+		{
+			globalConfig.PostLoad(ServerSettings);
+			ConfigService.Set(IoHash.Zero, globalConfig);
+		}
+
+		protected void UpdateConfig(Action<GlobalConfig> action)
+		{
+			GlobalConfig globalConfig = GlobalConfig.CurrentValue;
+			action(globalConfig);
+			globalConfig.PostLoad(ServerSettings);
+			ConfigService.Set(IoHash.Zero, globalConfig);
+		}
+
 		protected override void ConfigureSettings(ServerSettings settings)
 		{
 			DirectoryReference baseDir = DirectoryReference.Combine(Program.DataDir, "Tests");
@@ -179,6 +191,12 @@ namespace Horde.Build.Tests
 			services.AddSingleton<IAuditLog<AgentId>>(sp => sp.GetRequiredService<IAuditLogFactory<AgentId>>().Create("Agents.Log", "AgentId"));
 			services.AddSingleton<ITelemetrySink, NullTelemetrySink>();
 
+			services.AddSingleton<ConfigService>();
+			services.AddSingleton<IOptionsFactory<GlobalConfig>>(sp => sp.GetRequiredService<ConfigService>());
+			services.AddSingleton<IOptionsChangeTokenSource<GlobalConfig>>(sp => sp.GetRequiredService<ConfigService>());
+
+			services.AddSingleton<IConfigSource, FileConfigSource>();
+
 			services.AddSingleton<IAgentCollection, AgentCollection>();
 			services.AddSingleton<IAgentSoftwareCollection, AgentSoftwareCollection>();
 			services.AddSingleton<IArtifactCollection, ArtifactCollection>();
@@ -221,7 +239,6 @@ namespace Horde.Build.Tests
 			services.AddSingleton<AgentSoftwareService>();
 			services.AddSingleton<FleetService>();
 			services.AddSingleton<ConsistencyService>();
-			services.AddSingleton<ConfigCollection>();
 			services.AddSingleton<ComputeService>();
 			services.AddSingleton<RequestTrackerService>();
 			services.AddSingleton<GlobalsService>();
@@ -240,10 +257,8 @@ namespace Horde.Build.Tests
 			services.AddSingleton<IPerforceService, PerforceServiceStub>();
 			services.AddSingleton<PerforceLoadBalancer>();
 			services.AddSingleton<PoolService>();
-			services.AddSingleton<ProjectService>();
 			services.AddSingleton<RpcService>();
 			services.AddSingleton<ScheduleService>();
-			services.AddSingleton<StreamService>();
 			services.AddSingleton<DeviceService>();
 			services.AddSingleton<TestDataService>();
 
@@ -256,36 +271,22 @@ namespace Horde.Build.Tests
 
 			services.AddSingleton<ILegacyStorageClient, BasicStorageClient>();
 
-			services.AddSingleton<ConfigUpdateService>();
 			services.AddSingleton<StorageService>();
 
 			services.AddSingleton<ISingletonDocument<AgentSoftwareChannels>>(new SingletonDocumentStub<AgentSoftwareChannels>());
 			services.AddSingleton<ISingletonDocument<DevicePlatformMapV1>>(new SingletonDocumentStub<DevicePlatformMapV1>());
 		}
 
-		public async Task<IStream?> CreateOrReplaceStreamAsync(StreamId streamId, IStream? stream, ProjectId projectId, StreamConfig config)
-		{
-			JsonSerializerOptions options = new JsonSerializerOptions();
-			Startup.ConfigureJsonSerializer(options);
-
-			byte[] data = JsonSerializer.SerializeToUtf8Bytes(config, options);
-			IoHash hash = IoHash.Compute(data);
-			string revision = hash.ToString();
-			await ConfigCollection.AddConfigDataAsync(revision, data);
-
-			return await StreamService.StreamCollection.TryCreateOrReplaceAsync(streamId, stream, hash.ToString(), projectId);
-		}
-
 		public Task<Fixture> CreateFixtureAsync()
 		{
-			return Fixture.Create(ConfigCollection, GraphCollection, TemplateCollection, JobService, ArtifactCollection, StreamService, AgentService);
+			return Fixture.Create(ConfigService, GraphCollection, TemplateCollection, JobService, ArtifactCollection, AgentService, ServerSettings);
 		}
 
 		private JobsController GetJobsController()
         {
 			ILogger<JobsController> logger = ServiceProvider.GetRequiredService<ILogger<JobsController>>();
-			JobsController jobsCtrl = new JobsController(GraphCollection, CommitService, PerforceService, StreamService, JobService,
-		        TemplateCollection, ArtifactCollection, UserCollection, NotificationService, AgentService, logger);
+			JobsController jobsCtrl = new JobsController(GraphCollection, CommitService, PerforceService, StreamCollection, JobService,
+		        TemplateCollection, ArtifactCollection, UserCollection, NotificationService, AgentService, GlobalConfigSnapshot, logger);
 	        jobsCtrl.ControllerContext = GetControllerContext();
 	        return jobsCtrl;
         }
@@ -293,35 +294,35 @@ namespace Horde.Build.Tests
 		private DevicesController GetDevicesController()
 		{
 			ILogger<DevicesController> logger = ServiceProvider.GetRequiredService<ILogger<DevicesController>>();
-			DevicesController devicesCtrl = new DevicesController(DeviceService, AclService, UserCollection, logger);
+			DevicesController devicesCtrl = new DevicesController(UserCollection, DeviceService, GlobalConfigSnapshot, logger);
 			devicesCtrl.ControllerContext = GetControllerContext();
 			return devicesCtrl;
 		}
 
 		private TestDataController GetTestDataController()
 		{
-			TestDataController dataCtrl = new TestDataController(TestDataService, StreamService, JobService, TestDataCollection);
+			TestDataController dataCtrl = new TestDataController(TestDataService, StreamCollection, JobService, TestDataCollection, GlobalConfigSnapshot);
 			dataCtrl.ControllerContext = GetControllerContext();
 			return dataCtrl;
 		}
 
 		private AgentsController GetAgentsController()
 		{
-			AgentsController agentCtrl = new AgentsController(AclService, AgentService);
+			AgentsController agentCtrl = new AgentsController(AgentService, GlobalConfigSnapshot);
 			agentCtrl.ControllerContext = GetControllerContext();
 			return agentCtrl;
 		}
 		
 		private PoolsController GetPoolsController()
 		{
-			PoolsController controller = new PoolsController(AclService, PoolService);
+			PoolsController controller = new PoolsController(PoolService, GlobalConfigSnapshot);
 			controller.ControllerContext = GetControllerContext();
 			return controller;
 		}
 		
 		private LeasesController GetLeasesController()
 		{
-			LeasesController controller = new LeasesController(AclService, AgentService);
+			LeasesController controller = new LeasesController(AgentService, GlobalConfigSnapshot);
 			controller.ControllerContext = GetControllerContext();
 			return controller;
 		}
diff --git a/Engine/Source/Programs/Horde/Horde.Build.Tests/ToolTests.cs b/Engine/Source/Programs/Horde/Horde.Build.Tests/ToolTests.cs
--- a/Engine/Source/Programs/Horde/Horde.Build.Tests/ToolTests.cs
+++ b/Engine/Source/Programs/Horde/Horde.Build.Tests/ToolTests.cs
@@ -5,6 +5,7 @@ using System.Collections.Generic;
 using System.IO;
 using System.Text;
 using System.Threading.Tasks;
+using Horde.Build.Server;
 using Horde.Build.Tools;
 using Horde.Build.Utilities;
 using HordeCommon;
@@ -25,18 +26,17 @@ namespace Horde.Build.Tests
 			ToolCollection collection = ServiceProvider.GetRequiredService<ToolCollection>();
 
 			ToolId toolId = new ToolId("ugs");
-			await collection.ConfigureAsync(new List<ToolConfig> { new ToolConfig(toolId) { Name = "UnrealGameSync", Description = "Tool for syncing content from source control" } });
+			SetConfig(new GlobalConfig { Tools = new List<ToolConfig> { new ToolConfig(toolId) { Name = "UnrealGameSync", Description = "Tool for syncing content from source control" } } });
 
-			ITool tool = Deref(await collection.GetAsync(toolId));
+			ITool tool = Deref(await collection.GetAsync(toolId, GlobalConfig.CurrentValue));
 			Assert.AreEqual(tool.Id, new ToolId("ugs"));
-			Assert.AreEqual(tool.Name, "UnrealGameSync");
-			Assert.AreEqual(tool.Description, "Tool for syncing content from source control");
+			Assert.AreEqual(tool.Config.Name, "UnrealGameSync");
+			Assert.AreEqual(tool.Config.Description, "Tool for syncing content from source control");
 
-			ITool tool2 = Deref(await collection.GetAsync(tool.Id));
+			ITool tool2 = Deref(await collection.GetAsync(tool.Id, GlobalConfig.CurrentValue));
 			Assert.AreEqual(tool.Id, tool2.Id);
-			Assert.AreEqual(tool.LastUpdateTime, tool2.LastUpdateTime);
-			Assert.AreEqual(tool.Name, tool2.Name);
-			Assert.AreEqual(tool.Description, tool2.Description);
+			Assert.AreEqual(tool.Config.Name, tool2.Config.Name);
+			Assert.AreEqual(tool.Config.Description, tool2.Config.Description);
         }
 
 		[TestMethod]
@@ -45,12 +45,12 @@ namespace Horde.Build.Tests
 			ToolCollection collection = ServiceProvider.GetRequiredService<ToolCollection>();
 
 			ToolId toolId = new ToolId("ugs");
-			await collection.ConfigureAsync(new List<ToolConfig> { new ToolConfig(toolId) { Name = "UnrealGameSync", Description = "Tool for syncing content from source control" } });
+			SetConfig(new GlobalConfig { Tools = new List<ToolConfig> { new ToolConfig(toolId) { Name = "UnrealGameSync", Description = "Tool for syncing content from source control" } } });
 
-			ITool tool = Deref(await collection.GetAsync(toolId));
+			ITool tool = Deref(await collection.GetAsync(toolId, GlobalConfig.CurrentValue));
 			Assert.AreEqual(new ToolId("ugs"), tool.Id);
-			Assert.AreEqual("UnrealGameSync", tool.Name);
-			Assert.AreEqual("Tool for syncing content from source control", tool.Description);
+			Assert.AreEqual("UnrealGameSync", tool.Config.Name);
+			Assert.AreEqual("Tool for syncing content from source control", tool.Config.Description);
 			Assert.AreEqual(0, tool.Deployments.Count);
 
 			byte[] deploymentData = Encoding.UTF8.GetBytes("hello world");
@@ -58,7 +58,7 @@ namespace Horde.Build.Tests
 			ToolDeploymentId deploymentId;
 			using (MemoryStream stream = new MemoryStream(deploymentData))
 			{
-				tool = Deref(await collection.CreateDeploymentAsync(tool, new ToolDeploymentConfig { Version = "1.0", Duration = TimeSpan.FromMinutes(5.0), CreatePaused = true }, stream));
+				tool = Deref(await collection.CreateDeploymentAsync(tool, new ToolDeploymentConfig { Version = "1.0", Duration = TimeSpan.FromMinutes(5.0), CreatePaused = true }, stream, GlobalConfig.CurrentValue));
 				Assert.AreEqual(1, tool.Deployments.Count);
 				Assert.IsNull(tool.Deployments[0].StartedAt);
 				deploymentId = tool.Deployments[^1].Id;
@@ -68,7 +68,7 @@ namespace Horde.Build.Tests
 			FakeClock clock = ServiceProvider.GetRequiredService<FakeClock>();
 			await clock.AdvanceAsync(TimeSpan.FromHours(1.0));
 
-			tool = Deref(await collection.GetAsync(tool.Id));
+			tool = Deref(await collection.GetAsync(tool.Id, GlobalConfig.CurrentValue));
 			Assert.AreEqual(1, tool.Deployments.Count);
 			Assert.IsNull(tool.Deployments[0].StartedAt);
 
@@ -87,7 +87,7 @@ namespace Horde.Build.Tests
 
 			// Check it stays paused
 			await clock.AdvanceAsync(TimeSpan.FromHours(1.0));
-			tool = Deref(await collection.GetAsync(tool.Id));
+			tool = Deref(await collection.GetAsync(tool.Id, GlobalConfig.CurrentValue));
 			Assert.AreEqual(1, tool.Deployments.Count);
 			Assert.IsNull(tool.Deployments[0].StartedAt);
 			Assert.IsTrue((tool.Deployments[0].Progress - 0.5) < 0.1);
diff --git a/Engine/Source/Programs/Horde/Horde.Build/Acls/Acl.cs b/Engine/Source/Programs/Horde/Horde.Build/Acls/Acl.cs
--- a/Engine/Source/Programs/Horde/Horde.Build/Acls/Acl.cs
+++ b/Engine/Source/Programs/Horde/Horde.Build/Acls/Acl.cs
@@ -504,588 +504,4 @@ namespace Horde.Build.Acls
 
         #endregion
     }
-
-    #region V1
-
-    /// <summary>
-    /// Stores information about a claim
-    /// </summary>
-    public class AclClaim
-	{
-		/// <summary>
-		/// The claim type, typically a URI
-		/// </summary>
-		public string Type { get; set; }
-
-		/// <summary>
-		/// The claim value
-		/// </summary>
-		public string Value { get; set; }
-
-		/// <summary>
-		/// Constructor
-		/// </summary>
-		/// <param name="claim">The claim object</param>
-		public AclClaim(Claim claim)
-			: this(claim.Type, claim.Value)
-		{
-		}
-
-		/// <summary>
-		/// Constructor
-		/// </summary>
-		/// <param name="type">The claim type</param>
-		/// <param name="value">The claim value</param>
-		public AclClaim(string type, string value)
-		{
-			Type = type;
-			Value = value;
-		}
-
-		/// <summary>
-		/// Constructs a claim from a request object
-		/// </summary>
-		/// <param name="request">The request object</param>
-		public AclClaim(AclClaimConfig request)
-			: this(request.Type, request.Value)
-		{
-		}
-	}
-
-	/// <summary>
-	/// Describes an entry in the ACL for a particular claim
-	/// </summary>
-	public class AclEntry
-	{
-		/// <summary>
-		/// Claim for this entry
-		/// </summary>
-		public AclClaim Claim { get; set; }
-
-		/// <summary>
-		/// List of allowed operations
-		/// </summary>
-		[BsonSerializer(typeof(AclActionSetSerializer))]
-		public HashSet<AclAction> Actions { get; set; }
-
-		/// <summary>
-		/// Private constructor for serialization
-		/// </summary>
-		[BsonConstructor]
-		private AclEntry()
-		{
-			Claim = new AclClaim(String.Empty, String.Empty);
-			Actions = new HashSet<AclAction>();
-		}
-
-		/// <summary>
-		/// Constructor
-		/// </summary>
-		/// <param name="claim">The claim this entry applies to</param>
-		/// <param name="actions">List of allowed operations</param>
-		internal AclEntry(AclClaim claim, IEnumerable<AclAction> actions)
-		{
-			Claim = claim;
-			Actions = new HashSet<AclAction>(actions);
-		}
-
-		/// <summary>
-		/// Constructs an ACL entry from a request
-		/// </summary>
-		/// <param name="request">Request instance</param>
-		/// <returns>New ACL entry</returns>
-		public static AclEntry FromRequest(AclEntryConfig request)
-		{
-			return new AclEntry(new AclClaim(request.Claim.Type, request.Claim.Value), ParseActionNames(request.Actions).ToArray());
-		}
-
-		/// <summary>
-		/// Parses a list of names as an AclAction bitmask
-		/// </summary>
-		/// <param name="actionNames">Array of names</param>
-		/// <returns>Action bitmask</returns>
-		public static List<AclAction> ParseActionNames(string[]? actionNames)
-		{
-			List<AclAction> actions = new();
-			if (actionNames != null)
-			{
-				foreach (string name in actionNames)
-				{
-					actions.Add(Enum.Parse<AclAction>(name, true));
-				}
-			}
-			return actions;
-		}
-
-		/// <summary>
-		/// Build a list of action names from an array of flags
-		/// </summary>
-		/// <param name="actions"></param>
-		/// <returns></returns>
-		public static List<string> GetActionNames(IEnumerable<AclAction> actions)
-		{
-            List<string> actionNames = new();
-            foreach(AclAction action in actions)
-			{
-				string name = Enum.GetName(typeof(AclAction), action)!;
-				actionNames.Add(name);
-			}
-			return actionNames;
-		}
-	}
-
-	/// <summary>
-	/// Represents an access control list for an object in the database
-	/// </summary>
-	public class Acl
-	{
-		/// <summary>
-		/// List of entries for this ACL
-		/// </summary>
-		public List<AclEntry> Entries { get; set; }
-
-		/// <summary>
-		/// Whether to inherit permissions from the parent ACL by default
-		/// </summary>
-		public bool Inherit { get; set; } = true;
-
-		/// <summary>
-		/// Specifies a list of exceptions to the inheritance setting
-		/// </summary>
-		public List<AclAction>? Exceptions { get; set; }
-
-		/// <summary>
-		/// Default constructor
-		/// </summary>
-		public Acl()
-		{
-			Entries = new List<AclEntry>();
-			Inherit = true;
-		}
-
-		/// <summary>
-		/// Constructor
-		/// </summary>
-		/// <param name="entries">List of entries for this ACL</param>
-		/// <param name="inherit">Whether to inherit permissions from the parent object by default</param>
-		public Acl(List<AclEntry> entries, bool inherit)
-		{
-			Entries = entries;
-			Inherit = inherit;
-		}
-
-		/// <summary>
-		/// Constructor
-		/// </summary>
-		/// <param name="entries">List of entries for this ACL</param>
-		public Acl(params AclEntry[] entries)
-			: this(entries.ToList(), true)
-		{
-		}
-
-		/// <summary>
-		/// Tests whether a user is authorized to perform the given actions
-		/// </summary>
-		/// <param name="action">Action that is being performed. This should be a single flag.</param>
-		/// <param name="user">The principal to authorize</param>
-		/// <returns>True/false if the action is allowed or denied, null if there is no specific setting for this user</returns>
-		public bool? Authorize(AclAction action, ClaimsPrincipal user)
-		{
-			// Check if there's a specific entry for this action
-			foreach (AclEntry entry in Entries)
-			{
-				if(entry.Actions.Contains(action) && user.HasClaim(entry.Claim.Type, entry.Claim.Value))
-				{
-					return true;
-				}
-			}
-
-			// Otherwise check if we're prevented from inheriting permissions
-			if(Inherit)
-			{
-				if(Exceptions != null && Exceptions.Contains(action))
-				{
-					return false;
-				}
-			}
-			else
-			{
-				if (Exceptions == null || !Exceptions.Contains(action))
-				{
-					return false;
-				}
-			}
-
-			// Otherwise allow to propagate up the hierarchy
-			return null;
-		}
-
-		/// <summary>
-		/// Merge new settings into an ACL
-		/// </summary>
-		/// <param name="baseAcl">The current acl</param>
-		/// <param name="update">The update to apply</param>
-		/// <returns>The new ACL value. Null if the ACL has all default settings.</returns>
-		public static Acl? Merge(Acl? baseAcl, AclConfig? update)
-		{
-			Acl? newAcl = null;
-			if(update != null)
-			{
-				newAcl = new Acl();
-
-				if(update.Entries != null)
-				{
-					newAcl.Entries = update.Entries.ConvertAll(x => AclEntry.FromRequest(x));
-				}
-				else if(baseAcl != null)
-				{
-					newAcl.Entries = baseAcl.Entries;
-				}
-
-				if(update.Inherit != null)
-				{
-					newAcl.Inherit = update.Inherit.Value;
-				}
-				else if(baseAcl != null)
-				{
-					newAcl.Inherit = baseAcl.Inherit;
-				}
-			}
-			return newAcl;
-		}
-
-		/// <summary>
-		/// Creates an update definition for the given ACL. Clears the ACL property if it's null.
-		/// </summary>
-		/// <typeparam name="T">Type of document containing the ACL</typeparam>
-		/// <param name="field">Selector for the ACL property</param>
-		/// <param name="newAcl">The new ACL value</param>
-		/// <returns>Update definition for the document</returns>
-		public static UpdateDefinition<T> CreateUpdate<T>(Expression<Func<T, object>> field, Acl newAcl)
-		{
-			if (newAcl.Entries.Count == 0 && newAcl.Inherit)
-			{
-				return Builders<T>.Update.Unset(field);
-			}
-			else
-			{
-				return Builders<T>.Update.Set(field, newAcl);
-			}
-		}
-	}
-
-    #endregion
-
-    #region V2
-
-    /// <summary>
-    /// Describes an entry in the ACL for a particular claim
-    /// </summary>
-    public class AclEntryV2 : IEquatable<AclEntryV2>
-    {
-        /// <summary>
-        /// The claim type, typically a URI
-        /// </summary>
-        [BsonElement("t")]
-        public string Type { get; set; }
-
-        /// <summary>
-        /// The claim value
-        /// </summary>
-        [BsonElement("v")]
-        public string Value { get; set; }
-
-        /// <summary>
-        /// List of allowed operations
-        /// </summary>
-        [BsonElement("a")]
-        [BsonSerializer(typeof(AclActionListSerializer))]
-        public List<AclAction> Actions { get; set; }
-
-        /// <summary>
-        /// Private constructor for serialization
-        /// </summary>
-        [BsonConstructor]
-        private AclEntryV2()
-        {
-            Type = String.Empty;
-            Value = String.Empty;
-            Actions = new List<AclAction>();
-        }
-
-        /// <summary>
-        /// Constructor
-        /// </summary>
-        /// <param name="type">The claim type</param>
-        /// <param name="value">The claim value</param>
-        /// <param name="actions">List of allowed operations</param>
-        public AclEntryV2(string type, string value, IEnumerable<AclAction> actions)
-        {
-            Type = type;
-            Value = value;
-            Actions = new List<AclAction>(actions);
-        }
-
-        /// <inheritdoc/>
-        public override bool Equals(object? obj) => Equals(obj as AclEntryV2);
-
-        /// <inheritdoc/>
-        public override int GetHashCode() => throw new NotSupportedException();
-
-        /// <inheritdoc/>
-        public bool Equals(AclEntryV2? other) => other != null && String.Equals(Type, other.Type, StringComparison.Ordinal) && String.Equals(Value, other.Value, StringComparison.Ordinal) && Actions.SequenceEqual(other.Actions);
-    }
-
-    /// <summary>
-    /// Represents an access control list for an object in the database
-    /// </summary>
-    public class AclV2 : IEquatable<AclV2>
-    {
-        /// <summary>
-        /// List of entries for this ACL
-        /// </summary>
-        [BsonElement("al")]
-        public List<AclEntryV2> Allow { get; set; } = new List<AclEntryV2>();
-
-        /// <summary>
-        /// Whether to inherit permissions from the parent ACL by default
-        /// </summary>
-        [BsonElement("in")]
-        public bool Inherit { get; set; } = true;
-
-        /// <summary>
-        /// Specifies a list of exceptions to the inheritance setting
-        /// </summary>
-        [BsonElement("ex")]
-        [BsonSerializer(typeof(AclActionListSerializer))]
-        public List<AclAction> Exceptions { get; set; } = new List<AclAction>();
-
-        /// <summary>
-        /// Default constructor
-        /// </summary>
-        public AclV2()
-        {
-            Allow = new List<AclEntryV2>();
-            Inherit = true;
-        }
-
-        /// <summary>
-        /// Constructor
-        /// </summary>
-        /// <param name="allow">List of entries for this ACL</param>
-        /// <param name="inherit">Whether to inherit permissions from the parent object by default</param>
-        public AclV2(IEnumerable<AclEntryV2> allow, bool inherit)
-        {
-            Allow.AddRange(allow);
-            Inherit = inherit;
-        }
-
-        /// <summary>
-        /// Constructor
-        /// </summary>
-        /// <param name="allow">List of entries for this ACL</param>
-        public AclV2(params AclEntryV2[] allow)
-            : this(allow.ToList(), true)
-        {
-        }
-
-        /// <summary>
-        /// Checks whether the ACL is set to its default value
-        /// </summary>
-        /// <returns>True if the ACL is its default value</returns>
-        public bool IsDefault() => Allow.Count == 0 && Inherit && Exceptions.Count == 0;
-
-		/// <summary>
-		/// Checks whether an ACL is null or set to its default value
-		/// </summary>
-		/// <returns>True if the ACL is its default value</returns>
-		public static bool IsNullOrDefault(AclV2? acl) => acl is null || acl.IsDefault();
-
-		/// <summary>
-		/// Tests whether a user is authorized to perform the given actions
-		/// </summary>
-		/// <param name="action">Action that is being performed. This should be a single flag.</param>
-		/// <param name="user">The principal to authorize</param>
-		/// <returns>True/false if the action is allowed or denied, null if there is no specific setting for this user</returns>
-		public bool? Authorize(AclAction action, ClaimsPrincipal user)
-        {
-            // Check if there's a specific entry for this action
-            foreach (AclEntryV2 entry in Allow)
-            {
-                if (entry.Actions.Contains(action) && user.HasClaim(entry.Type, entry.Value))
-                {
-                    return true;
-                }
-            }
-
-            // Otherwise check if we're prevented from inheriting permissions
-            if (Inherit)
-            {
-                if (Exceptions != null && Exceptions.Contains(action))
-                {
-                    return false;
-                }
-            }
-            else
-            {
-                if (Exceptions == null || !Exceptions.Contains(action))
-                {
-                    return false;
-                }
-            }
-
-            // Otherwise allow to propagate up the hierarchy
-            return null;
-        }
-
-        /// <inheritdoc/>
-        public override bool Equals(object? obj) => Equals(obj as AclV2);
-
-        /// <inheritdoc/>
-        public bool Equals(AclV2? other) => other != null && Allow.SequenceEqual(other.Allow) && Inherit == other.Inherit && Exceptions.SequenceEqual(other.Exceptions);
-
-		/// <summary>
-		/// Checks whether two acls are equal
-		/// </summary>
-		/// <param name="lhs"></param>
-		/// <param name="rhs"></param>
-		/// <returns></returns>
-		public static bool Equals(AclV2? lhs, AclV2? rhs)
-		{
-			if (lhs is null)
-			{
-				return rhs is null || rhs.IsDefault();
-			}
-			else
-			{
-				return lhs.Equals(rhs);
-			}
-		}
-
-        /// <inheritdoc/>
-        public override int GetHashCode() => throw new NotSupportedException();
-
-        /// <summary>
-        /// Register this type's class map
-        /// </summary>
-        internal static void ConfigureClassMap(BsonClassMap<AclV2> cm)
-        {
-            cm.AutoMap();
-            cm.MapMember(x => x.Allow).SetIgnoreIfDefault(true).SetDefaultValue(new List<AclEntryV2>());
-            cm.MapMember(x => x.Exceptions).SetIgnoreIfDefault(true).SetDefaultValue(new List<AclAction>());
-        }
-    }
-
-    #endregion
-
-    /// <summary>
-    /// Serializer for a list of unique values
-    /// </summary>
-    public sealed class AclActionListSerializer : IBsonSerializer<List<AclAction>>
-    {
-        /// <inheritdoc/>
-        public Type ValueType => typeof(List<AclAction>);
-
-        /// <inheritdoc/>
-        void IBsonSerializer.Serialize(BsonSerializationContext context, BsonSerializationArgs args, object value) => Serialize(context, args, (List<AclAction>)value);
-
-        /// <inheritdoc/>
-        object IBsonSerializer.Deserialize(BsonDeserializationContext context, BsonDeserializationArgs args) => ((IBsonSerializer<List<AclAction>>)this).Deserialize(context, args);
-
-        /// <inheritdoc/>
-        public List<AclAction> Deserialize(BsonDeserializationContext context, BsonDeserializationArgs args)
-        {
-            List<AclAction> values = new();
-
-            context.Reader.ReadStartArray();
-            for (; ; )
-            {
-                BsonType type = context.Reader.ReadBsonType();
-                if (type == BsonType.EndOfDocument)
-                {
-                    break;
-                }
-                else
-                {
-                    values.Add((AclAction)Enum.Parse(typeof(AclAction), context.Reader.ReadString()));
-                }
-            }
-            context.Reader.ReadEndArray();
-
-            return values;
-        }
-
-        /// <inheritdoc/>
-        public void Serialize(BsonSerializationContext context, BsonSerializationArgs args, List<AclAction> values)
-        {
-            context.Writer.WriteStartArray();
-            if (values.Count > 0)
-            {
-                for (int idx = 0; idx < values.Count; idx++)
-                {
-                    AclAction value = values[idx];
-                    if (values.IndexOf(value) == idx)
-                    {
-                        context.Writer.WriteString(value.ToString());
-                    }
-                }
-            }
-            context.Writer.WriteEndArray();
-        }
-    }
-
-    /// <summary>
-    /// Serializer for JobStepRefId objects
-    /// </summary>
-    public sealed class AclActionSetSerializer : IBsonSerializer<HashSet<AclAction>>
-	{
-		/// <inheritdoc/>
-		public Type ValueType => typeof(HashSet<AclAction>);
-
-		/// <inheritdoc/>
-		void IBsonSerializer.Serialize(BsonSerializationContext context, BsonSerializationArgs args, object value)
-		{
-			Serialize(context, args, (HashSet<AclAction>)value);
-		}
-
-		/// <inheritdoc/>
-		object IBsonSerializer.Deserialize(BsonDeserializationContext context, BsonDeserializationArgs args)
-		{
-			return ((IBsonSerializer<HashSet<AclAction>>)this).Deserialize(context, args);
-		}
-
-		/// <inheritdoc/>
-		public HashSet<AclAction> Deserialize(BsonDeserializationContext context, BsonDeserializationArgs args)
-		{
-            HashSet<AclAction> values = new();
-
-			context.Reader.ReadStartArray();
-			for(; ;)
-			{
-				BsonType type = context.Reader.ReadBsonType();
-				if(type == BsonType.EndOfDocument)
-				{
-					break;
-				}
-				else
-				{
-					values.Add((AclAction)Enum.Parse(typeof(AclAction), context.Reader.ReadString()));
-				}
-			}
-			context.Reader.ReadEndArray();
-
-			return values;
-		}
-
-		/// <inheritdoc/>
-		public void Serialize(BsonSerializationContext context, BsonSerializationArgs args, HashSet<AclAction> values)
-		{
-			context.Writer.WriteStartArray();
-			foreach (AclAction value in values)
-			{
-				context.Writer.WriteString(value.ToString());
-			}
-			context.Writer.WriteEndArray();
-		}
-	}
 }
diff --git a/Engine/Source/Programs/Horde/Horde.Build/Acls/AclConfig.cs b/Engine/Source/Programs/Horde/Horde.Build/Acls/AclConfig.cs
--- a/Engine/Source/Programs/Horde/Horde.Build/Acls/AclConfig.cs
+++ b/Engine/Source/Programs/Horde/Horde.Build/Acls/AclConfig.cs
@@ -1,7 +1,9 @@
 // Copyright Epic Games, Inc. All Rights Reserved.
 
+using System;
 using System.Collections.Generic;
 using System.ComponentModel.DataAnnotations;
+using System.Security.Claims;
 
 namespace Horde.Build.Acls
 {
@@ -13,7 +15,7 @@ namespace Horde.Build.Acls
 		/// <summary>
 		/// Entries to replace the existing ACL
 		/// </summary>
-		public List<AclEntryConfig>? Entries { get; set; }
+		public List<AclEntryConfig> Entries { get; set; } = new List<AclEntryConfig>();
 
 		/// <summary>
 		/// Whether to inherit permissions from the parent ACL
@@ -23,7 +25,44 @@ namespace Horde.Build.Acls
 		/// <summary>
 		/// List of exceptions to the inherited setting
 		/// </summary>
-		public List<string>? Exceptions { get; set; }
+		public List<AclAction>? Exceptions { get; set; }
+
+		/// <summary>
+		/// Tests whether a user is authorized to perform the given actions
+		/// </summary>
+		/// <param name="action">Action that is being performed. This should be a single flag.</param>
+		/// <param name="user">The principal to authorize</param>
+		/// <returns>True/false if the action is allowed or denied, null if there is no specific setting for this user</returns>
+		public bool? Authorize(AclAction action, ClaimsPrincipal user)
+		{
+			// Check if there's a specific entry for this action
+			foreach (AclEntryConfig entry in Entries)
+			{
+				if (entry.Actions.Contains(action) && user.HasClaim(entry.Claim.Type, entry.Claim.Value))
+				{
+					return true;
+				}
+			}
+
+			// Otherwise check if we're prevented from inheriting permissions
+			if (Inherit ?? true)
+			{
+				if (Exceptions != null && Exceptions.Contains(action))
+				{
+					return false;
+				}
+			}
+			else
+			{
+				if (Exceptions == null || !Exceptions.Contains(action))
+				{
+					return false;
+				}
+			}
+
+			// Otherwise allow to propagate up the hierarchy
+			return null;
+		}
 	}
 
 	/// <summary>
@@ -35,13 +74,33 @@ namespace Horde.Build.Acls
 		/// Name of the user or group
 		/// </summary>
 		[Required]
-		public AclClaimConfig Claim { get; set; } = null!;
+		public AclClaimConfig Claim { get; set; }
 
 		/// <summary>
 		/// Array of actions to allow
 		/// </summary>
 		[Required]
-		public string[] Actions { get; set; } = null!;
+		public List<AclAction> Actions { get; set; }
+
+		/// <summary>
+		/// Private constructor for serialization
+		/// </summary>
+		public AclEntryConfig()
+		{
+			Claim = new AclClaimConfig();
+			Actions = new List<AclAction>();
+		}
+
+		/// <summary>
+		/// Constructor
+		/// </summary>
+		/// <param name="claim">The claim this entry applies to</param>
+		/// <param name="actions">List of allowed operations</param>
+		public AclEntryConfig(AclClaimConfig claim, IEnumerable<AclAction> actions)
+		{
+			Claim = claim;
+			Actions = new List<AclAction>(actions);
+		}
 	}
 
 	/// <summary>
@@ -60,5 +119,34 @@ namespace Horde.Build.Acls
 		/// </summary>
 		[Required]
 		public string Value { get; set; } = null!;
+
+		/// <summary>
+		/// Constructor
+		/// </summary>
+		public AclClaimConfig()
+		{
+			Type = String.Empty;
+			Value = String.Empty;
+		}
+
+		/// <summary>
+		/// Constructor
+		/// </summary>
+		/// <param name="claim">The claim object</param>
+		public AclClaimConfig(Claim claim)
+			: this(claim.Type, claim.Value)
+		{
+		}
+
+		/// <summary>
+		/// Constructor
+		/// </summary>
+		/// <param name="type">The claim type</param>
+		/// <param name="value">The claim value</param>
+		public AclClaimConfig(string type, string value)
+		{
+			Type = type;
+			Value = value;
+		}
 	}
 }
