Horde: Refactor codebase to use new config system.

diff --git a/Engine/Source/Programs/Horde/Horde.Build/Acls/AclService.cs b/Engine/Source/Programs/Horde/Horde.Build/Acls/AclService.cs
--- a/Engine/Source/Programs/Horde/Horde.Build/Acls/AclService.cs
+++ b/Engine/Source/Programs/Horde/Horde.Build/Acls/AclService.cs
@@ -4,136 +4,32 @@ using System;
 using System.Collections.Generic;
 using System.IdentityModel.Tokens.Jwt;
 using System.Linq;
-using System.Net;
 using System.Security.Claims;
 using System.Threading.Tasks;
 using Horde.Build.Agents;
 using Horde.Build.Agents.Sessions;
 using Horde.Build.Server;
-using Horde.Build.Users;
 using Horde.Build.Utilities;
-using Microsoft.Extensions.Options;
 using Microsoft.IdentityModel.Tokens;
 using MongoDB.Driver;
 
 namespace Horde.Build.Acls
 {
 	using SessionId = ObjectId<ISession>;
-	using UserId = ObjectId<IUser>;
-
-	/// <summary>
-	/// Cache of global ACL
-	/// </summary>
-	public class GlobalPermissionsCache
-	{
-		/// <summary>
-		/// The root acl
-		/// </summary>
-		public Acl? RootAcl { get; set; }
-	}
 
 	/// <summary>
 	/// Wraps functionality for manipulating permissions
 	/// </summary>
 	public class AclService
 	{
-		/// <summary>
-		/// Name of the role that can be used to administer agents
-		/// </summary>
-		public static AclClaim AgentRegistrationClaim { get; } = new AclClaim(HordeClaimTypes.Role, "agent-registration");
-
-		/// <summary>
-		/// Name of the role used to upload software
-		/// </summary>
-		public static AclClaim DownloadSoftwareClaim { get; } = new AclClaim(HordeClaimTypes.Role, "download-software");
-
-		/// <summary>
-		/// Name of the role used to upload software
-		/// </summary>
-		public static AclClaim UploadSoftwareClaim { get; } = new AclClaim(HordeClaimTypes.Role, "upload-software");
-
-		/// <summary>
-		/// Name of the role used to upload software
-		/// </summary>
-		public static AclClaim ConfigureProjectsClaim { get; } = new AclClaim(HordeClaimTypes.Role, "configure-projects");
-
-		/// <summary>
-		/// Name of the role used to upload software
-		/// </summary>
-		public static AclClaim StartChainedJobClaim { get; } = new AclClaim(HordeClaimTypes.Role, "start-chained-job");
-
-		/// <summary>
-		/// Role for all agents
-		/// </summary>
-		public static AclClaim AgentRoleClaim { get; } = new AclClaim(HordeClaimTypes.Role, "agent");
-
-		private readonly Acl _defaultAcl = new();
 		private readonly GlobalsService _globalsService;
 
 		/// <summary>
 		/// Constructor
 		/// </summary>
-		/// <param name="globalsService">The database service instance</param>
-		/// <param name="settings">The settings object</param>
-		public AclService(GlobalsService globalsService, IOptionsMonitor<ServerSettings> settings)
+		public AclService(GlobalsService globalsService)
 		{
 			_globalsService = globalsService;
-
-            List<AclAction> adminActions = new();
-			foreach (AclAction? action in Enum.GetValues(typeof(AclAction)))
-			{
-				adminActions.Add(action!.Value);
-			}
-
-			_defaultAcl = new();
-			_defaultAcl.Entries.Add(new AclEntry(new AclClaim(ClaimTypes.Role, "internal:AgentRegistration"), new[] { AclAction.CreateAgent, AclAction.CreateSession }));
-			_defaultAcl.Entries.Add(new AclEntry(AgentRegistrationClaim, new[] { AclAction.CreateAgent, AclAction.CreateSession, AclAction.UpdateAgent, AclAction.DownloadSoftware, AclAction.CreatePool, AclAction.UpdatePool, AclAction.ViewPool, AclAction.DeletePool, AclAction.ListPools, AclAction.ViewStream, AclAction.ViewProject, AclAction.ViewJob, AclAction.ViewCosts }));
-			_defaultAcl.Entries.Add(new AclEntry(AgentRoleClaim, new[] { AclAction.ViewProject, AclAction.ViewStream, AclAction.CreateEvent, AclAction.DownloadSoftware }));
-			_defaultAcl.Entries.Add(new AclEntry(DownloadSoftwareClaim, new[] { AclAction.DownloadSoftware }));
-			_defaultAcl.Entries.Add(new AclEntry(UploadSoftwareClaim, new[] { AclAction.UploadSoftware }));
-			_defaultAcl.Entries.Add(new AclEntry(ConfigureProjectsClaim, new[] { AclAction.CreateProject, AclAction.UpdateProject, AclAction.ViewProject, AclAction.CreateStream, AclAction.UpdateStream, AclAction.ViewStream, AclAction.ChangePermissions }));
-			_defaultAcl.Entries.Add(new AclEntry(StartChainedJobClaim, new[] { AclAction.CreateJob, AclAction.ExecuteJob, AclAction.UpdateJob, AclAction.ViewJob, AclAction.ViewTemplate, AclAction.ViewStream }));
-
-			ServerSettings settingsValue = settings.CurrentValue;
-			if (settingsValue.AdminClaimType != null && settingsValue.AdminClaimValue != null)
-			{
-				_defaultAcl.Entries.Add(new AclEntry(new AclClaim(settingsValue.AdminClaimType, settingsValue.AdminClaimValue), adminActions.ToArray()));
-			}
-		}
-
-		/// <summary>
-		/// Gets the root ACL scope table
-		/// </summary>
-		/// <returns>Scopes instance</returns>
-		public async Task<Acl> GetRootAcl()
-		{
-			IGlobals globals = await _globalsService.GetAsync();
-			return Acl.Merge(new Acl(), globals.Config.Acl) ?? new Acl();
-		}
-
-		/// <summary>
-		/// Authorizes a user against a given scope
-		/// </summary>
-		/// <param name="action">The action being performed</param>
-		/// <param name="user">The principal to validate</param>
-		/// <param name="cache">The ACL scope cache</param>
-		/// <returns>Async task</returns>
-		public async Task<bool> AuthorizeAsync(AclAction action, ClaimsPrincipal user, GlobalPermissionsCache? cache = null)
-		{
-			Acl? rootAcl;
-			if(cache == null)
-			{
-				rootAcl = await GetRootAcl();
-			}
-			else if(cache.RootAcl == null)
-			{
-				rootAcl = cache.RootAcl = await GetRootAcl();
-			}
-			else
-			{
-				rootAcl = cache.RootAcl;
-			}
-			return rootAcl.Authorize(action, user) ?? _defaultAcl.Authorize(action, user) ?? false;
 		}
 
 		/// <summary>
@@ -142,7 +38,7 @@ namespace Horde.Build.Acls
 		/// <param name="claims">List of claims to include</param>
 		/// <param name="expiry">Time that the token expires</param>
 		/// <returns>JWT security token with a claim for creating new agents</returns>
-		public async ValueTask<string> IssueBearerTokenAsync(IEnumerable<AclClaim> claims, TimeSpan? expiry)
+		public async ValueTask<string> IssueBearerTokenAsync(IEnumerable<AclClaimConfig> claims, TimeSpan? expiry)
 		{
 			return await IssueBearerTokenAsync(claims.Select(x => new Claim(x.Type, x.Value)), expiry);
 		}
@@ -199,45 +95,6 @@ namespace Horde.Build.Acls
 				return new AgentId(claim.Value);
 			}
 		}
-
-		/// <summary>
-		/// Gets the role for a specific agent
-		/// </summary>
-		/// <param name="agentId">The session id</param>
-		/// <returns>New claim instance</returns>
-		public static AclClaim GetAgentClaim(AgentId agentId)
-		{
-			return new AclClaim(HordeClaimTypes.AgentId, agentId.ToString());
-		}
-
-		/// <summary>
-		/// Gets the role for a specific agent session
-		/// </summary>
-		/// <param name="sessionId">The session id</param>
-		/// <returns>New claim instance</returns>
-		public static AclClaim GetSessionClaim(SessionId sessionId)
-		{
-			return new AclClaim(HordeClaimTypes.AgentSessionId, sessionId.ToString());
-		}
-
-		/// <summary>
-		/// Determines whether the given user can masquerade as a given user
-		/// </summary>
-		/// <param name="user"></param>
-		/// <param name="userId"></param>
-		/// <returns></returns>
-		public Task<bool> AuthorizeAsUserAsync(ClaimsPrincipal user, UserId userId)
-		{
-			UserId? currentUserId = user.GetUserId();
-			if (currentUserId != null && currentUserId.Value == userId)
-			{
-				return Task.FromResult(true);
-			}
-			else
-			{
-				return AuthorizeAsync(AclAction.Impersonate, user);
-			}
-		}
 	}
 
     internal static class ClaimExtensions
diff --git a/Engine/Source/Programs/Horde/Horde.Build/Agents/AgentCollection.cs b/Engine/Source/Programs/Horde/Horde.Build/Agents/AgentCollection.cs
--- a/Engine/Source/Programs/Horde/Horde.Build/Agents/AgentCollection.cs
+++ b/Engine/Source/Programs/Horde/Horde.Build/Agents/AgentCollection.cs
@@ -7,7 +7,6 @@ using System.Linq.Expressions;
 using System.Threading.Tasks;
 using EpicGames.Redis;
 using Google.Protobuf.WellKnownTypes;
-using Horde.Build.Acls;
 using Horde.Build.Agents.Pools;
 using Horde.Build.Agents.Sessions;
 using Horde.Build.Auditing;
diff --git a/Engine/Source/Programs/Horde/Horde.Build/Agents/AgentService.cs b/Engine/Source/Programs/Horde/Horde.Build/Agents/AgentService.cs
--- a/Engine/Source/Programs/Horde/Horde.Build/Agents/AgentService.cs
+++ b/Engine/Source/Programs/Horde/Horde.Build/Agents/AgentService.cs
@@ -175,10 +175,10 @@ namespace Horde.Build.Agents
 		/// <returns>Bearer token for the agent</returns>
 		public async ValueTask<string> IssueSessionTokenAsync(AgentId agentId, SessionId sessionId)
 		{
-			List<AclClaim> claims = new List<AclClaim>();
-			claims.Add(AclService.AgentRoleClaim);
-			claims.Add(AclService.GetAgentClaim(agentId));
-			claims.Add(AclService.GetSessionClaim(sessionId));
+			List<AclClaimConfig> claims = new List<AclClaimConfig>();
+			claims.Add(HordeClaims.AgentRoleClaim);
+			claims.Add(HordeClaims.GetAgentClaim(agentId));
+			claims.Add(HordeClaims.GetSessionClaim(sessionId));
 			return await _aclService.IssueBearerTokenAsync(claims, null);
 		}
 
@@ -1014,20 +1014,6 @@ namespace Horde.Build.Agents
 			_dogStatsd.Gauge("agents.total.unspecified.count", numAgentsTotalUnspecified);
 		}
 
-		/// <summary>
-		/// Determines if the user is authorized to perform an action on a particular agent
-		/// </summary>
-		/// <param name="agent">The agent to check</param>
-		/// <param name="action">The action being performed</param>
-		/// <param name="user">The principal to authorize</param>
-		/// <param name="cache">The permissions cache</param>
-		/// <returns>True if the action is authorized</returns>
-		public async Task<bool> AuthorizeAsync(IAgent agent, AclAction action, ClaimsPrincipal user, GlobalPermissionsCache? cache)
-		{
-			_ = agent;
-			return await _aclService.AuthorizeAsync(action, user, cache);
-		}
-
 		/// <summary>
 		/// Determines if the user is authorized to perform an action on a particular agent
 		/// </summary>
diff --git a/Engine/Source/Programs/Horde/Horde.Build/Agents/AgentsController.cs b/Engine/Source/Programs/Horde/Horde.Build/Agents/AgentsController.cs
--- a/Engine/Source/Programs/Horde/Horde.Build/Agents/AgentsController.cs
+++ b/Engine/Source/Programs/Horde/Horde.Build/Agents/AgentsController.cs
@@ -10,14 +10,15 @@ using Horde.Build.Agents.Pools;
 using Horde.Build.Agents.Sessions;
 using Horde.Build.Agents.Software;
 using Horde.Build.Auditing;
+using Horde.Build.Server;
 using Horde.Build.Utilities;
 using Microsoft.AspNetCore.Authorization;
 using Microsoft.AspNetCore.Mvc;
 using Microsoft.Extensions.Logging;
+using Microsoft.Extensions.Options;
 
 namespace Horde.Build.Agents
 {
-	using AgentSoftwareChannelName = StringId<AgentSoftwareChannels>;
 	using LeaseId = ObjectId<ILease>;
 	using PoolId = StringId<IPool>;
 	using SessionId = ObjectId<ISession>;
@@ -30,18 +31,16 @@ namespace Horde.Build.Agents
 	[Route("[controller]")]
 	public class AgentsController : HordeControllerBase
 	{
-		readonly AclService _aclService;
 		readonly AgentService _agentService;
+		readonly IOptionsSnapshot<GlobalConfig> _globalConfig;
 
 		/// <summary>
 		/// Constructor
 		/// </summary>
-		/// <param name="aclService">The ACL service singleton</param>
-		/// <param name="agentService">The agent service</param>
-		public AgentsController(AclService aclService, AgentService agentService)
+		public AgentsController(AgentService agentService, IOptionsSnapshot<GlobalConfig> globalConfig)
 		{
-			_aclService = aclService;
 			_agentService = agentService;
+			_globalConfig = globalConfig;
 		}
 
 		/// <summary>
@@ -53,7 +52,7 @@ namespace Horde.Build.Agents
 		[Route("/api/v1/agents")]
 		public async Task<ActionResult<CreateAgentResponse>> CreateAgentAsync([FromBody] CreateAgentRequest request)
 		{
-			if(!await _aclService.AuthorizeAsync(AclAction.CreateAgent, User))
+			if(!_globalConfig.Value.Authorize(AclAction.CreateAgent, User))
 			{
 				return Forbid(AclAction.CreateAgent);
 			}
@@ -77,13 +76,17 @@ namespace Horde.Build.Agents
 		[ProducesResponseType(typeof(List<GetAgentResponse>), 200)]
 		public async Task<ActionResult<List<object>>> FindAgentsAsync([FromQuery] PoolId? poolId = null, [FromQuery] int? index = null, [FromQuery] int? count = null, [FromQuery] DateTimeOffset? modifiedAfter = null, [FromQuery] PropertyFilter? filter = null)
 		{
-			GlobalPermissionsCache permissionsCache = new GlobalPermissionsCache();
+			if (!_globalConfig.Value.Authorize(AclAction.ListAgents, User))
+			{
+				return Forbid(AclAction.ListAgents);
+			}
+
 			List<IAgent> agents = await _agentService.FindAgentsAsync(poolId, modifiedAfter?.UtcDateTime, index, count);
 
 			List<object> responses = new List<object>();
 			foreach (IAgent agent in agents)
 			{
-				responses.Add(await GetAgentResponseAsync(agent, permissionsCache, filter));
+				responses.Add(await GetAgentResponseAsync(agent, filter));
 			}
 
 			return responses;
@@ -100,22 +103,27 @@ namespace Horde.Build.Agents
 		[ProducesResponseType(typeof(GetAgentResponse), 200)]
 		public async Task<ActionResult<object>> GetAgentAsync(AgentId agentId, [FromQuery] PropertyFilter? filter = null)
 		{
+			if (!_globalConfig.Value.Authorize(AclAction.ViewAgent, User))
+			{
+				return Forbid(AclAction.ViewAgent, agentId);
+			}
+
 			IAgent? agent = await _agentService.GetAgentAsync(agentId);
 			if (agent == null)
 			{
 				return NotFound(agentId);
 			}
 
-			return await GetAgentResponseAsync(agent, new GlobalPermissionsCache(), filter);
+			return await GetAgentResponseAsync(agent, filter);
 		}
 
 		/// <summary>
 		/// Gets an individual agent response
 		/// </summary>
-		async ValueTask<object> GetAgentResponseAsync(IAgent agent, GlobalPermissionsCache permissionsCache, PropertyFilter? filter = null)
+		async ValueTask<object> GetAgentResponseAsync(IAgent agent, PropertyFilter? filter = null)
 		{
 			double? rate = null;
-			if (await _agentService.AuthorizeAsync(agent, AclAction.ViewCosts, User, permissionsCache))
+			if (_globalConfig.Value.Authorize(AclAction.ViewCosts, User))
 			{
 				rate = await _agentService.GetRateAsync(agent.Id);
 			}
@@ -140,11 +148,15 @@ namespace Horde.Build.Agents
 		[Route("/api/v1/agents/{agentId}")]
 		public async Task<ActionResult> UpdateAgentAsync(AgentId agentId, [FromBody] UpdateAgentRequest update)
 		{
+			if (!_globalConfig.Value.Authorize(AclAction.UpdateAgent, User))
+			{
+				return Forbid(AclAction.UpdateAgent, agentId);
+			}
+
 			List<PoolId>? updatePools = update.Pools?.ConvertAll(x => new PoolId(x));
 
 			string userName = User.GetUserName() ?? "Unknown";
 
-			GlobalPermissionsCache cache = new GlobalPermissionsCache();
 			for (; ; )
 			{
 				IAgent? agent = await _agentService.GetAgentAsync(agentId);
@@ -152,10 +164,6 @@ namespace Horde.Build.Agents
 				{
 					return NotFound(agentId);
 				}
-				if (!await _agentService.AuthorizeAsync(agent, AclAction.UpdateAgent, User, cache))
-				{
-					return Forbid(AclAction.UpdateAgent, agentId);
-				}
 
 				IAgent? newAgent = await _agentService.Agents.TryUpdateSettingsAsync(agent, update.Enabled, update.RequestConform, update.RequestFullConform, update.RequestRestart, update.RequestShutdown, $"Manual ({userName})", update.Pools?.ConvertAll(x => new PoolId(x)), update.Comment);
 				if (newAgent == null)
@@ -210,15 +218,16 @@ namespace Horde.Build.Agents
 		[Route("/api/v1/agents/{agentId}")]
 		public async Task<ActionResult> DeleteAgentAsync(AgentId agentId)
 		{
+			if (!_globalConfig.Value.Authorize(AclAction.DeleteAgent, User))
+			{
+				return Forbid(AclAction.DeleteAgent, agentId);
+			}
+
 			IAgent? agent = await _agentService.GetAgentAsync(agentId);
 			if (agent == null)
 			{
 				return NotFound(agentId);
 			}
-			if (!await _agentService.AuthorizeAsync(agent, AclAction.DeleteAgent, User, null))
-			{
-				return Forbid(AclAction.DeleteAgent, agentId);
-			}
 
 			await _agentService.DeleteAgentAsync(agent);
 			return new OkResult();
@@ -256,15 +265,16 @@ namespace Horde.Build.Agents
 		[Route("/api/v1/agents/{agentId}/sessions")]
 		public async Task<ActionResult<List<GetAgentSessionResponse>>> FindSessionsAsync(AgentId agentId, [FromQuery] DateTimeOffset? startTime, [FromQuery] DateTimeOffset? finishTime, [FromQuery] int index = 0, [FromQuery] int count = 50)
 		{
+			if (!_globalConfig.Value.Authorize(AclAction.ViewSession, User))
+			{
+				return Forbid(AclAction.ViewSession, agentId);
+			}
+
 			IAgent? agent = await _agentService.GetAgentAsync(agentId);
 			if (agent == null)
 			{
 				return NotFound(agentId);
 			}
-			if (!await _agentService.AuthorizeAsync(agent, AclAction.ViewSession, User, null))
-			{
-				return Forbid(AclAction.ViewSession, agentId);
-			}
 
 			List<ISession> sessions = await _agentService.FindSessionsAsync(agentId, startTime?.UtcDateTime, finishTime?.UtcDateTime, index, count);
 			return sessions.ConvertAll(x => new GetAgentSessionResponse(x));
@@ -280,15 +290,16 @@ namespace Horde.Build.Agents
 		[Route("/api/v1/agents/{agentId}/sessions/{sessionId}")]
 		public async Task<ActionResult<GetAgentSessionResponse>> GetSessionAsync(AgentId agentId, SessionId sessionId)
 		{
+			if (!_globalConfig.Value.Authorize(AclAction.ViewSession, User))
+			{
+				return Forbid(AclAction.ViewSession, agentId);
+			}
+
 			IAgent? agent = await _agentService.GetAgentAsync(agentId);
 			if (agent == null)
 			{
 				return NotFound(agentId);
 			}
-			if (!await _agentService.AuthorizeAsync(agent, AclAction.ViewSession, User, null))
-			{
-				return Forbid(AclAction.ViewSession, agentId);
-			}
 
 			ISession? session = await _agentService.GetSessionAsync(sessionId);
 			if(session == null || session.AgentId != agentId)
@@ -315,20 +326,21 @@ namespace Horde.Build.Agents
 		[ProducesResponseType(200, Type = typeof(List<GetAgentLeaseResponse>))]
 		public async Task<ActionResult<List<object>>> FindLeasesAsync(AgentId agentId, [FromQuery] SessionId? sessionId, [FromQuery] DateTimeOffset? startTime, [FromQuery] DateTimeOffset? finishTime, [FromQuery] int index = 0, [FromQuery] int count = 1000, [FromQuery] PropertyFilter? filter = null)
 		{
+			if (!_globalConfig.Value.Authorize(AclAction.ViewLeases, User))
+			{
+				return Forbid(AclAction.ViewLeases, agentId);
+			}
+
 			IAgent? agent = await _agentService.GetAgentAsync(agentId);
 			if (agent == null)
 			{
 				return NotFound(agentId);
 			}
-			if (!await _agentService.AuthorizeAsync(agent, AclAction.ViewLeases, User, null))
-			{
-				return Forbid(AclAction.ViewLeases, agentId);
-			}
 
 			List<ILease> leases = await _agentService.FindLeasesAsync(agentId, sessionId, startTime?.UtcDateTime, finishTime?.UtcDateTime, index, count);
 
 			double? agentRate = null;
-			if (await _aclService.AuthorizeAsync(AclAction.ViewCosts, User))
+			if (_globalConfig.Value.Authorize(AclAction.ViewCosts, User))
 			{
 				agentRate = await _agentService.GetRateAsync(agentId);
 			}
@@ -353,15 +365,16 @@ namespace Horde.Build.Agents
 		[Route("/api/v1/agents/{agentId}/leases/{leaseId}")]
 		public async Task<ActionResult<GetAgentLeaseResponse>> GetLeaseAsync(AgentId agentId, LeaseId leaseId)
 		{
+			if (!_globalConfig.Value.Authorize(AclAction.ViewLeases, User))
+			{
+				return Forbid(AclAction.ViewLeases, agentId);
+			}
+
 			IAgent? agent = await _agentService.GetAgentAsync(agentId);
 			if(agent == null)
 			{
 				return NotFound(agentId);
 			}
-			if (!await _aclService.AuthorizeAsync(AclAction.ViewLeases, User))
-			{
-				return Forbid(AclAction.ViewLeases, agentId);
-			}
 
 			ILease? lease = await _agentService.GetLeaseAsync(leaseId);
 			if (lease == null || lease.AgentId != agentId)
@@ -370,7 +383,7 @@ namespace Horde.Build.Agents
 			}
 
 			double? agentRate = null;
-			if (await _aclService.AuthorizeAsync(AclAction.ViewCosts, User))
+			if (_globalConfig.Value.Authorize(AclAction.ViewCosts, User))
 			{
 				agentRate = await _agentService.GetRateAsync(agentId);
 			}
diff --git a/Engine/Source/Programs/Horde/Horde.Build/Agents/Fleet/FleetService.cs b/Engine/Source/Programs/Horde/Horde.Build/Agents/Fleet/FleetService.cs
--- a/Engine/Source/Programs/Horde/Horde.Build/Agents/Fleet/FleetService.cs
+++ b/Engine/Source/Programs/Horde/Horde.Build/Agents/Fleet/FleetService.cs
@@ -71,7 +71,7 @@ namespace Horde.Build.Agents.Fleet
 		private readonly ILeaseCollection _leaseCollection;
 		private readonly IPoolCollection _poolCollection;
 		private readonly IDowntimeService _downtimeService;
-		private readonly StreamService _streamService;
+		private readonly IStreamCollection _streamCollection;
 		private readonly IDogStatsd _dogStatsd;
 		private readonly IFleetManagerFactory _fleetManagerFactory;
 		private readonly IClock _clock;
@@ -81,6 +81,7 @@ namespace Horde.Build.Agents.Fleet
 		private readonly TimeSpan _defaultScaleOutCooldown;
 		private readonly TimeSpan _defaultScaleInCooldown;
 		private readonly IOptions<ServerSettings> _settings;
+		private readonly IOptionsMonitor<GlobalConfig> _globalConfig;
 		private readonly ILogger<FleetService> _logger;
 		
 		/// <summary>
@@ -93,12 +94,13 @@ namespace Horde.Build.Agents.Fleet
 			ILeaseCollection leaseCollection,
 			IPoolCollection poolCollection,
 			IDowntimeService downtimeService,
-			StreamService streamService,
+			IStreamCollection streamCollection,
 			IDogStatsd dogStatsd,
 			IFleetManagerFactory fleetManagerFactory,
 			IClock clock,
 			IMemoryCache cache,
 			IOptions<ServerSettings> settings,
+			IOptionsMonitor<GlobalConfig> globalConfig,
 			ILogger<FleetService> logger)
 		{
 			_agentCollection = agentCollection;
@@ -107,11 +109,12 @@ namespace Horde.Build.Agents.Fleet
 			_leaseCollection = leaseCollection;
 			_poolCollection = poolCollection;
 			_downtimeService = downtimeService;
-			_streamService = streamService;
+			_streamCollection = streamCollection;
 			_dogStatsd = dogStatsd;
 			_fleetManagerFactory = fleetManagerFactory;
 			_clock = clock;
 			_cache = cache;
+			_globalConfig = globalConfig;
 			_logger = logger;
 			_ticker = clock.AddSharedTicker<FleetService>(TimeSpan.FromSeconds(30), TickLeaderAsync, _logger);
 			_tickerHighFrequency = clock.AddSharedTicker("FleetService.TickHighFrequency", TimeSpan.FromSeconds(30), TickHighFrequencyAsync, _logger);
@@ -337,7 +340,7 @@ namespace Horde.Build.Agents.Fleet
 						{
 							case PoolSizeStrategy.JobQueue:
 								JobQueueSettings jqSettings = DeserializeConfig<JobQueueSettings>(info.Config);
-								JobQueueStrategy jqStrategy = new (_jobCollection, _graphCollection, _streamService, _clock, _cache, jqSettings);
+								JobQueueStrategy jqStrategy = new (_jobCollection, _graphCollection, _streamCollection, _clock, _cache, _globalConfig, jqSettings);
 								return info.ExtraAgentCount != 0 ? new ExtraAgentCountStrategy(jqStrategy, info.ExtraAgentCount) : jqStrategy;
 							
 							case PoolSizeStrategy.LeaseUtilization:
@@ -365,7 +368,7 @@ namespace Horde.Build.Agents.Fleet
 			switch (pool.SizeStrategy ?? _settings.Value.DefaultAgentPoolSizeStrategy)
 			{
 				case PoolSizeStrategy.JobQueue:
-					return new JobQueueStrategy(_jobCollection, _graphCollection, _streamService, _clock, _cache, pool.JobQueueSettings);
+					return new JobQueueStrategy(_jobCollection, _graphCollection, _streamCollection, _clock, _cache, _globalConfig, pool.JobQueueSettings);
 				case PoolSizeStrategy.LeaseUtilization:
 					LeaseUtilizationSettings luSettings = new();
 					if (pool.MinAgents != null) luSettings.MinAgents = pool.MinAgents.Value;
diff --git a/Engine/Source/Programs/Horde/Horde.Build/Agents/Fleet/JobQueueStrategy.cs b/Engine/Source/Programs/Horde/Horde.Build/Agents/Fleet/JobQueueStrategy.cs
--- a/Engine/Source/Programs/Horde/Horde.Build/Agents/Fleet/JobQueueStrategy.cs
+++ b/Engine/Source/Programs/Horde/Horde.Build/Agents/Fleet/JobQueueStrategy.cs
@@ -9,10 +9,12 @@ using System.Threading.Tasks;
 using Horde.Build.Agents.Pools;
 using Horde.Build.Jobs;
 using Horde.Build.Jobs.Graphs;
+using Horde.Build.Server;
 using Horde.Build.Streams;
 using Horde.Build.Utilities;
 using HordeCommon;
 using Microsoft.Extensions.Caching.Memory;
+using Microsoft.Extensions.Options;
 using OpenTracing;
 using OpenTracing.Util;
 
@@ -97,26 +99,29 @@ namespace Horde.Build.Agents.Fleet
 		
 		private readonly IJobCollection _jobs;
 		private readonly IGraphCollection _graphs;
-		private readonly StreamService _streamService;
+		private readonly IStreamCollection _streamCollection;
 		private readonly IClock _clock;
 		private readonly IMemoryCache _cache;
-		
+		private readonly IOptionsMonitor<GlobalConfig> _globalConfig;
+
 		/// <summary>
 		/// Constructor
 		/// </summary>
 		/// <param name="jobs"></param>
 		/// <param name="graphs"></param>
-		/// <param name="streamService"></param>
+		/// <param name="streamCollection"></param>
 		/// <param name="clock"></param>
 		/// <param name="cache"></param>
+		/// <param name="globalConfig"></param>
 		/// <param name="settings"></param>
-		public JobQueueStrategy(IJobCollection jobs, IGraphCollection graphs, StreamService streamService, IClock clock, IMemoryCache cache, JobQueueSettings? settings = null)
+		public JobQueueStrategy(IJobCollection jobs, IGraphCollection graphs, IStreamCollection streamCollection, IClock clock, IMemoryCache cache, IOptionsMonitor<GlobalConfig> globalConfig, JobQueueSettings? settings = null)
 		{
 			_jobs = jobs;
 			_graphs = graphs;
-			_streamService = streamService;
+			_streamCollection = streamCollection;
 			_clock = clock;
 			_cache = cache;
+			_globalConfig = globalConfig;
 			Settings = settings ?? new JobQueueSettings();
 		}
 
@@ -129,7 +134,7 @@ namespace Horde.Build.Agents.Fleet
 		/// <param name="job">Job to extract from</param>
 		/// <param name="streams">Cached lookup table of streams</param>
 		/// <returns></returns>
-		private async Task<List<(IJob Job, IJobStepBatch Batch, PoolId PoolId)>> GetJobBatchesWithPools(IJob job, Dictionary<StreamId, IStream> streams)
+		private async Task<List<(IJob Job, IJobStepBatch Batch, PoolId PoolId)>> GetJobBatchesWithPools(IJob job, Dictionary<StreamId, StreamConfig> streams)
 		{
 			IGraph graph = await _graphs.GetAsync(job.GraphHash);
 
@@ -151,13 +156,13 @@ namespace Horde.Build.Agents.Fleet
 					continue;
 				}
 
-				if (!streams.TryGetValue(job.StreamId, out IStream? stream))
+				if (!streams.TryGetValue(job.StreamId, out StreamConfig? streamConfig))
 				{
 					continue;
 				}
 
 				string batchAgentType = graph.Groups[batch.GroupIdx].AgentType;
-				if (!stream.Config.AgentTypes.TryGetValue(batchAgentType, out AgentConfig? agentType))
+				if (!streamConfig.AgentTypes.TryGetValue(batchAgentType, out AgentConfig? agentType))
 				{
 					continue;
 				}
@@ -171,9 +176,8 @@ namespace Horde.Build.Agents.Fleet
 		internal async Task<Dictionary<PoolId, int>> GetPoolQueueSizesAsync(DateTimeOffset jobsCreatedAfter)
 		{
 			using IScope scope = GlobalTracer.Instance.BuildSpan("JobQueueStrategy.GetPoolQueueSizes").StartActive();
-			
-			List<IStream> streamsList = await _streamService.GetStreamsAsync();
-			Dictionary<StreamId, IStream> streams = streamsList.ToDictionary(x => x.Id, x => x);
+
+			Dictionary<StreamId, StreamConfig> streams = _globalConfig.CurrentValue.Streams.ToDictionary(x => x.Id, x => (StreamConfig)x);
 			List<IJob> recentJobs = await _jobs.FindAsync(minCreateTime: jobsCreatedAfter);
 
 			List<(IJob Job, IJobStepBatch Batch, PoolId PoolId)> jobBatches = new();
diff --git a/Engine/Source/Programs/Horde/Horde.Build/Agents/IAgent.cs b/Engine/Source/Programs/Horde/Horde.Build/Agents/IAgent.cs
--- a/Engine/Source/Programs/Horde/Horde.Build/Agents/IAgent.cs
+++ b/Engine/Source/Programs/Horde/Horde.Build/Agents/IAgent.cs
@@ -2,6 +2,7 @@
 
 using System;
 using System.Collections.Generic;
+using System.Diagnostics.CodeAnalysis;
 using System.Linq;
 using System.Security.Cryptography;
 using System.Threading.Tasks;
@@ -660,15 +661,15 @@ namespace Horde.Build.Agents
 		/// Gets all the autosdk workspaces required for an agent
 		/// </summary>
 		/// <param name="agent"></param>
-		/// <param name="globals"></param>
+		/// <param name="globalConfig"></param>
 		/// <param name="workspaces"></param>
 		/// <returns></returns>
-		public static HashSet<AgentWorkspace> GetAutoSdkWorkspaces(this IAgent agent, IGlobals globals, List<AgentWorkspace> workspaces)
+		public static HashSet<AgentWorkspace> GetAutoSdkWorkspaces(this IAgent agent, GlobalConfig globalConfig, List<AgentWorkspace> workspaces)
 		{
 			HashSet<AgentWorkspace> autoSdkWorkspaces = new HashSet<AgentWorkspace>();
 			foreach (string? clusterName in workspaces.Select(x => x.Cluster).Distinct())
 			{
-				PerforceCluster? cluster = globals.Config.FindPerforceCluster(clusterName);
+				PerforceCluster? cluster = globalConfig.FindPerforceCluster(clusterName);
 				if (cluster != null)
 				{
 					AgentWorkspace? autoSdkWorkspace = GetAutoSdkWorkspace(agent, cluster);
@@ -774,5 +775,52 @@ namespace Horde.Build.Agents
 			workspaceMessages.Add(result);
 			return true;
 		}
+
+		/// <summary>
+		/// Tries to get an agent workspace definition from the given type name
+		/// </summary>
+		/// <param name="streamConfig">The stream object</param>
+		/// <param name="agentType">The agent type</param>
+		/// <param name="workspace">Receives the agent workspace definition</param>
+		/// <returns>True if the agent type was valid, and an agent workspace could be created</returns>
+		public static bool TryGetAgentWorkspace(this StreamConfig streamConfig, AgentConfig agentType, [NotNullWhen(true)] out (AgentWorkspace, bool)? workspace)
+		{
+			// Get the workspace settings
+			if (agentType.Workspace == null)
+			{
+				// Use the default settings (fast switching workspace, clean 
+				workspace = (new AgentWorkspace(null, null, streamConfig.GetDefaultWorkspaceIdentifier(), streamConfig.Name, null, false, null), true);
+				return true;
+			}
+			else
+			{
+				// Try to get the matching workspace type
+				WorkspaceConfig? workspaceType;
+				if (!streamConfig.WorkspaceTypes.TryGetValue(agentType.Workspace, out workspaceType))
+				{
+					workspace = null;
+					return false;
+				}
+
+				// Get the workspace identifier
+				string identifier;
+				if (workspaceType.Identifier != null)
+				{
+					identifier = workspaceType.Identifier;
+				}
+				else if (workspaceType.Incremental)
+				{
+					identifier = $"{streamConfig.GetEscapedName()}+{agentType.Workspace}";
+				}
+				else
+				{
+					identifier = streamConfig.GetDefaultWorkspaceIdentifier();
+				}
+
+				// Create the new workspace
+				workspace = (new AgentWorkspace(workspaceType.Cluster, workspaceType.UserName, identifier, workspaceType.Stream ?? streamConfig.Name, workspaceType.View, workspaceType.Incremental, workspaceType.Method), workspaceType.UseAutoSdk);
+				return true;
+			}
+		}
 	}
 }
diff --git a/Engine/Source/Programs/Horde/Horde.Build/Agents/Leases/LeasesController.cs b/Engine/Source/Programs/Horde/Horde.Build/Agents/Leases/LeasesController.cs
--- a/Engine/Source/Programs/Horde/Horde.Build/Agents/Leases/LeasesController.cs
+++ b/Engine/Source/Programs/Horde/Horde.Build/Agents/Leases/LeasesController.cs
@@ -6,9 +6,11 @@ using System.Linq;
 using System.Threading.Tasks;
 using Horde.Build.Acls;
 using Horde.Build.Agents.Sessions;
+using Horde.Build.Server;
 using Horde.Build.Utilities;
 using Microsoft.AspNetCore.Authorization;
 using Microsoft.AspNetCore.Mvc;
+using Microsoft.Extensions.Options;
 using OpenTracing;
 using OpenTracing.Util;
 
@@ -25,25 +27,16 @@ namespace Horde.Build.Agents.Leases
 	[Route("[controller]")]
 	public class LeasesController : HordeControllerBase
 	{
-		/// <summary>
-		/// Singleton instance of the ACL service
-		/// </summary>
-		readonly AclService _aclService;
-
-		/// <summary>
-		/// Singleton instance of the agent service
-		/// </summary>
 		readonly AgentService _agentService;
+		readonly IOptionsSnapshot<GlobalConfig> _globalConfig;
 
 		/// <summary>
 		/// Constructor
 		/// </summary>
-		/// <param name="aclService">The ACL service singleton</param>
-		/// <param name="agentService">The agent service</param>
-		public LeasesController(AclService aclService, AgentService agentService)
+		public LeasesController(AgentService agentService, IOptionsSnapshot<GlobalConfig> globalConfig)
 		{
-			_aclService = aclService;
 			_agentService = agentService;
+			_globalConfig = globalConfig;
 		}
 
 		/// <summary>
@@ -73,28 +66,12 @@ namespace Horde.Build.Agents.Leases
 			[FromQuery] int count = 1000,
 			[FromQuery] PropertyFilter? filter = null)
 		{
-			GlobalPermissionsCache permissionsCache = new GlobalPermissionsCache();
-			if (agentId == null)
+			if (!_globalConfig.Value.Authorize(AclAction.ViewLeases, User))
 			{
-				if (!await _aclService.AuthorizeAsync(AclAction.ViewLeases, User, permissionsCache))
-				{
-					return Forbid(AclAction.ViewLeases);
-				}
-			}
-			else
-			{
-				IAgent? agent = await _agentService.GetAgentAsync(agentId.Value);
-				if (agent == null)
-				{
-					return NotFound(agentId.Value);
-				}
-				if (!await _agentService.AuthorizeAsync(agent, AclAction.ViewLeases, User, permissionsCache))
-				{
-					return Forbid(AclAction.ViewLeases, agentId.Value);
-				}
+				return Forbid(AclAction.ViewLeases);
 			}
 
-			bool includeCosts = await _aclService.AuthorizeAsync(AclAction.ViewCosts, User, permissionsCache);
+			bool includeCosts = _globalConfig.Value.Authorize(AclAction.ViewCosts, User);
 
 			List<ILease> leases;
 			if (minFinishTime == null && maxFinishTime == null)
@@ -138,6 +115,11 @@ namespace Horde.Build.Agents.Leases
 		[Route("/api/v1/leases/{leaseId}")]
 		public async Task<ActionResult<GetAgentLeaseResponse>> GetLeaseAsync(LeaseId leaseId)
 		{
+			if (!_globalConfig.Value.Authorize(AclAction.ViewLeases, User))
+			{
+				return Forbid(AclAction.ViewLeases);
+			}
+
 			ILease? lease = await _agentService.GetLeaseAsync(leaseId);
 			if (lease == null)
 			{
@@ -150,14 +132,8 @@ namespace Horde.Build.Agents.Leases
 				return NotFound(lease.AgentId);
 			}
 
-			GlobalPermissionsCache permissionsCache = new GlobalPermissionsCache();
-			if (!await _agentService.AuthorizeAsync(agent, AclAction.ViewLeases, User, permissionsCache))
-			{
-				return Forbid(AclAction.ViewLeases);
-			}
-
 			double? agentRate = null;
-			if (await _aclService.AuthorizeAsync(AclAction.ViewCosts, User, permissionsCache))
+			if (_globalConfig.Value.Authorize(AclAction.ViewCosts, User))
 			{
 				agentRate = await _agentService.GetRateAsync(agent.Id);
 			}
@@ -176,6 +152,11 @@ namespace Horde.Build.Agents.Leases
 		[Route("/api/v1/leases/{leaseId}")]
 		public async Task<ActionResult> UpdateLeaseAsync(LeaseId leaseId, [FromBody] UpdateLeaseRequest request)
 		{
+			if (!_globalConfig.Value.Authorize(AclAction.AdminWrite, User))
+			{
+				return Forbid(AclAction.AdminWrite);
+			}
+
 			// only update supported right now is abort
 			if (!request.Aborted.HasValue || !request.Aborted.Value)
 			{
@@ -194,11 +175,6 @@ namespace Horde.Build.Agents.Leases
 				return NotFound(lease.AgentId);
 			}
 
-			if (!await _agentService.AuthorizeAsync(agent, AclAction.AdminWrite, User, null))
-			{
-				return Forbid(AclAction.AdminWrite, lease.AgentId);
-			}
-
 			AgentLease? agentLease = agent.Leases.FirstOrDefault(x => x.Id == leaseId);
 
 			if (agentLease == null)
diff --git a/Engine/Source/Programs/Horde/Horde.Build/Agents/Pools/PoolService.cs b/Engine/Source/Programs/Horde/Horde.Build/Agents/Pools/PoolService.cs
--- a/Engine/Source/Programs/Horde/Horde.Build/Agents/Pools/PoolService.cs
+++ b/Engine/Source/Programs/Horde/Horde.Build/Agents/Pools/PoolService.cs
@@ -246,8 +246,9 @@ namespace Horde.Build.Agents.Pools
 		/// </summary>
 		/// <param name="agent">The agent to return workspaces for</param>
 		/// <param name="validAtTime">Absolute time at which we expect the results to be valid. Values may be cached as long as they are after this time.</param>
+		/// <param name="globalConfig">Current configuration</param>
 		/// <returns>List of workspaces</returns>
-		public async Task<HashSet<AgentWorkspace>> GetWorkspacesAsync(IAgent agent, DateTime validAtTime)
+		public async Task<HashSet<AgentWorkspace>> GetWorkspacesAsync(IAgent agent, DateTime validAtTime, GlobalConfig globalConfig)
 		{
 			bool addAutoSdkWorkspace = false;
 			HashSet<AgentWorkspace> workspaces = new HashSet<AgentWorkspace>();
@@ -265,8 +266,7 @@ namespace Horde.Build.Agents.Pools
 
 			if (addAutoSdkWorkspace)
 			{
-				IGlobals globals = await _globalsService.GetAsync();
-				workspaces.UnionWith(agent.GetAutoSdkWorkspaces(globals, workspaces.ToList()));
+				workspaces.UnionWith(agent.GetAutoSdkWorkspaces(globalConfig, workspaces.ToList()));
 			}
 
 			return workspaces;
diff --git a/Engine/Source/Programs/Horde/Horde.Build/Agents/Pools/PoolUpdateService.cs b/Engine/Source/Programs/Horde/Horde.Build/Agents/Pools/PoolUpdateService.cs
--- a/Engine/Source/Programs/Horde/Horde.Build/Agents/Pools/PoolUpdateService.cs
+++ b/Engine/Source/Programs/Horde/Horde.Build/Agents/Pools/PoolUpdateService.cs
@@ -5,11 +5,14 @@ using System.Collections.Generic;
 using System.Linq;
 using System.Threading;
 using System.Threading.Tasks;
+using Horde.Build.Projects;
+using Horde.Build.Server;
 using Horde.Build.Streams;
 using Horde.Build.Utilities;
 using HordeCommon;
 using Microsoft.Extensions.Hosting;
 using Microsoft.Extensions.Logging;
+using Microsoft.Extensions.Options;
 
 namespace Horde.Build.Agents.Pools
 {
@@ -22,16 +25,18 @@ namespace Horde.Build.Agents.Pools
 	{
 		readonly IPoolCollection _pools;
 		readonly IStreamCollection _streams;
+		readonly IOptionsMonitor<GlobalConfig> _globalConfig;
 		readonly ILogger<PoolService> _logger;
 		readonly ITicker _ticker;
 
 		/// <summary>
 		/// Constructor
 		/// </summary>
-		public PoolUpdateService(IPoolCollection pools, IStreamCollection streams, IClock clock, ILogger<PoolService> logger)
+		public PoolUpdateService(IPoolCollection pools, IStreamCollection streams, IClock clock, IOptionsMonitor<GlobalConfig> globalConfig, ILogger<PoolService> logger)
 		{
 			_pools = pools;
 			_streams = streams;
+			_globalConfig = globalConfig;
 			_logger = logger;
 			_ticker = clock.AddTicker($"{nameof(PoolUpdateService)}.{nameof(UpdatePoolsAsync)}", TimeSpan.FromSeconds(30.0), UpdatePoolsAsync, logger);
 		}
@@ -71,15 +76,17 @@ namespace Horde.Build.Agents.Pools
 				HashSet<PoolId> poolsWithAutoSdk = new HashSet<PoolId>();
 				Dictionary<PoolId, List<AgentWorkspace>> poolToAgentWorkspaces = new Dictionary<PoolId, List<AgentWorkspace>>();
 
+				// Capture the current config state
+				GlobalConfig globalConfig = _globalConfig.CurrentValue;
+
 				// Populate the workspace list from the current stream
-				List<IStream> activeStreams = await _streams.FindAllAsync();
-				foreach (IStream activeStream in activeStreams)
+				foreach (StreamConfig streamConfig in globalConfig.Streams)
 				{
-					foreach (KeyValuePair<string, AgentConfig> agentTypePair in activeStream.Config.AgentTypes)
+					foreach (KeyValuePair<string, AgentConfig> agentTypePair in streamConfig.AgentTypes)
 					{
 						// Create the new agent workspace
 						(AgentWorkspace, bool)? result;
-						if (activeStream.TryGetAgentWorkspace(agentTypePair.Value, out result))
+						if (streamConfig.TryGetAgentWorkspace(agentTypePair.Value, out result))
 						{
 							(AgentWorkspace agentWorkspace, bool useAutoSdk) = result.Value;
 							AgentConfig agentType = agentTypePair.Value;
diff --git a/Engine/Source/Programs/Horde/Horde.Build/Agents/Pools/PoolsController.cs b/Engine/Source/Programs/Horde/Horde.Build/Agents/Pools/PoolsController.cs
--- a/Engine/Source/Programs/Horde/Horde.Build/Agents/Pools/PoolsController.cs
+++ b/Engine/Source/Programs/Horde/Horde.Build/Agents/Pools/PoolsController.cs
@@ -6,9 +6,11 @@ using System.Linq;
 using System.Threading.Tasks;
 using Horde.Build.Acls;
 using Horde.Build.Agents.Fleet;
+using Horde.Build.Server;
 using Horde.Build.Utilities;
 using Microsoft.AspNetCore.Authorization;
 using Microsoft.AspNetCore.Mvc;
+using Microsoft.Extensions.Options;
 
 namespace Horde.Build.Agents.Pools
 {
@@ -22,25 +24,16 @@ namespace Horde.Build.Agents.Pools
 	[Route("[controller]")]
 	public class PoolsController : HordeControllerBase
 	{
-		/// <summary>
-		/// Singleton instance of the ACL service
-		/// </summary>
-		private readonly AclService _aclService;
-
-		/// <summary>
-		/// Singleton instance of the pool service
-		/// </summary>
 		private readonly PoolService _poolService;
+		private readonly IOptionsSnapshot<GlobalConfig> _globalConfig;
 
 		/// <summary>
 		/// Constructor
 		/// </summary>
-		/// <param name="aclService">The ACL service</param>
-		/// <param name="poolService">The pool service</param>
-		public PoolsController(AclService aclService, PoolService poolService)
+		public PoolsController(PoolService poolService, IOptionsSnapshot<GlobalConfig> globalConfig)
 		{
-			_aclService = aclService;
 			_poolService = poolService;
+			_globalConfig = globalConfig;
 		}
 
 		/// <summary>
@@ -52,7 +45,7 @@ namespace Horde.Build.Agents.Pools
 		[Route("/api/v1/pools")]
 		public async Task<ActionResult<CreatePoolResponse>> CreatePoolAsync([FromBody] CreatePoolRequest create)
 		{
-			if(!await _aclService.AuthorizeAsync(AclAction.CreatePool, User))
+			if(!_globalConfig.Value.Authorize(AclAction.CreatePool, User))
 			{
 				return Forbid(AclAction.CreatePool);
 			}
@@ -83,7 +76,7 @@ namespace Horde.Build.Agents.Pools
 		[ProducesResponseType(typeof(List<GetPoolResponse>), 200)]
 		public async Task<ActionResult<List<object>>> GetPoolsAsync([FromQuery] PropertyFilter? filter = null)
 		{
-			if (!await _aclService.AuthorizeAsync(AclAction.ListPools, User))
+			if (!_globalConfig.Value.Authorize(AclAction.ListPools, User))
 			{
 				return Forbid(AclAction.ListPools);
 			}
@@ -109,7 +102,7 @@ namespace Horde.Build.Agents.Pools
 		[ProducesResponseType(typeof(GetPoolResponse), 200)]
 		public async Task<ActionResult<object>> GetPoolAsync(string poolId, [FromQuery] PropertyFilter? filter = null)
 		{
-			if (!await _aclService.AuthorizeAsync(AclAction.ViewPool, User))
+			if (!_globalConfig.Value.Authorize(AclAction.ViewPool, User))
 			{
 				return Forbid(AclAction.ViewPool);
 			}
@@ -135,7 +128,7 @@ namespace Horde.Build.Agents.Pools
 		[Route("/api/v1/pools/{poolId}")]
 		public async Task<ActionResult> UpdatePoolAsync(string poolId, [FromBody] UpdatePoolRequest update)
 		{
-			if (!await _aclService.AuthorizeAsync(AclAction.UpdatePool, User))
+			if (!_globalConfig.Value.Authorize(AclAction.UpdatePool, User))
 			{
 				return Forbid(AclAction.UpdatePool);
 			}
@@ -169,7 +162,7 @@ namespace Horde.Build.Agents.Pools
 		[Route("/api/v1/pools/{poolId}")]
 		public async Task<ActionResult> DeletePoolAsync(string poolId)
 		{
-			if (!await _aclService.AuthorizeAsync(AclAction.DeletePool, User))
+			if (!_globalConfig.Value.Authorize(AclAction.DeletePool, User))
 			{
 				return Forbid(AclAction.DeletePool);
 			}
@@ -191,7 +184,7 @@ namespace Horde.Build.Agents.Pools
 		[Route("/api/v1/pools")]
 		public async Task<ActionResult> UpdatePoolAsync([FromBody] List<BatchUpdatePoolRequest> batchUpdates)
 		{
-			if (!await _aclService.AuthorizeAsync(AclAction.UpdatePool, User))
+			if (!_globalConfig.Value.Authorize(AclAction.UpdatePool, User))
 			{
 				return Forbid(AclAction.UpdatePool);
 			}
diff --git a/Engine/Source/Programs/Horde/Horde.Build/Agents/Software/AgentSoftwareController.cs b/Engine/Source/Programs/Horde/Horde.Build/Agents/Software/AgentSoftwareController.cs
--- a/Engine/Source/Programs/Horde/Horde.Build/Agents/Software/AgentSoftwareController.cs
+++ b/Engine/Source/Programs/Horde/Horde.Build/Agents/Software/AgentSoftwareController.cs
@@ -4,10 +4,12 @@ using System.Collections.Generic;
 using System.IO;
 using System.Threading.Tasks;
 using Horde.Build.Acls;
+using Horde.Build.Server;
 using Horde.Build.Utilities;
 using Microsoft.AspNetCore.Authorization;
 using Microsoft.AspNetCore.Http;
 using Microsoft.AspNetCore.Mvc;
+using Microsoft.Extensions.Options;
 using Microsoft.Net.Http.Headers;
 
 namespace Horde.Build.Agents.Software
@@ -22,25 +24,16 @@ namespace Horde.Build.Agents.Software
 	[Route("[controller]")]
 	public class AgentSoftwareController : ControllerBase
 	{
-		/// <summary>
-		/// Singleton instance of the ACL service
-		/// </summary>
-		private readonly AclService _aclService;
-
-		/// <summary>
-		/// Singleton instance of the client service
-		/// </summary>
 		private readonly AgentSoftwareService _agentSoftwareService;
+		private readonly IOptionsSnapshot<GlobalConfig> _globalConfig;
 
 		/// <summary>
 		/// Constructor
 		/// </summary>
-		/// <param name="aclService">The ACL service</param>
-		/// <param name="agentSoftwareService">The client service</param>
-		public AgentSoftwareController(AclService aclService, AgentSoftwareService agentSoftwareService)
+		public AgentSoftwareController(AgentSoftwareService agentSoftwareService, IOptionsSnapshot<GlobalConfig> globalConfig)
 		{
-			_aclService = aclService;
 			_agentSoftwareService = agentSoftwareService;
+			_globalConfig = globalConfig;
 		}
 
 		/// <summary>
@@ -53,7 +46,7 @@ namespace Horde.Build.Agents.Software
 		[ProducesResponseType(typeof(List<GetAgentSoftwareChannelResponse>), 200)]
 		public async Task<ActionResult<List<object>>> FindSoftwareAsync([FromQuery] PropertyFilter? filter = null)
 		{
-			if (!await _aclService.AuthorizeAsync(AclAction.DownloadSoftware, User))
+			if (!_globalConfig.Value.Authorize(AclAction.DownloadSoftware, User))
 			{
 				return Forbid();
 			}
@@ -79,7 +72,7 @@ namespace Horde.Build.Agents.Software
 		[ProducesResponseType(typeof(GetAgentSoftwareChannelResponse), 200)]
 		public async Task<ActionResult<object>> FindSoftwareAsync(string name, [FromQuery] PropertyFilter? filter = null)
 		{
-			if (!await _aclService.AuthorizeAsync(AclAction.DownloadSoftware, User))
+			if (!_globalConfig.Value.Authorize(AclAction.DownloadSoftware, User))
 			{
 				return Forbid();
 			}
@@ -104,7 +97,7 @@ namespace Horde.Build.Agents.Software
 		[Route("/api/v1/agentsoftware/{name}")]
 		public async Task<ActionResult<object>> SetSoftwareAsync(string name, [FromBody] UpdateAgentSoftwareChannelRequest request, [FromQuery] PropertyFilter? filter = null)
 		{
-			if (!await _aclService.AuthorizeAsync(AclAction.UploadSoftware, User))
+			if (!_globalConfig.Value.Authorize(AclAction.UploadSoftware, User))
 			{
 				return Forbid();
 			}
@@ -123,7 +116,7 @@ namespace Horde.Build.Agents.Software
 		[Route("/api/v1/agentsoftware/{name}/zip")]
 		public async Task<ActionResult<UploadAgentSoftwareResponse>> UploadArchiveAsync(string name, [FromForm] IFormFile file)
 		{
-			if (!await _aclService.AuthorizeAsync(AclAction.UploadSoftware, User))
+			if (!_globalConfig.Value.Authorize(AclAction.UploadSoftware, User))
 			{
 				return Forbid();
 			}
@@ -153,7 +146,7 @@ namespace Horde.Build.Agents.Software
 		[Route("/api/v1/agentsoftware/{name}/zip")]
 		public async Task<ActionResult> GetArchiveAsync(string name)
 		{
-			if (!await _aclService.AuthorizeAsync(AclAction.DownloadSoftware, User))
+			if (!_globalConfig.Value.Authorize(AclAction.DownloadSoftware, User))
 			{
 				return Forbid();
 			}
@@ -165,7 +158,6 @@ namespace Horde.Build.Agents.Software
 			}
 
 			return new FileStreamResult(new MemoryStream(data), new MediaTypeHeaderValue("application/octet-stream")) { FileDownloadName = $"HordeAgent.zip" };
-			
 		}
 	}
 }
diff --git a/Engine/Source/Programs/Horde/Horde.Build/Commands/Bundles/PerforceCommand.cs b/Engine/Source/Programs/Horde/Horde.Build/Commands/Bundles/PerforceCommand.cs
--- a/Engine/Source/Programs/Horde/Horde.Build/Commands/Bundles/PerforceCommand.cs
+++ b/Engine/Source/Programs/Horde/Horde.Build/Commands/Bundles/PerforceCommand.cs
@@ -6,11 +6,13 @@ using System.Threading.Tasks;
 using EpicGames.Core;
 using Horde.Build.Configuration;
 using Horde.Build.Perforce;
+using Horde.Build.Server;
 using Horde.Build.Streams;
 using Horde.Build.Utilities;
 using Microsoft.Extensions.Configuration;
 using Microsoft.Extensions.DependencyInjection;
 using Microsoft.Extensions.Logging;
+using Microsoft.Extensions.Options;
 
 namespace Horde.Build.Commands.Bundles
 {
@@ -27,31 +29,30 @@ namespace Horde.Build.Commands.Bundles
 
 		readonly IConfiguration _configuration;
 		readonly ILoggerProvider _loggerProvider;
+		readonly IOptionsSnapshot<GlobalConfig> _globalConfig;
 
-		public PerforceCommand(IConfiguration configuration, ILoggerProvider loggerProvider)
+		public PerforceCommand(IConfiguration configuration, ILoggerProvider loggerProvider, IOptionsSnapshot<GlobalConfig> globalConfig)
 		{
 			_configuration = configuration;
 			_loggerProvider = loggerProvider;
+			_globalConfig = globalConfig;
 		}
 
 		public override async Task<int> ExecuteAsync(ILogger logger)
 		{
 			using ServiceProvider serviceProvider = Startup.CreateServiceProvider(_configuration, _loggerProvider);
 
-			ConfigUpdateService configUpdateService = serviceProvider.GetRequiredService<ConfigUpdateService>();
-			await configUpdateService.ForceUpdateAsync(CancellationToken.None);
-
 			PerforceReplicator replicator = serviceProvider.GetRequiredService<PerforceReplicator>();
 			IStreamCollection streamCollection = serviceProvider.GetRequiredService<IStreamCollection>();
 
-			IStream? stream = await streamCollection.GetAsync(new StreamId(StreamId));
-			if (stream == null)
+			StreamConfig? streamConfig;
+			if (!_globalConfig.Value.TryGetStream(new StreamId(StreamId), out streamConfig))
 			{
-				throw new FatalErrorException($"Stream '{stream}' not found");
+				throw new FatalErrorException($"Stream '{StreamId}' not found");
 			}
 
 			PerforceReplicationOptions options = new PerforceReplicationOptions();
-			await replicator.WriteAsync(stream, Change, options, default);
+			await replicator.WriteAsync(streamConfig, Change, options, default);
 
 			return 0;
 		}
diff --git a/Engine/Source/Programs/Horde/Horde.Build/Compute/ComputeController.cs b/Engine/Source/Programs/Horde/Horde.Build/Compute/ComputeController.cs
--- a/Engine/Source/Programs/Horde/Horde.Build/Compute/ComputeController.cs
+++ b/Engine/Source/Programs/Horde/Horde.Build/Compute/ComputeController.cs
@@ -6,9 +6,11 @@ using System.Threading.Tasks;
 using EpicGames.Horde.Compute;
 using EpicGames.Horde.Compute.Impl;
 using Horde.Build.Acls;
+using Horde.Build.Server;
 using Horde.Build.Utilities;
 using Microsoft.AspNetCore.Authorization;
 using Microsoft.AspNetCore.Mvc;
+using Microsoft.Extensions.Options;
 
 namespace Horde.Build.Compute
 {
@@ -20,18 +22,16 @@ namespace Horde.Build.Compute
 	[Route("[controller]")]
 	public class ComputeController : HordeControllerBase
 	{
-		readonly AclService _aclService;
 		readonly ComputeService _computeService;
+		readonly IOptionsSnapshot<GlobalConfig> _globalConfig;
 
 		/// <summary>
 		/// Constructor
 		/// </summary>
-		/// <param name="aclService"></param>
-		/// <param name="computeService">The compute service singleton</param>
-		public ComputeController(AclService aclService, ComputeService computeService)
+		public ComputeController(ComputeService computeService, IOptionsSnapshot<GlobalConfig> globalConfig)
 		{
-			_aclService = aclService;
 			_computeService = computeService;
+			_globalConfig = globalConfig;
 		}
 
 		/// <summary>
@@ -44,7 +44,7 @@ namespace Horde.Build.Compute
 		[Route("/api/v1/compute/{clusterId}")]
 		public async Task<ActionResult<GetComputeClusterInfo>> GetClusterInfoAsync([FromRoute] ClusterId clusterId)
 		{
-			if (!await _aclService.AuthorizeAsync(AclAction.ViewComputeTasks, User))
+			if (!_globalConfig.Value.Authorize(AclAction.ViewComputeTasks, User))
 			{
 				return Forbid(AclAction.ViewComputeTasks);
 			}
@@ -79,7 +79,7 @@ namespace Horde.Build.Compute
 		[Route("/api/v1/compute/{clusterId}")]
 		public async Task<ActionResult> AddTasksAsync([FromRoute] ClusterId clusterId, [FromBody] AddTasksRequest request)
 		{
-			if(!await _aclService.AuthorizeAsync(AclAction.AddComputeTasks, User))
+			if(!_globalConfig.Value.Authorize(AclAction.AddComputeTasks, User))
 			{
 				return Forbid(AclAction.AddComputeTasks);
 			}
@@ -104,7 +104,7 @@ namespace Horde.Build.Compute
 		[Route("/api/v1/compute/{clusterId}/updates/{channelId}")]
 		public async Task<ActionResult<GetTaskUpdatesResponse>> GetUpdatesAsync([FromRoute] ClusterId clusterId, [FromRoute] ChannelId channelId, [FromQuery] int wait = 0)
 		{
-			if (!await _aclService.AuthorizeAsync(AclAction.ViewComputeTasks, User))
+			if (!_globalConfig.Value.Authorize(AclAction.ViewComputeTasks, User))
 			{
 				return Forbid(AclAction.ViewComputeTasks);
 			}
diff --git a/Engine/Source/Programs/Horde/Horde.Build/Compute/ComputeService.cs b/Engine/Source/Programs/Horde/Horde.Build/Compute/ComputeService.cs
--- a/Engine/Source/Programs/Horde/Horde.Build/Compute/ComputeService.cs
+++ b/Engine/Source/Programs/Horde/Horde.Build/Compute/ComputeService.cs
@@ -23,6 +23,7 @@ using HordeCommon.Rpc.Tasks;
 using Microsoft.Extensions.Caching.Memory;
 using Microsoft.Extensions.Hosting;
 using Microsoft.Extensions.Logging;
+using Microsoft.Extensions.Options;
 using OpenTracing;
 using OpenTracing.Util;
 using StackExchange.Redis;
@@ -126,7 +127,7 @@ namespace Horde.Build.Compute
 		readonly RedisMessageQueue<ComputeTaskStatus> _messageQueue;
 		readonly ITicker _expireTasksTicker;
 		readonly IMemoryCache _requirementsCache;
-		readonly LazyCachedValue<Task<IGlobals>> _globals;
+		readonly IOptionsMonitor<GlobalConfig> _globalConfig;
 		readonly ILogger _logger;
 
 		static ComputeService()
@@ -137,14 +138,14 @@ namespace Horde.Build.Compute
 		/// <summary>
 		/// Constructor
 		/// </summary>
-		public ComputeService(GlobalsService globalsService, RedisService redisService, ILegacyStorageClient storageClient, IClock clock, ILogger<ComputeService> logger)
+		public ComputeService(GlobalsService globalsService, RedisService redisService, ILegacyStorageClient storageClient, IClock clock, IOptionsMonitor<GlobalConfig> globalConfig, ILogger<ComputeService> logger)
 		{
 			_storageClient = storageClient;
 			_taskScheduler = new RedisTaskScheduler<QueueKey, ComputeTaskInfo>(redisService.ConnectionPool, "compute/tasks/", logger);
 			_messageQueue = new RedisMessageQueue<ComputeTaskStatus>(redisService.GetDatabase(), "compute/messages/");
 			_expireTasksTicker = clock.AddTicker<ComputeService>(TimeSpan.FromMinutes(2.0), ExpireTasksAsync, logger);
 			_requirementsCache = new MemoryCache(new MemoryCacheOptions());
-			_globals = new LazyCachedValue<Task<IGlobals>>(async () => await globalsService.GetAsync(), TimeSpan.FromSeconds(120.0));
+			_globalConfig = globalConfig;
 			_logger = logger;
 
 			OnLeaseStartedProperties.Add(x => x.TaskRefId);
@@ -206,14 +207,14 @@ namespace Horde.Build.Compute
 		}
 
 		/// <inheritdoc/>
-		public async Task<IComputeClusterInfo> GetClusterInfoAsync(ClusterId clusterId)
+		public Task<IComputeClusterInfo> GetClusterInfoAsync(ClusterId clusterId)
 		{
-			ComputeClusterConfig? config = await GetClusterAsync(clusterId);
+			ComputeClusterConfig? config = GetCluster(clusterId);
 			if (config == null)
 			{
 				throw new KeyNotFoundException();
 			}
-			return new ClusterInfo(config);
+			return Task.FromResult<IComputeClusterInfo>(new ClusterInfo(config));
 		}
 
 		/// <inheritdoc/>
@@ -229,10 +230,9 @@ namespace Horde.Build.Compute
 			await Task.WhenAll(tasks);
 		}
 
-		async ValueTask<ComputeClusterConfig?> GetClusterAsync(ClusterId clusterId)
+		ComputeClusterConfig? GetCluster(ClusterId clusterId)
 		{
-			IGlobals globals = await _globals.GetCached();
-			return globals.Config.Compute.FirstOrDefault(x => new ClusterId(x.Id) == clusterId);
+			return _globalConfig.CurrentValue.Compute.FirstOrDefault(x => new ClusterId(x.Id) == clusterId);
 		}
 
 		/// <inheritdoc/>
@@ -266,7 +266,7 @@ namespace Horde.Build.Compute
 		{
 			(QueueKey queueKey, ComputeTaskInfo taskInfo) = entry;
 
-			ComputeClusterConfig? cluster = await GetClusterAsync(taskInfo.ClusterId);
+			ComputeClusterConfig? cluster = GetCluster(taskInfo.ClusterId);
 			if (cluster == null)
 			{
 				_logger.LogWarning("Invalid cluster '{ClusterId}'; failing task {TaskRefId}", taskInfo.ClusterId, taskInfo.TaskRefId);
@@ -452,7 +452,7 @@ namespace Horde.Build.Compute
 		{
 			Requirements? requirements = null;
 
-			ComputeClusterConfig? clusterConfig = await GetClusterAsync(queueKey.ClusterId);
+			ComputeClusterConfig? clusterConfig = GetCluster(queueKey.ClusterId);
 			if (clusterConfig != null)
 			{
 				NamespaceId namespaceId = new NamespaceId(clusterConfig.NamespaceId);
diff --git a/Engine/Source/Programs/Horde/Horde.Build/Configuration/ConfigService.cs b/Engine/Source/Programs/Horde/Horde.Build/Configuration/ConfigService.cs
--- a/Engine/Source/Programs/Horde/Horde.Build/Configuration/ConfigService.cs
+++ b/Engine/Source/Programs/Horde/Horde.Build/Configuration/ConfigService.cs
@@ -10,7 +10,6 @@ using System.Threading;
 using System.Threading.Tasks;
 using EpicGames.Core;
 using EpicGames.Redis;
-using Horde.Build.Notifications;
 using Horde.Build.Projects;
 using Horde.Build.Server;
 using Horde.Build.Streams;
@@ -25,9 +24,6 @@ using StackExchange.Redis;
 
 namespace Horde.Build.Configuration
 {
-	using ProjectId = StringId<IProject>;
-	using StreamId = StringId<IStream>;
-
 	/// <summary>
 	/// Service which processes runtime configuration data.
 	/// </summary>
@@ -38,13 +34,7 @@ namespace Horde.Build.Configuration
 		class ConfigSnapshot
 		{
 			[ProtoMember(1)]
-			public byte[] Global = null!;
-
-			[ProtoMember(2)]
-			public Dictionary<ProjectId, byte[]> Projects { get; set; } = new Dictionary<ProjectId, byte[]>();
-
-			[ProtoMember(3)]
-			public Dictionary<StreamId, byte[]> Streams { get; set; } = new Dictionary<StreamId, byte[]>();
+			public byte[] Data { get; set; } = null!;
 
 			[ProtoMember(4)]
 			public Dictionary<Uri, string> Dependencies { get; set; } = new Dictionary<Uri, string>();
@@ -114,8 +104,7 @@ namespace Horde.Build.Configuration
 		record class ConfigState(IoHash Hash, GlobalConfig GlobalConfig);
 
 		readonly RedisService _redisService;
-		readonly INotificationService _notificationService;
-		readonly ServerSettings _settings;
+		readonly ServerSettings _serverSettings;
 		readonly Dictionary<string, IConfigSource> _sources;
 		readonly JsonSerializerOptions _jsonOptions;
 		readonly RedisKey _snapshotKey = "config";
@@ -133,14 +122,18 @@ namespace Horde.Build.Configuration
 		/// <inheritdoc/>
 		string IOptionsChangeTokenSource<GlobalConfig>.Name => String.Empty;
 
+		/// <summary>
+		/// Event for notifications that the config has been updated
+		/// </summary>
+		public event Action<Exception?>? OnConfigUpdate;
+
 		/// <summary>
 		/// Constructor
 		/// </summary>
-		public ConfigService(RedisService redisService, INotificationService notificationService, IOptions<ServerSettings> settings, IEnumerable<IConfigSource> sources, IClock clock, ILogger<ConfigService> logger)
+		public ConfigService(RedisService redisService, IOptions<ServerSettings> serverSettings, IEnumerable<IConfigSource> sources, IClock clock, ILogger<ConfigService> logger)
 		{
 			_redisService = redisService;
-			_notificationService = notificationService;
-			_settings = settings.Value;
+			_serverSettings = serverSettings.Value;
 			_sources = sources.ToDictionary(x => x.Scheme, x => x, StringComparer.OrdinalIgnoreCase);
 			_logger = logger;
 
@@ -175,6 +168,28 @@ namespace Horde.Build.Configuration
 			await _ticker.StopAsync();
 		}
 
+		/// <summary>
+		/// Accessor for tests to update the global config
+		/// </summary>
+		/// <param name="hash">Hash for the config data</param>
+		/// <param name="globalConfig">New config value</param>
+		public void Set(IoHash hash, GlobalConfig globalConfig)
+		{
+			// Set the new state
+			_stateTask = Task.FromResult(new ConfigState(hash, globalConfig));
+
+			// Notify any watchers of the new config object
+			for (; ; )
+			{
+				ChangeToken? token = _currentChangeToken;
+				if (Interlocked.CompareExchange(ref _currentChangeToken, null, token) == token)
+				{
+					token?.TriggerChange();
+					break;
+				}
+			}
+		}
+
 		/// <inheritdoc/>
 		public GlobalConfig Create(string name) => _stateTask.Result.GlobalConfig;
 
@@ -202,7 +217,7 @@ namespace Horde.Build.Configuration
 		async Task<ConfigState> GetStartupStateAsync()
 		{
 			ReadOnlyMemory<byte> data = ReadOnlyMemory<byte>.Empty;
-			if (!_settings.ForceConfigUpdateOnStartup)
+			if (!_serverSettings.ForceConfigUpdateOnStartup)
 			{
 				data = await ReadSnapshotDataAsync();
 			}
@@ -266,13 +281,13 @@ namespace Horde.Build.Configuration
 			try
 			{
 				ConfigSnapshot snapshot = await CreateSnapshotAsync(cancellationToken);
-				_notificationService.NotifyConfigUpdate(null);
+				OnConfigUpdate?.Invoke(null);
 				return snapshot;
 			}
 			catch (Exception ex)
 			{
 				_logger.LogError(ex, "Exception while updating config: {Message}", ex.Message);
-				_notificationService.NotifyConfigUpdate(ex);
+				OnConfigUpdate?.Invoke(ex);
 				return null;
 			}
 		}
@@ -286,15 +301,15 @@ namespace Horde.Build.Configuration
 		{
 			// Get the path to the root config file
 			Uri globalConfigUri;
-			if (Path.IsPathRooted(_settings.ConfigPath) && !_settings.ConfigPath.StartsWith("//", StringComparison.Ordinal))
+			if (Path.IsPathRooted(_serverSettings.ConfigPath) && !_serverSettings.ConfigPath.StartsWith("//", StringComparison.Ordinal))
 			{
 				// absolute path to config
-				globalConfigUri = new Uri(_settings.ConfigPath);
+				globalConfigUri = new Uri(_serverSettings.ConfigPath);
 			}
 			else
 			{
 				// relative (development) or perforce path
-				globalConfigUri = ConfigType.CombinePaths(new Uri(FileReference.Combine(Program.AppDir, "_").FullName), _settings.ConfigPath);
+				globalConfigUri = ConfigType.CombinePaths(new Uri(FileReference.Combine(Program.AppDir, "_").FullName), _serverSettings.ConfigPath);
 			}
 
 			// Read the config files
@@ -304,23 +319,7 @@ namespace Horde.Build.Configuration
 				ConfigSnapshot snapshot = new ConfigSnapshot();
 
 				GlobalConfig globalConfig = await ConfigType.ReadAsync<GlobalConfig>(globalConfigUri, context, cancellationToken);
-				snapshot.Global = JsonSerializer.SerializeToUtf8Bytes(globalConfig, context.JsonOptions);
-
-				foreach (ProjectConfigRef projectConfigRef in globalConfig.Projects)
-				{
-					context.PropertyPathToFile.Clear();
-
-					ProjectConfig projectConfig = await ConfigType.ReadAsync<ProjectConfig>(new Uri(projectConfigRef.Path), context, cancellationToken);
-					snapshot.Projects.Add(projectConfigRef.Id, JsonSerializer.SerializeToUtf8Bytes(projectConfig, context.JsonOptions));
-
-					foreach (StreamConfigRef streamConfigRef in projectConfig.Streams)
-					{
-						context.PropertyPathToFile.Clear();
-
-						StreamConfig streamConfig = await ConfigType.ReadAsync<StreamConfig>(new Uri(streamConfigRef.Path), context, cancellationToken);
-						snapshot.Streams.Add(streamConfigRef.Id, JsonSerializer.SerializeToUtf8Bytes(streamConfig, context.JsonOptions));
-					}
-				}
+				snapshot.Data = JsonSerializer.SerializeToUtf8Bytes(globalConfig, context.JsonOptions);
 
 				// Save all the dependencies
 				foreach ((Uri depUri, IConfigFile depFile) in context.Files)
@@ -410,18 +409,7 @@ namespace Horde.Build.Configuration
 			{
 				// Update the config object
 				GlobalConfig globalConfig = CreateGlobalConfig(data);
-				_stateTask = Task.FromResult(new ConfigState(hash, globalConfig));
-
-				// Notify any watchers of the new config object
-				for (; ; )
-				{
-					ChangeToken? token = _currentChangeToken;
-					if (Interlocked.CompareExchange(ref _currentChangeToken, null, token) == token)
-					{
-						token?.TriggerChange();
-						break;
-					}
-				}
+				Set(hash, globalConfig);
 			}
 		}
 
@@ -431,18 +419,21 @@ namespace Horde.Build.Configuration
 			ConfigSnapshot snapshot = DeserializeSnapshot(data);
 
 			// Build the new config and store the project and stream configs inside it
-			GlobalConfig globalConfig = JsonSerializer.Deserialize<GlobalConfig>(snapshot.Global, _jsonOptions)!;
-			foreach (ProjectConfigRef projectConfigRef in globalConfig.Projects)
-			{
-				ProjectConfig projectConfig = JsonSerializer.Deserialize<ProjectConfig>(snapshot.Projects[projectConfigRef.Id], _jsonOptions)!;
-				projectConfigRef.Target = projectConfig;
+			GlobalConfig globalConfig = JsonSerializer.Deserialize<GlobalConfig>(snapshot.Data, _jsonOptions)!;
+			globalConfig.Revision = IoHash.Compute(data.Span).ToString();
 
-				foreach (StreamConfigRef streamConfigRef in projectConfig.Streams)
+			// Compute hashes for all the stream objects
+			foreach (ProjectConfig projectConfig in globalConfig.Projects)
+			{
+				foreach (StreamConfig streamConfig in projectConfig.Streams)
 				{
-					StreamConfig streamConfig = JsonSerializer.Deserialize<StreamConfig>(snapshot.Streams[streamConfigRef.Id], _jsonOptions)!;
-					streamConfigRef.Target = streamConfig;
+					byte[] streamData = JsonSerializer.SerializeToUtf8Bytes(streamConfig, _jsonOptions);
+					streamConfig.Revision = IoHash.Compute(streamData).ToString();
 				}
 			}
+
+			// Run the postload callbacks on all the config objects
+			globalConfig.PostLoad(_serverSettings);
 			return globalConfig;
 		}
 	}
diff --git a/Engine/Source/Programs/Horde/Horde.Build/Devices/DeviceService.cs b/Engine/Source/Programs/Horde/Horde.Build/Devices/DeviceService.cs
--- a/Engine/Source/Programs/Horde/Horde.Build/Devices/DeviceService.cs
+++ b/Engine/Source/Programs/Horde/Horde.Build/Devices/DeviceService.cs
@@ -89,18 +89,16 @@ namespace Horde.Build.Devices
 	/// </summary>
 	public sealed class DeviceService : IHostedService, IDisposable
 	{
-		readonly GlobalsService _globalsService;
-		readonly AclService _aclService;
 		readonly INotificationService _notificationService;
 		readonly JobService _jobService;
-		readonly StreamService _streamService;
-		readonly ProjectService _projectService;
-		IUserCollection UserCollection { get; set; }
+		readonly IStreamCollection _streamCollection;
+		readonly IUserCollection _userCollection;
 		readonly ILogger<DeviceService> _logger;
 		readonly IDeviceCollection _devices;
 		readonly ITicker _ticker;
 		readonly ITicker _telemetryTicker;
 		readonly IOptionsMonitor<ServerSettings> _settings;
+		readonly IOptionsMonitor<GlobalConfig> _globalConfig;
 
 		/// <summary>
 		/// Platform map V1 singleton
@@ -117,22 +115,19 @@ namespace Horde.Build.Devices
 		/// <summary>
 		/// Device service constructor
 		/// </summary>
-		public DeviceService(GlobalsService globalsService, IDeviceCollection devices, ISingletonDocument<DevicePlatformMapV1> platformMapSingleton, IUserCollection userCollection, JobService jobService, ProjectService projectService, StreamService streamService, AclService aclService, IOptionsMonitor<ServerSettings> settings, INotificationService notificationService, IClock clock, ILogger<DeviceService> logger)
+		public DeviceService(IDeviceCollection devices, ISingletonDocument<DevicePlatformMapV1> platformMapSingleton, IUserCollection userCollection, JobService jobService, IStreamCollection streamCollection, IOptionsMonitor<ServerSettings> settings, IOptionsMonitor<GlobalConfig> globalConfig, INotificationService notificationService, IClock clock, ILogger<DeviceService> logger)
 		{
-			UserCollection = userCollection;
-			_globalsService = globalsService;
+			_userCollection = userCollection;
 			_devices = devices;
 			_jobService = jobService;
-			_projectService = projectService;
-			_streamService = streamService;
-			_aclService = aclService;
+			_streamCollection = streamCollection;
 			_notificationService = notificationService;
 			_ticker = clock.AddSharedTicker<DeviceService>(TimeSpan.FromMinutes(1.0), TickAsync, logger);
 			_telemetryTicker = clock.AddSharedTicker("DeviceService.Telemetry", TimeSpan.FromMinutes(10.0), TickTelemetryAsync, logger);
 			_logger = logger;
 			_settings = settings;
+			_globalConfig = globalConfig;
 			_platformMapSingleton = platformMapSingleton;
-
 		}
 
 		/// <inheritdoc/>
@@ -163,16 +158,15 @@ namespace Horde.Build.Devices
 		{
 			try
 			{
-				IGlobals globals = await _globalsService.GetAsync();
-
-				if (globals.Config.Devices != null)
+				GlobalConfig globalConfig = _globalConfig.CurrentValue;
+				if (globalConfig.Devices != null)
 				{
 					await _platformMapSingleton.UpdateAsync(platformMap => {
 
 						platformMap.PlatformMap.Clear();
 						platformMap.PerfSpecHighMap.Clear();
 
-						foreach (DevicePlatformConfig platform in globals.Config.Devices.Platforms)
+						foreach (DevicePlatformConfig platform in globalConfig.Devices.Platforms)
 						{
 							DevicePlatformId id = new DevicePlatformId(platform.Id);
 							foreach (string name in platform.Names)
@@ -211,6 +205,8 @@ namespace Horde.Build.Devices
 			{
 				using IScope scope = GlobalTracer.Instance.BuildSpan("DeviceService.TickAsync").StartActive();
 
+				GlobalConfig globalConfig = _globalConfig.CurrentValue;
+
 				if (runUpgrade)
 				{
 					try
@@ -245,7 +241,7 @@ namespace Horde.Build.Devices
 					{
 						foreach ((UserId, IDevice) expiredDevice in expireNotifications)
 						{
-							await NotifyDeviceServiceAsync($"Device {expiredDevice.Item2.PlatformId.ToString().ToUpperInvariant()} / {expiredDevice.Item2.Name} checkout will expire in 24 hours.  Please visit https://horde.devtools.epicgames.com/devices to renew the checkout if needed.", null, null, null, expiredDevice.Item1);
+							await NotifyDeviceServiceAsync(globalConfig, $"Device {expiredDevice.Item2.PlatformId.ToString().ToUpperInvariant()} / {expiredDevice.Item2.Name} checkout will expire in 24 hours.  Please visit https://horde.devtools.epicgames.com/devices to renew the checkout if needed.", null, null, null, expiredDevice.Item1);
 						}
 					}
 
@@ -255,7 +251,7 @@ namespace Horde.Build.Devices
 					{
 						foreach ((UserId, IDevice) expiredDevice in expireCheckouts)
 						{
-							await NotifyDeviceServiceAsync($"Device {expiredDevice.Item2.PlatformId.ToString().ToUpperInvariant()} / {expiredDevice.Item2.Name} checkout has expired.  The device has been returned to the shared pool and should no longer be accessed.  Please visit https://horde.devtools.epicgames.com/devices to checkout devices as needed.", null, null, null, expiredDevice.Item1);
+							await NotifyDeviceServiceAsync(globalConfig, $"Device {expiredDevice.Item2.PlatformId.ToString().ToUpperInvariant()} / {expiredDevice.Item2.Name} checkout has expired.  The device has been returned to the shared pool and should no longer be accessed.  Please visit https://horde.devtools.epicgames.com/devices to checkout devices as needed.", null, null, null, expiredDevice.Item1);
 						}
 					}
 				}
@@ -497,12 +493,7 @@ namespace Horde.Build.Devices
 		/// <summary>
 		/// 
 		/// </summary>
-		/// <param name="message"></param>
-		/// <param name="deviceId"></param>
-		/// <param name="jobId"></param>
-		/// <param name="stepId"></param>
-		/// <param name="userId"></param>
-		public async Task NotifyDeviceServiceAsync(string message, DeviceId? deviceId = null, string? jobId = null, string? stepId = null, UserId? userId = null)
+		public async Task NotifyDeviceServiceAsync(GlobalConfig globalConfig, string message, DeviceId? deviceId = null, string? jobId = null, string? stepId = null, UserId? userId = null)
 		{
 			try 
 			{
@@ -511,12 +502,12 @@ namespace Horde.Build.Devices
 				IJob? job = null;
 				IJobStep? step = null;
 				INode? node = null;
-				IStream? stream = null;
+				StreamConfig? streamConfig = null;
 				IUser? user = null;
 
 				if (userId.HasValue)
 				{
-					user = await UserCollection.GetUserAsync(userId.Value);
+					user = await _userCollection.GetUserAsync(userId.Value);
 					if (user == null)
 					{
 						_logger.LogError("Unable to send device notification, can't find User {UserId}", userId.Value);
@@ -536,7 +527,7 @@ namespace Horde.Build.Devices
 
 					if (job != null)
 					{
-						stream = await _streamService.GetStreamAsync(job.StreamId);
+						globalConfig.TryGetStream(job.StreamId, out streamConfig);
 
 						if (stepId != null)
 						{
@@ -554,7 +545,7 @@ namespace Horde.Build.Devices
 					}
 				}
 
-				_notificationService.NotifyDeviceService(message, device, pool, stream, job, step, node, user);
+				_notificationService.NotifyDeviceService(message, device, pool, streamConfig, job, step, node, user);
 
 			}
 			catch (Exception ex)
@@ -568,10 +559,11 @@ namespace Horde.Build.Devices
 		/// </summary>
 		/// <param name="action"></param>
 		/// <param name="user"></param>
+		/// <param name="globalConfig"></param>
 		/// <returns></returns>
-		public async Task<bool> AuthorizeAsync(AclAction action, ClaimsPrincipal user)
+		public static bool Authorize(AclAction action, ClaimsPrincipal user, GlobalConfig globalConfig)
 		{
-			// Tihs is deprecated and device auth should be going through GetUserPoolAuthorizationsAsync
+			// This is deprecated and device auth should be going through GetUserPoolAuthorizationsAsync
 
 			if (action == AclAction.DeviceRead)
 			{
@@ -583,46 +575,46 @@ namespace Horde.Build.Devices
 				return true;
 			}
 
-			if (await _aclService.AuthorizeAsync(AclAction.AdminWrite, user))
+			if (globalConfig.Authorize(AclAction.AdminWrite, user))
 			{
 				return true;
 			}
 
 			return false;
-
 		}
 
 		/// <summary>
 		/// Get list of pool authorizations for a user
 		/// </summary>
 		/// <param name="user"></param>
+		/// <param name="globalConfig"></param>
 		/// <returns></returns>
-		public async Task<List<DevicePoolAuthorization>> GetUserPoolAuthorizationsAsync(ClaimsPrincipal user)
+		public async Task<List<DevicePoolAuthorization>> GetUserPoolAuthorizationsAsync(ClaimsPrincipal user, GlobalConfig globalConfig)
 		{
-			
 			List<DevicePoolAuthorization> authPools = new List<DevicePoolAuthorization>();
 
 			List<IDevicePool> allPools = await GetPoolsAsync();
-			List<IProject> projects = await _projectService.GetProjectsAsync();
+			IReadOnlyList<ProjectConfig> projects = globalConfig.Projects;
 
 			// Set of projects associated with device pools
 			HashSet<ProjectId> projectIds = new HashSet<ProjectId>(allPools.Where(x => x.ProjectIds != null).SelectMany(x => x.ProjectIds!));
 
 			Dictionary<ProjectId, bool> deviceRead = new Dictionary<ProjectId, bool>();
 			Dictionary<ProjectId, bool> deviceWrite = new Dictionary<ProjectId, bool>();
-			ProjectPermissionsCache permissionsCache = new ProjectPermissionsCache();
 
 			foreach (ProjectId projectId in projectIds)
 			{
-
 				if (projects.Where(x => x.Id == projectId).FirstOrDefault() == null)
 				{
 					_logger.LogWarning("Device pool authorization references missing project id {ProjectId}", projectId);
 					continue;
 				}
 
-				deviceRead.Add(projectId,  await _projectService.AuthorizeAsync(projectId, AclAction.DeviceRead, user, permissionsCache));
-				deviceWrite.Add(projectId, await _projectService.AuthorizeAsync(projectId, AclAction.DeviceWrite, user, permissionsCache));
+				if (globalConfig.TryGetProject(projectId, out ProjectConfig? projectConfig))
+				{
+					deviceRead.Add(projectId, projectConfig.Authorize(AclAction.DeviceRead, user));
+					deviceWrite.Add(projectId, projectConfig.Authorize(AclAction.DeviceWrite, user));
+				}
 			}
 
 			// for global pools which aren't associated with a project
@@ -630,7 +622,7 @@ namespace Horde.Build.Devices
 
 			if (!globalPoolAccess)
 			{
-				if (await _aclService.AuthorizeAsync(AclAction.AdminWrite, user))
+				if (globalConfig.Authorize(AclAction.AdminWrite, user))
 				{
 					globalPoolAccess = true;
 				}
@@ -680,10 +672,11 @@ namespace Horde.Build.Devices
 		/// </summary>
 		/// <param name="id"></param>
 		/// <param name="user"></param>
+		/// <param name="globalConfig"></param>
 		/// <returns></returns>
-		public async Task<DevicePoolAuthorization?> GetUserPoolAuthorizationAsync(DevicePoolId id, ClaimsPrincipal user)
+		public async Task<DevicePoolAuthorization?> GetUserPoolAuthorizationAsync(DevicePoolId id, ClaimsPrincipal user, GlobalConfig globalConfig)
 		{
-			List<DevicePoolAuthorization> auth = await GetUserPoolAuthorizationsAsync(user);
+			List<DevicePoolAuthorization> auth = await GetUserPoolAuthorizationsAsync(user, globalConfig);
 			return auth.Where(x => x.Pool.Id == id).FirstOrDefault();
 		}
 
diff --git a/Engine/Source/Programs/Horde/Horde.Build/Devices/DevicesController.cs b/Engine/Source/Programs/Horde/Horde.Build/Devices/DevicesController.cs
--- a/Engine/Source/Programs/Horde/Horde.Build/Devices/DevicesController.cs
+++ b/Engine/Source/Programs/Horde/Horde.Build/Devices/DevicesController.cs
@@ -7,6 +7,7 @@ using System.Text.Json;
 using System.Threading.Tasks;
 using Horde.Build.Acls;
 using Horde.Build.Projects;
+using Horde.Build.Server;
 using Horde.Build.Users;
 using Horde.Build.Utilities;
 using Microsoft.AspNetCore.Authorization;
@@ -28,21 +29,9 @@ namespace Horde.Build.Devices
 	/// </summary>
 	public class DevicesController : ControllerBase
 	{
-
-		/// <summary>
-		/// The acl service singleton
-		/// </summary>
-		readonly AclService _aclService;
-
-		/// <summary>
-		/// The user collection instance
-		/// </summary>
-		IUserCollection UserCollection { get; set; }
-
-		/// <summary>
-		/// Singleton instance of the device service
-		/// </summary>
+		readonly IUserCollection _userCollection;
 		readonly DeviceService _deviceService;
+		readonly IOptionsSnapshot<GlobalConfig> _globalConfig;
 
 		/// <summary>
 		///  Logger for controller
@@ -52,12 +41,12 @@ namespace Horde.Build.Devices
 		/// <summary>
 		/// Constructor
 		/// </summary>
-		public DevicesController(DeviceService deviceService, AclService aclService, IUserCollection userCollection, ILogger<DevicesController> logger)
+		public DevicesController(IUserCollection userCollection, DeviceService deviceService, IOptionsSnapshot<GlobalConfig> globalConfig, ILogger<DevicesController> logger)
 		{
-			UserCollection = userCollection;
+			_userCollection = userCollection;
 			_deviceService = deviceService;
+			_globalConfig = globalConfig;
 			_logger = logger;
-			_aclService = aclService;
 		}
 
 		// DEVICES
@@ -70,15 +59,14 @@ namespace Horde.Build.Devices
 		[Route("/api/v2/devices")]
 		public async Task<ActionResult<CreateDeviceResponse>> CreateDeviceAsync([FromBody] CreateDeviceRequest deviceRequest)
 		{
-
-			DevicePoolAuthorization? poolAuth = await _deviceService.GetUserPoolAuthorizationAsync(new DevicePoolId(deviceRequest.PoolId!), User);
+			DevicePoolAuthorization? poolAuth = await _deviceService.GetUserPoolAuthorizationAsync(new DevicePoolId(deviceRequest.PoolId!), User, _globalConfig.Value);
 
 			if (poolAuth == null || !poolAuth.Write)
 			{
 				return Forbid();
 			}
 
-			IUser? internalUser = await UserCollection.GetUserAsync(User);
+			IUser? internalUser = await _userCollection.GetUserAsync(User);
 			if (internalUser == null)
 			{
 				return NotFound();
@@ -132,8 +120,7 @@ namespace Horde.Build.Devices
 		[ProducesResponseType(typeof(List<GetDeviceResponse>), 200)]
 		public async Task<ActionResult<List<object>>> GetDevicesAsync()
 		{
-
-			List<DevicePoolAuthorization> poolAuth = await _deviceService.GetUserPoolAuthorizationsAsync(User);
+			List<DevicePoolAuthorization> poolAuth = await _deviceService.GetUserPoolAuthorizationsAsync(User, _globalConfig.Value);
 
 			List<IDevice> devices = await _deviceService.GetDevicesAsync();
 
@@ -179,7 +166,7 @@ namespace Horde.Build.Devices
 				return BadRequest($"Unable to find device with id {deviceId}");
 			}
 
-			DevicePoolAuthorization? poolAuth = await _deviceService.GetUserPoolAuthorizationAsync(device.PoolId, User);
+			DevicePoolAuthorization? poolAuth = await _deviceService.GetUserPoolAuthorizationAsync(device.PoolId, User, _globalConfig.Value);
 
 			if (poolAuth == null || !poolAuth.Write)
 			{
@@ -204,7 +191,7 @@ namespace Horde.Build.Devices
 		[ProducesResponseType(typeof(List<GetDeviceResponse>), 200)]
 		public async Task<ActionResult> UpdateDeviceAsync(string deviceId, [FromBody] UpdateDeviceRequest update)
 		{
-			IUser? internalUser = await UserCollection.GetUserAsync(User);
+			IUser? internalUser = await _userCollection.GetUserAsync(User);
 			if (internalUser == null)
 			{
 				return NotFound();
@@ -219,7 +206,7 @@ namespace Horde.Build.Devices
 				return BadRequest($"Device with id ${deviceId} does not exist");
 			}
 
-			DevicePoolAuthorization? poolAuth = await _deviceService.GetUserPoolAuthorizationAsync(device.PoolId, User);
+			DevicePoolAuthorization? poolAuth = await _deviceService.GetUserPoolAuthorizationAsync(device.PoolId, User, _globalConfig.Value);
 
 			if (poolAuth == null || !poolAuth.Write)
 			{
@@ -270,7 +257,7 @@ namespace Horde.Build.Devices
 		public async Task<ActionResult> CheckoutDeviceAsync(string deviceId, [FromBody] CheckoutDeviceRequest request)
 		{
 
-			IUser? internalUser = await UserCollection.GetUserAsync(User);
+			IUser? internalUser = await _userCollection.GetUserAsync(User);
 			if (internalUser == null)
 			{
 				return NotFound();
@@ -285,7 +272,7 @@ namespace Horde.Build.Devices
 				return BadRequest($"Device with id ${deviceId} does not exist");
 			}
 
-			DevicePoolAuthorization? poolAuth = await _deviceService.GetUserPoolAuthorizationAsync(device.PoolId, User);
+			DevicePoolAuthorization? poolAuth = await _deviceService.GetUserPoolAuthorizationAsync(device.PoolId, User, _globalConfig.Value);
 
 			if (poolAuth == null || !poolAuth.Write)
 			{
@@ -318,7 +305,7 @@ namespace Horde.Build.Devices
 		[Route("/api/v2/devices/{deviceId}")]
 		public async Task<ActionResult> DeleteDeviceAsync(string deviceId)
 		{
-			if (!await _deviceService.AuthorizeAsync(AclAction.DeviceWrite, User))
+			if (!DeviceService.Authorize(AclAction.DeviceWrite, User, _globalConfig.Value))
 			{
 				return Forbid();
 			}
@@ -331,7 +318,7 @@ namespace Horde.Build.Devices
 				return NotFound();
 			}
 
-			DevicePoolAuthorization? poolAuth = await _deviceService.GetUserPoolAuthorizationAsync(device.PoolId, User);
+			DevicePoolAuthorization? poolAuth = await _deviceService.GetUserPoolAuthorizationAsync(device.PoolId, User, _globalConfig.Value);
 
 			if (poolAuth == null || !poolAuth.Write)
 			{
@@ -352,7 +339,7 @@ namespace Horde.Build.Devices
 		[Route("/api/v2/devices/platforms")]
 		public async Task<ActionResult<CreateDevicePlatformResponse>> CreatePlatformAsync([FromBody] CreateDevicePlatformRequest request)
 		{
-			if (!await _deviceService.AuthorizeAsync(AclAction.DeviceWrite, User))
+			if (!DeviceService.Authorize(AclAction.DeviceWrite, User, _globalConfig.Value))
 			{
 				return Forbid();
 			}
@@ -377,7 +364,7 @@ namespace Horde.Build.Devices
 		[Route("/api/v2/devices/platforms/{platformId}")]
 		public async Task<ActionResult<CreateDevicePlatformResponse>> UpdatePlatformAsync(string platformId, [FromBody] UpdateDevicePlatformRequest request)
 		{
-			if (!await _deviceService.AuthorizeAsync(AclAction.DeviceWrite, User))
+			if (!DeviceService.Authorize(AclAction.DeviceWrite, User, _globalConfig.Value))
 			{
 				return Forbid();
 			}
@@ -396,8 +383,7 @@ namespace Horde.Build.Devices
 		[ProducesResponseType(typeof(List<GetDevicePlatformResponse>), 200)]
 		public async Task<ActionResult<List<object>>> GetDevicePlatformsAsync()
 		{
-
-			if (!await _deviceService.AuthorizeAsync(AclAction.DeviceRead, User))
+			if (!DeviceService.Authorize(AclAction.DeviceRead, User, _globalConfig.Value))
 			{
 				return Forbid();
 			}
@@ -425,7 +411,7 @@ namespace Horde.Build.Devices
 		[Route("/api/v2/devices/pools")]
 		public async Task<ActionResult<CreateDevicePoolResponse>> CreatePoolAsync([FromBody] CreateDevicePoolRequest request)
 		{
-			if (!await _deviceService.AuthorizeAsync(AclAction.DeviceWrite, User))
+			if (!DeviceService.Authorize(AclAction.DeviceWrite, User, _globalConfig.Value))
 			{
 				return Forbid();
 			}
@@ -451,8 +437,7 @@ namespace Horde.Build.Devices
 		[Route("/api/v2/devices/pools")]
 		public async Task<ActionResult> UpdatePoolAsync([FromBody] UpdateDevicePoolRequest request)
 		{
-
-			DevicePoolAuthorization? poolAuth = await _deviceService.GetUserPoolAuthorizationAsync(new DevicePoolId(request.Id), User);
+			DevicePoolAuthorization? poolAuth = await _deviceService.GetUserPoolAuthorizationAsync(new DevicePoolId(request.Id), User, _globalConfig.Value);
 
 			if (poolAuth == null || !poolAuth.Write)
 			{
@@ -480,8 +465,7 @@ namespace Horde.Build.Devices
 		[ProducesResponseType(typeof(List<GetDevicePoolResponse>), 200)]
 		public async Task<ActionResult<List<object>>> GetDevicePoolsAsync()
 		{
-
-			List<DevicePoolAuthorization> poolAuth = await _deviceService.GetUserPoolAuthorizationsAsync(User);
+			List<DevicePoolAuthorization> poolAuth = await _deviceService.GetUserPoolAuthorizationsAsync(User, _globalConfig.Value);
 
 			List<IDevicePool> pools = await _deviceService.GetPoolsAsync();
 
@@ -515,7 +499,7 @@ namespace Horde.Build.Devices
 			[FromQuery] int index = 0,
 			[FromQuery] int count = 1024)
 		{
-			if (!await _deviceService.AuthorizeAsync(AclAction.DeviceRead, User))
+			if (!DeviceService.Authorize(AclAction.DeviceRead, User, _globalConfig.Value))
 			{
 				return Forbid();
 			}
@@ -563,8 +547,7 @@ namespace Horde.Build.Devices
 		[ProducesResponseType(typeof(CreateDeviceReservationResponse), 200)]
 		public async Task<ActionResult<CreateDeviceReservationResponse>> CreateDeviceReservation([FromBody] CreateDeviceReservationRequest request)
 		{
-
-			if (!await _deviceService.AuthorizeAsync(AclAction.DeviceWrite, User))
+			if (!DeviceService.Authorize(AclAction.DeviceWrite, User, _globalConfig.Value))
 			{
 				return Forbid();
 			}
@@ -646,7 +629,7 @@ namespace Horde.Build.Devices
 		[ProducesResponseType(typeof(List<GetDeviceReservationResponse>), 200)]
 		public async Task<ActionResult<List<GetDeviceReservationResponse>>> GetDeviceReservations()
 		{
-			if (!await _deviceService.AuthorizeAsync(AclAction.DeviceRead, User))
+			if (!DeviceService.Authorize(AclAction.DeviceRead, User, _globalConfig.Value))
 			{
 				return Forbid();
 			}
@@ -688,7 +671,7 @@ namespace Horde.Build.Devices
 		[Route("/api/v2/devices/reservations/{reservationId}")]
 		public async Task<ActionResult> UpdateReservationAsync(string reservationId)
 		{
-			if (!await _deviceService.AuthorizeAsync(AclAction.DeviceWrite, User))
+			if (!DeviceService.Authorize(AclAction.DeviceWrite, User, _globalConfig.Value))
 			{
 				return Forbid();
 			}
@@ -713,7 +696,7 @@ namespace Horde.Build.Devices
 		[Route("/api/v2/devices/reservations/{reservationId}")]
 		public async Task<ActionResult> DeleteReservationAsync(string reservationId)
 		{
-			if (!await _deviceService.AuthorizeAsync(AclAction.DeviceWrite, User))
+			if (!DeviceService.Authorize(AclAction.DeviceWrite, User, _globalConfig.Value))
 			{
 				return Forbid();
 			}
@@ -748,7 +731,7 @@ namespace Horde.Build.Devices
 			[FromQuery] int index = 0,
 			[FromQuery] int count = 1024)
 		{
-			if (!await _deviceService.AuthorizeAsync(AclAction.DeviceRead, User))
+			if (!DeviceService.Authorize(AclAction.DeviceRead, User, _globalConfig.Value))
 			{
 				return Forbid();
 			}
diff --git a/Engine/Source/Programs/Horde/Horde.Build/Issues/External/IExternalIssueService.cs b/Engine/Source/Programs/Horde/Horde.Build/Issues/External/IExternalIssueService.cs
--- a/Engine/Source/Programs/Horde/Horde.Build/Issues/External/IExternalIssueService.cs
+++ b/Engine/Source/Programs/Horde/Horde.Build/Issues/External/IExternalIssueService.cs
@@ -77,9 +77,9 @@ namespace Horde.Build.Issues.External
 		/// <summary>
 		///  Get projects for provided keys
 		/// </summary>
-		/// <param name="stream"></param>
+		/// <param name="streamConfig"></param>
 		/// <returns></returns>;
-		Task<List<IExternalIssueProject>> GetProjects(IStream stream);
+		Task<List<IExternalIssueProject>> GetProjects(StreamConfig streamConfig);
 	}
 
 	/// <summary>
@@ -106,7 +106,7 @@ namespace Horde.Build.Issues.External
 		}
 
 		/// <inheritdoc/>
-		public Task<List<IExternalIssueProject>> GetProjects(IStream stream)
+		public Task<List<IExternalIssueProject>> GetProjects(StreamConfig streamConfig)
 		{
 			return Task.FromResult(new List<IExternalIssueProject>());
 		}
diff --git a/Engine/Source/Programs/Horde/Horde.Build/Issues/External/JiraService.cs b/Engine/Source/Programs/Horde/Horde.Build/Issues/External/JiraService.cs
--- a/Engine/Source/Programs/Horde/Horde.Build/Issues/External/JiraService.cs
+++ b/Engine/Source/Programs/Horde/Horde.Build/Issues/External/JiraService.cs
@@ -20,6 +20,7 @@ using System.Threading;
 using System.Linq;
 using Horde.Build.Streams;
 using Horde.Build.Users;
+using Horde.Build.Server;
 
 namespace Horde.Build.Issues.External
 {
@@ -112,42 +113,39 @@ namespace Horde.Build.Issues.External
 		/// <summary>
 		/// Singleton instance of the stream service
 		/// </summary>
-		readonly StreamService _streamService;
+		readonly IStreamCollection _streamCollection;
 
 		/// <summary>
 		/// Singleton instance of the project service
 		/// </summary>
 		readonly IssueService _issueService;
 
-		/// <summary>
-		/// Config collection
-		/// </summary>
-		readonly ConfigCollection _configCollection;
-
 		readonly ITicker _ticker;
 
 		readonly Uri _jiraUrl;
 
 		readonly Dictionary<string, JiraProject> _cachedJiraProjects = new Dictionary<string, JiraProject>();
 
+		readonly IOptionsMonitor<GlobalConfig> _globalConfig;
+
 		/// <summary>
 		/// Jira service constructor
 		/// </summary>
 		/// <param name="settings"></param>
-		/// <param name="streamService"></param>
+		/// <param name="streamCollection"></param>
 		/// <param name="issueService"></param>
-		/// <param name="configCollection"></param>
 		/// <param name="clock"></param>
+		/// <param name="globalConfig"></param>
 		/// <param name="logger"></param>
-		public JiraService(IOptions<ServerSettings> settings, StreamService streamService, IssueService issueService, ConfigCollection configCollection, IClock clock, ILogger<JiraService> logger)
+		public JiraService(IOptions<ServerSettings> settings, IStreamCollection streamCollection, IssueService issueService, IClock clock, IOptionsMonitor<GlobalConfig> globalConfig, ILogger<JiraService> logger)
 		{
 			_settings = settings.Value;
 			_logger = logger;
-			_streamService = streamService;
+			_streamCollection = streamCollection;
 			_issueService = issueService;
-			_configCollection = configCollection;
 			_ticker = clock.AddTicker<JiraService>(TimeSpan.FromMinutes(2.0), TickAsync, logger);
 			_jiraUrl = _settings.JiraUrl!;
+			_globalConfig = globalConfig;
 
 			// setup http client for Jira rest api queries
 			_client = new HttpClient();
@@ -173,29 +171,17 @@ namespace Horde.Build.Issues.External
 		/// <inheritdoc/>
 		public string? GetIssueUrl(string key) => $"{_jiraUrl}browse/{key}";
 
-		public async Task<List<IExternalIssueProject>> GetProjects(IStream stream)
+		public Task<List<IExternalIssueProject>> GetProjects(StreamConfig streamConfig)
 		{
 			HashSet<string> projectKeys = new HashSet<string>();
 			List<IExternalIssueProject> result = new List<IExternalIssueProject>();
 
-			StreamConfig? config = null;
-
-			try
-			{
-				config = await _configCollection.GetConfigAsync<StreamConfig>(stream.ConfigRevision);
-			}
-			catch (Exception ex)
-			{
-				_logger.LogError(ex, "Unable to get stream config for {StreamId}", stream.Id);
-
-			}
-
-			if (config == null || config.Workflows == null)
+			if (streamConfig == null || streamConfig.Workflows == null)
 			{
-				return result;
+				return Task.FromResult(result);
 			}
 
-			config.Workflows.ForEach(workflow =>
+			streamConfig.Workflows.ForEach(workflow =>
 			{
 				if (workflow.ExternalIssues != null)
 				{
@@ -211,7 +197,7 @@ namespace Horde.Build.Issues.External
 				}
 			}
 
-			return result;
+			return Task.FromResult(result);
 		}
 
 		async Task UpdateJiraProjects(string[] jiraProjectKeys)
@@ -319,21 +305,11 @@ namespace Horde.Build.Issues.External
 
 			HashSet<string> jiraProjectKeys = new HashSet<string>();
 
-			List<IStream> cachedStreams = await _streamService.StreamCollection.FindAllAsync();
-			foreach (IStream cachedStream in cachedStreams)
-			{
-				StreamConfig? config = null;
-				try
-				{
-					config = await _configCollection.GetConfigAsync<StreamConfig>(cachedStream.ConfigRevision);
-				}
-				catch (Exception ex)
-				{
-					_logger.LogError(ex, "Unable to get stream config for {StreamId}", cachedStream.Id);
-					continue;
-				}
+			GlobalConfig globalConfig = _globalConfig.CurrentValue;
 
-				foreach (WorkflowConfig workflow in config.Workflows)
+			foreach(StreamConfig streamConfig in globalConfig.Streams)
+			{
+				foreach (WorkflowConfig workflow in streamConfig.Workflows)
 				{
 					if (workflow.ExternalIssues == null || String.IsNullOrEmpty(workflow.ExternalIssues.ProjectKey))
 					{
diff --git a/Engine/Source/Programs/Horde/Horde.Build/Issues/IssueMessages.cs b/Engine/Source/Programs/Horde/Horde.Build/Issues/IssueMessages.cs
--- a/Engine/Source/Programs/Horde/Horde.Build/Issues/IssueMessages.cs
+++ b/Engine/Source/Programs/Horde/Horde.Build/Issues/IssueMessages.cs
@@ -163,12 +163,12 @@ namespace Horde.Build.Issues
 		/// <summary>
 		/// Constructor
 		/// </summary>
-		/// <param name="stream">The stream to construct from</param>
+		/// <param name="streamId">The stream to construct from</param>
 		/// <param name="spans">List of spans for the given stream</param>
 		/// <param name="steps">List of steps for the given stream</param>
-		public GetIssueStreamResponse(IStream stream, List<IIssueSpan> spans, List<IIssueStep> steps)
+		public GetIssueStreamResponse(StreamId streamId, List<IIssueSpan> spans, List<IIssueStep> steps)
 		{
-			StreamId = stream.Id.ToString();
+			StreamId = streamId.ToString();
 
 			foreach (IIssueSpan span in spans)
 			{
@@ -330,9 +330,9 @@ namespace Horde.Build.Issues
 		/// Constructor
 		/// </summary>
 		/// <param name="details">Issue to construct from</param>
-		/// <param name="stream"></param>
+		/// <param name="streamConfig"></param>
 		/// <param name="spans">The spans to construct from</param>
-		public GetIssueAffectedStreamResponse(IIssueDetails details, IStream? stream, IEnumerable<IIssueSpan> spans)
+		public GetIssueAffectedStreamResponse(IIssueDetails details, StreamConfig? streamConfig, IEnumerable<IIssueSpan> spans)
 		{
 			IIssueSpan firstSpan = spans.First();
 			StreamId = firstSpan.StreamId.ToString();
@@ -343,9 +343,9 @@ namespace Horde.Build.Issues
 			foreach (IGrouping<TemplateId, IIssueSpan> template in spans.GroupBy(x => x.TemplateRefId))
 			{
 				string templateName = template.Key.ToString();
-				if (stream != null && stream.Templates.TryGetValue(template.Key, out ITemplateRef? templateRef))
+				if (streamConfig != null && streamConfig.TryGetTemplate(template.Key, out TemplateRefConfig? templateRefConfig))
 				{
-					templateName = templateRef.Config.Name;
+					templateName = templateRefConfig.Name;
 				}
 
 				HashSet<ObjectId> unresolvedTemplateSpans = new HashSet<ObjectId>(template.Where(x => x.NextSuccess == null).Select(x => x.Id));
