Horde: Refactor codebase to use new config system.

diff --git a/Engine/Source/Programs/Horde/Horde.Build/Issues/IssueReportService.cs b/Engine/Source/Programs/Horde/Horde.Build/Issues/IssueReportService.cs
--- a/Engine/Source/Programs/Horde/Horde.Build/Issues/IssueReportService.cs
+++ b/Engine/Source/Programs/Horde/Horde.Build/Issues/IssueReportService.cs
@@ -6,7 +6,6 @@ using System.Linq;
 using System.Threading;
 using System.Threading.Tasks;
 using EpicGames.Core;
-using Horde.Build.Configuration;
 using Horde.Build.Jobs;
 using Horde.Build.Jobs.Graphs;
 using Horde.Build.Notifications;
@@ -16,11 +15,13 @@ using Horde.Build.Utilities;
 using HordeCommon;
 using Microsoft.Extensions.Hosting;
 using Microsoft.Extensions.Logging;
+using Microsoft.Extensions.Options;
 using MongoDB.Bson.Serialization.Attributes;
 using MongoDB.Bson.Serialization.Options;
 
 namespace Horde.Build.Issues
 {
+	using StreamId = StringId<IStream>;
 	using TemplateId = StringId<ITemplateRef>;
 	using WorkflowId = StringId<WorkflowConfig>;
 
@@ -33,18 +34,18 @@ namespace Horde.Build.Issues
 
 	public class IssueReport
 	{
+		public StreamId StreamId { get; }
 		public WorkflowId WorkflowId { get; }
-		public IStream Stream { get; }
 		public WorkflowStats WorkflowStats { get; }
 		public string? TriageChannel { get; }
 		public List<IIssue> Issues { get; } = new List<IIssue>();
 		public List<IIssueSpan> IssueSpans { get; } = new List<IIssueSpan>();
 		public bool GroupByTemplate { get; }
 
-		public IssueReport(WorkflowId workflowId, IStream stream, WorkflowStats workflowStats, string? triageChannel, bool groupByTemplate)
+		public IssueReport(StreamId streamId, WorkflowId workflowId, WorkflowStats workflowStats, string? triageChannel, bool groupByTemplate)
 		{
+			StreamId = streamId;
 			WorkflowId = workflowId;
-			Stream = stream;
 			WorkflowStats = workflowStats;
 			TriageChannel = triageChannel;
 			GroupByTemplate = groupByTemplate;
@@ -83,26 +84,26 @@ namespace Horde.Build.Issues
 		readonly IIssueCollection _issueCollection;
 		readonly IGraphCollection _graphCollection;
 		readonly IJobCollection _jobCollection;
-		readonly ConfigCollection _configCollection;
 		readonly INotificationService _notificationService;
 		readonly IClock _clock;
 		readonly ITicker _ticker;
+		readonly IOptionsMonitor<GlobalConfig> _globalConfig;
 		readonly ILogger<IssueReportService> _logger;
 
 		/// <summary>
 		/// Constructor
 		/// </summary>
-		public IssueReportService(MongoService mongoService, IStreamCollection streamCollection, IIssueCollection issueCollection, IGraphCollection graphCollection, IJobCollection jobCollection, ConfigCollection configCollection, INotificationService notificationService, IClock clock, ILogger<IssueReportService> logger)
+		public IssueReportService(MongoService mongoService, IStreamCollection streamCollection, IIssueCollection issueCollection, IGraphCollection graphCollection, IJobCollection jobCollection, INotificationService notificationService, IClock clock, IOptionsMonitor<GlobalConfig> globalConfig, ILogger<IssueReportService> logger)
 		{
 			_state = new SingletonDocument<IssueReportState>(mongoService);
 			_streamCollection = streamCollection;
 			_issueCollection = issueCollection;
 			_graphCollection = graphCollection;
 			_jobCollection = jobCollection;
-			_configCollection = configCollection;
 			_notificationService = notificationService;
 			_clock = clock;
 			_ticker = clock.AddSharedTicker<IssueReportService>(TimeSpan.FromMinutes(5.0), TickAsync, logger);
+			_globalConfig = globalConfig;
 			_logger = logger;
 		}
 
@@ -128,23 +129,22 @@ namespace Horde.Build.Issues
 			List<string> updateKeys = new List<string>();
 			List<IssueReportGroup> groups = new List<IssueReportGroup>();
 
-			List<IStream> streams = await _streamCollection.FindAllAsync();
-			foreach (IStream stream in streams)
+			GlobalConfig globalConfig = _globalConfig.CurrentValue;
+			foreach (StreamConfig streamConfig in globalConfig.Streams)
 			{
-				StreamConfig config = await _configCollection.GetConfigAsync<StreamConfig>(stream.ConfigRevision);
-				if (config.Workflows.Count > 0)
+				if (streamConfig.Workflows.Count > 0)
 				{
 					List<IIssue>? issues = null;
 					List<IIssueSpan>? spans = null;
 
-					foreach (WorkflowConfig workflow in config.Workflows)
+					foreach (WorkflowConfig workflowConfig in streamConfig.Workflows)
 					{
-						if (workflow.ReportChannel == null)
+						if (workflowConfig.ReportChannel == null)
 						{
 							continue;
 						}
 
-						string key = $"{stream.Id}:{workflow.Id}";
+						string key = $"{streamConfig.Id}:{workflowConfig.Id}";
 						invalidKeys.Remove(key);
 						
 						DateTime lastReportTime;
@@ -154,29 +154,29 @@ namespace Horde.Build.Issues
 							continue;
 						}
 
-						DateTime lastScheduledReportTime = GetLastScheduledReportTime(workflow, currentTime, _clock.TimeZone);
+						DateTime lastScheduledReportTime = GetLastScheduledReportTime(workflowConfig, currentTime, _clock.TimeZone);
 						if (lastReportTime > lastScheduledReportTime)
 						{
 							continue;
 						}
 
-						DateTime prevScheduledReportTime = GetLastScheduledReportTime(workflow, lastScheduledReportTime - TimeSpan.FromMinutes(1.0), _clock.TimeZone);
+						DateTime prevScheduledReportTime = GetLastScheduledReportTime(workflowConfig, lastScheduledReportTime - TimeSpan.FromMinutes(1.0), _clock.TimeZone);
 
-						_logger.LogInformation("Creating report for {StreamId} workflow {WorkflowId}", stream.Id, workflow.Id);
+						_logger.LogInformation("Creating report for {StreamId} workflow {WorkflowId}", streamConfig.Id, workflowConfig.Id);
 
-						issues ??= await _issueCollection.FindIssuesAsync(streamId: stream.Id);
+						issues ??= await _issueCollection.FindIssuesAsync(streamId: streamConfig.Id);
 						spans ??= await _issueCollection.FindSpansAsync(issueIds: issues.Select(x => x.Id).ToArray());
 
-						Dictionary<WorkflowId, WorkflowStats> workflowIdToStats = await GetWorkflowStatsAsync(stream, prevScheduledReportTime);
-						if (!workflowIdToStats.TryGetValue(workflow.Id, out WorkflowStats? workflowStats))
+						Dictionary<WorkflowId, WorkflowStats> workflowIdToStats = await GetWorkflowStatsAsync(streamConfig, prevScheduledReportTime);
+						if (!workflowIdToStats.TryGetValue(workflowConfig.Id, out WorkflowStats? workflowStats))
 						{
 							workflowStats = new WorkflowStats();
 						}
 
-						IssueReport report = new IssueReport(workflow.Id, stream, workflowStats, workflow.TriageChannel, workflow.GroupIssuesByTemplate);
+						IssueReport report = new IssueReport(streamConfig.Id, workflowConfig.Id, workflowStats, workflowConfig.TriageChannel, workflowConfig.GroupIssuesByTemplate);
 						foreach (IIssueSpan span in spans)
 						{
-							if (span.LastSuccess != null && span.LastFailure.Annotations.WorkflowId == workflow.Id)
+							if (span.LastSuccess != null && span.LastFailure.Annotations.WorkflowId == workflowConfig.Id)
 							{
 								report.IssueSpans.Add(span);
 							}
@@ -187,10 +187,10 @@ namespace Horde.Build.Issues
 
 						DateTime reportTime = lastScheduledReportTime;
 
-						IssueReportGroup? group = groups.FirstOrDefault(x => x.Channel == workflow.ReportChannel && x.Time == reportTime);
+						IssueReportGroup? group = groups.FirstOrDefault(x => x.Channel == workflowConfig.ReportChannel && x.Time == reportTime);
 						if (group == null)
 						{
-							group = new IssueReportGroup(workflow.ReportChannel, reportTime);
+							group = new IssueReportGroup(workflowConfig.ReportChannel, reportTime);
 							groups.Add(group);
 						}
 						group.Reports.Add(report);
@@ -222,15 +222,15 @@ namespace Horde.Build.Issues
 			}
 		}
 
-		private async Task<Dictionary<WorkflowId, WorkflowStats>> GetWorkflowStatsAsync(IStream stream, DateTime minTime)
+		private async Task<Dictionary<WorkflowId, WorkflowStats>> GetWorkflowStatsAsync(StreamConfig streamConfig, DateTime minTime)
 		{
-			List<IJob> jobs = await _jobCollection.FindAsync(streamId: stream.Id, minCreateTime: minTime);
+			List<IJob> jobs = await _jobCollection.FindAsync(streamId: streamConfig.Id, minCreateTime: minTime);
 
 			Dictionary<WorkflowId, WorkflowStats> workflowIdToStats = new Dictionary<WorkflowId, WorkflowStats>();
 			foreach (IGrouping<TemplateId, IJob> templateGroup in jobs.GroupBy(x => x.TemplateId))
 			{
 				WorkflowId? templateWorkflowId = null;
-				if (stream.Config.TryGetTemplate(templateGroup.Key, out TemplateRefConfig? templateRefConfig))
+				if (streamConfig.TryGetTemplate(templateGroup.Key, out TemplateRefConfig? templateRefConfig))
 				{
 					templateWorkflowId = templateRefConfig.Annotations.WorkflowId;
 				}
diff --git a/Engine/Source/Programs/Horde/Horde.Build/Issues/IssueService.cs b/Engine/Source/Programs/Horde/Horde.Build/Issues/IssueService.cs
--- a/Engine/Source/Programs/Horde/Horde.Build/Issues/IssueService.cs
+++ b/Engine/Source/Programs/Horde/Horde.Build/Issues/IssueService.cs
@@ -15,12 +15,14 @@ using Horde.Build.Jobs;
 using Horde.Build.Jobs.Graphs;
 using Horde.Build.Logs;
 using Horde.Build.Perforce;
+using Horde.Build.Server;
 using Horde.Build.Streams;
 using Horde.Build.Users;
 using Horde.Build.Utilities;
 using HordeCommon;
 using Microsoft.Extensions.Hosting;
 using Microsoft.Extensions.Logging;
+using Microsoft.Extensions.Options;
 using MongoDB.Bson;
 using MongoDB.Bson.Serialization;
 using MongoDB.Driver;
@@ -155,7 +157,7 @@ namespace Horde.Build.Issues
 		readonly IJobStepRefCollection _jobStepRefs;
 		readonly IIssueCollection _issueCollection;
 		readonly ICommitService _commitService;
-		readonly StreamService _streams;
+		readonly IStreamCollection _streams;
 		readonly IUserCollection _userCollection;
 		readonly ILogFileService _logFileService;
 		readonly IClock _clock;
@@ -206,10 +208,12 @@ namespace Horde.Build.Issues
 		/// </summary>
 		public IReadOnlyList<IIssueDetails> CachedOpenIssues => _cachedIssues;
 
+		readonly IOptionsMonitor<GlobalConfig> _globalConfig;
+
 		/// <summary>
 		/// Constructor
 		/// </summary>
-		public IssueService(IIssueCollection issueCollection, ICommitService commitService, IJobStepRefCollection jobStepRefs, StreamService streams, IUserCollection userCollection, ILogFileService logFileService, IClock clock, ILogger<IssueService> logger)
+		public IssueService(IIssueCollection issueCollection, ICommitService commitService, IJobStepRefCollection jobStepRefs, IStreamCollection streams, IUserCollection userCollection, ILogFileService logFileService, IClock clock, IOptionsMonitor<GlobalConfig> globalConfig, ILogger<IssueService> logger)
 		{
 			Type[] issueTypes = Assembly.GetExecutingAssembly().GetTypes().Where(x => !x.IsAbstract && typeof(IIssue).IsAssignableFrom(x)).ToArray();
 			foreach (Type issueType in issueTypes)
@@ -226,6 +230,7 @@ namespace Horde.Build.Issues
 			_logFileService = logFileService;
 			_clock = clock;
 			_ticker = clock.AddTicker<IssueService>(TimeSpan.FromMinutes(1.0), TickAsync, logger);
+			_globalConfig = globalConfig;
 			_logger = logger;
 
 			// Create all the issue handlers
@@ -261,13 +266,14 @@ namespace Horde.Build.Issues
 
 			Dictionary<StreamId, HashSet<TemplateId>> newCachedDesktopAlerts = new Dictionary<StreamId, HashSet<TemplateId>>();
 
-			List<IStream> cachedStreams = await _streams.StreamCollection.FindAllAsync();
-			foreach(IStream cachedStream in cachedStreams)
+			GlobalConfig globalConfig = _globalConfig.CurrentValue;
+
+			foreach(StreamConfig streamConfig in globalConfig.Streams)
 			{
-				HashSet<TemplateId> templates = new HashSet<TemplateId>(cachedStream.Templates.Where(x => x.Value.Config.ShowUgsAlerts).Select(x => x.Key));
+				HashSet<TemplateId> templates = new HashSet<TemplateId>(streamConfig.Templates.Where(x => x.ShowUgsAlerts).Select(x => x.Id));
 				if(templates.Count > 0)
 				{
-					newCachedDesktopAlerts[cachedStream.Id] = templates;
+					newCachedDesktopAlerts[streamConfig.Id] = templates;
 				}
 			}
 
@@ -294,7 +300,7 @@ namespace Horde.Build.Issues
 			}
 
 			// Update any issues that are listed as valid for a stream that no longer exists
-			HashSet<StreamId> validStreamIds = new HashSet<StreamId>(cachedStreams.Select(x => x.Id));
+			HashSet<StreamId> validStreamIds = new HashSet<StreamId>(globalConfig.Streams.Select(x => x.Id));
 			for (int idx = 0; idx < openIssues.Count; idx++)
 			{
 				IIssue? openIssue = openIssues[idx];
@@ -504,8 +510,8 @@ namespace Horde.Build.Issues
 
 			_logger.LogInformation("Updating issues for {JobId}:{BatchId}:{StepId}", job.Id, batchId, stepId);
 
-			IStream? stream = await _streams.GetStreamAsync(job.StreamId);
-			if (stream == null)
+			StreamConfig? streamConfig;
+			if (!_globalConfig.CurrentValue.TryGetStream(job.StreamId, out streamConfig))
 			{
 				throw new Exception($"Invalid stream id '{job.StreamId}' on job '{job.Id}'");
 			}
@@ -528,7 +534,7 @@ namespace Horde.Build.Issues
 
 			// Get all the annotations for this template
 			TemplateRefConfig? templateRef;
-			if (!stream.Config.TryGetTemplate(job.TemplateId, out templateRef))
+			if (!streamConfig.TryGetTemplate(job.TemplateId, out templateRef))
 			{
 				return;
 			}
@@ -545,7 +551,7 @@ namespace Horde.Build.Issues
 			WorkflowId? workflowId = annotations.WorkflowId;
 			if (workflowId != null)
 			{
-				if (!stream.Config.TryGetWorkflow(workflowId.Value, out workflow))
+				if (!streamConfig.TryGetWorkflow(workflowId.Value, out workflow))
 				{
 					_logger.LogWarning("Unable to find workflow {WorkflowId} for {JobId}:{BatchId}:{StepId}", workflowId, job.Id, batchId, stepId);
 					return;
@@ -580,14 +586,14 @@ namespace Horde.Build.Issues
 						{
 							continue;
 						}
-						if (!await AddEventsToNewSpans(stream, job, batch, step, node, openSpans, eventGroups, annotations, job.PromoteIssuesByDefault))
+						if (!await AddEventsToNewSpans(streamConfig, job, batch, step, node, openSpans, eventGroups, annotations, job.PromoteIssuesByDefault))
 						{
 							continue;
 						}
 					}
 
 					// Try to update the sentinels for any other open steps
-					if (!await TryUpdateSentinelsAsync(openSpans, stream, job, batch, step))
+					if (!await TryUpdateSentinelsAsync(openSpans, streamConfig, job, batch, step))
 					{
 						continue;
 					}
@@ -770,7 +776,7 @@ namespace Horde.Build.Issues
 		/// <summary>
 		/// Adds new spans for the given events
 		/// </summary>
-		/// <param name="stream">The stream containing the job</param>
+		/// <param name="streamConfig">The stream containing the job</param>
 		/// <param name="job">The job instance</param>
 		/// <param name="batch">The job batch</param>
 		/// <param name="step">The job step</param>
@@ -780,7 +786,7 @@ namespace Horde.Build.Issues
 		/// <param name="annotations"></param>
 		/// <param name="promoteByDefault"></param>
 		/// <returns>True if all events were added</returns>
-		async Task<bool> AddEventsToNewSpans(IStream stream, IJob job, IJobStepBatch batch, IJobStep step, INode node, List<IIssueSpan> openSpans, HashSet<IssueEventGroup> newEventGroups, IReadOnlyNodeAnnotations? annotations, bool promoteByDefault)
+		async Task<bool> AddEventsToNewSpans(StreamConfig streamConfig, IJob job, IJobStepBatch batch, IJobStep step, INode node, List<IIssueSpan> openSpans, HashSet<IssueEventGroup> newEventGroups, IReadOnlyNodeAnnotations? annotations, bool promoteByDefault)
 		{
 			while (newEventGroups.Count > 0)
 			{
@@ -805,13 +811,13 @@ namespace Horde.Build.Issues
 				NewIssueStepData stepData = new NewIssueStepData(job, batch, step, GetIssueSeverity(eventGroup.Events), annotations, promoteByDefault);
 
 				// Get the span data
-				NewIssueSpanData spanData = new NewIssueSpanData(stream.Id, stream.Name, job.TemplateId, node.Name, eventGroup.Fingerprint, stepData);
+				NewIssueSpanData spanData = new NewIssueSpanData(streamConfig.Id, streamConfig.Name, job.TemplateId, node.Name, eventGroup.Fingerprint, stepData);
 
 				IJobStepRef? prevJob = await _jobStepRefs.GetPrevStepForNodeAsync(job.StreamId, job.TemplateId, node.Name, job.Change, JobStepOutcome.Success, true);
 				if (prevJob != null)
 				{
 					spanData.LastSuccess = new NewIssueStepData(prevJob);
-					spanData.Suspects = await FindSuspectsForSpanAsync(stream, spanData.Fingerprint, spanData.LastSuccess.Change + 1, spanData.FirstFailure.Change);
+					spanData.Suspects = await FindSuspectsForSpanAsync(streamConfig, spanData.Fingerprint, spanData.LastSuccess.Change + 1, spanData.FirstFailure.Change);
 				}
 
 				IJobStepRef? nextJob = await _jobStepRefs.GetNextStepForNodeAsync(job.StreamId, job.TemplateId, node.Name, job.Change, JobStepOutcome.Success, true);
@@ -904,14 +910,15 @@ namespace Horde.Build.Issues
 			Dictionary<(StreamId, int), bool> fixChangeCache = new Dictionary<(StreamId, int), bool>();
 			for (; ; )
 			{
+				GlobalConfig globalConfig = _globalConfig.CurrentValue;
+
 				// Find all the spans that are attached to the issue
 				List<IIssueSpan> spans = await _issueCollection.FindSpansAsync(issue.Id);
 
 				// Remove any spans for streams that have been deleted
 				for (int idx = spans.Count - 1; idx >= 0; idx--)
 				{
-					IStream? stream = await _streams.GetStreamAsync(spans[idx].StreamId);
-					if (stream == null || stream.Deleted)
+					if (!globalConfig.TryGetStream(spans[idx].StreamId, out _))
 					{
 						spans.RemoveAt(idx);
 					}
@@ -1024,7 +1031,7 @@ namespace Horde.Build.Issues
 							newStream.ContainsFix = stream.ContainsFix;
 						}
 
-						newStream.ContainsFix ??= await ContainsFixChange(newStream.StreamId, issue.FixChange.Value, fixChangeCache);
+						newStream.ContainsFix ??= await ContainsFixChange(globalConfig, newStream.StreamId, issue.FixChange.Value, fixChangeCache);
 
 						if (spans.Any(x => x.StreamId == newStream.StreamId && x.LastFailure.Change > issue.FixChange.Value))
 						{
@@ -1101,16 +1108,16 @@ namespace Horde.Build.Issues
 		/// <summary>
 		/// Figure out if a stream contains a fix changelist
 		/// </summary>
-		async ValueTask<bool> ContainsFixChange(StreamId streamId, int fixChange, Dictionary<(StreamId, int), bool> cachedContainsFixChange)
+		async ValueTask<bool> ContainsFixChange(GlobalConfig globalConfig, StreamId streamId, int fixChange, Dictionary<(StreamId, int), bool> cachedContainsFixChange)
 		{
 			bool containsFixChange;
 			if (!cachedContainsFixChange.TryGetValue((streamId, fixChange), out containsFixChange) && fixChange > 0)
 			{
-				IStream? stream = await _streams.GetCachedStream(streamId);
-				if (stream != null)
+				StreamConfig? streamConfig;
+				if (globalConfig.TryGetStream(streamId, out streamConfig))
 				{
 					_logger.LogInformation("Querying fix changelist {FixChange} in {StreamId}", fixChange, streamId);
-					ICommit? change = await _commitService.GetCollection(stream).FindAsync(fixChange, fixChange, 1).FirstOrDefaultAsync();
+					ICommit? change = await _commitService.GetCollection(streamConfig).FindAsync(fixChange, fixChange, 1).FirstOrDefaultAsync();
 					containsFixChange = change != null;
 				}
 				cachedContainsFixChange[(streamId, fixChange)] = containsFixChange;
@@ -1121,20 +1128,20 @@ namespace Horde.Build.Issues
 		/// <summary>
 		/// Find suspects that can match a given span
 		/// </summary>
-		/// <param name="stream">The stream name</param>
+		/// <param name="streamConfig">The stream name</param>
 		/// <param name="fingerprint">The fingerprint for the span</param>
 		/// <param name="minChange">Minimum changelist to consider</param>
 		/// <param name="maxChange">Maximum changelist to consider</param>
 		/// <returns>List of suspects</returns>
-		async Task<List<NewIssueSpanSuspectData>> FindSuspectsForSpanAsync(IStream stream, IIssueFingerprint fingerprint, int minChange, int maxChange)
+		async Task<List<NewIssueSpanSuspectData>> FindSuspectsForSpanAsync(StreamConfig streamConfig, IIssueFingerprint fingerprint, int minChange, int maxChange)
 		{
 			List<NewIssueSpanSuspectData> suspects = new List<NewIssueSpanSuspectData>();
 			if (TryGetHandler(fingerprint, out IssueHandler? handler))
 			{
-				_logger.LogDebug("Querying for changes in {StreamName} between {MinChange} and {MaxChange}", stream.Name, minChange, maxChange);
+				_logger.LogDebug("Querying for changes in {StreamName} between {MinChange} and {MaxChange}", streamConfig.Name, minChange, maxChange);
 
 				// Get the submitted changes before this job
-				List<ICommit> changes = await _commitService.GetCollection(stream).FindAsync(minChange, maxChange, MaxChanges).ToListAsync();
+				List<ICommit> changes = await _commitService.GetCollection(streamConfig).FindAsync(minChange, maxChange, MaxChanges).ToListAsync();
 				_logger.LogDebug("Found {NumResults} changes", changes.Count);
 
 				// Get the handler to rank them
@@ -1246,19 +1253,19 @@ namespace Horde.Build.Issues
 		/// Update the sentinels for the given list of issues
 		/// </summary>
 		/// <param name="spans"></param>
-		/// <param name="stream">The stream instance</param>
+		/// <param name="streamConfig">The stream instance</param>
 		/// <param name="job"></param>
 		/// <param name="batch"></param>
 		/// <param name="step"></param>
 		/// <returns></returns>
-		async Task<bool> TryUpdateSentinelsAsync(List<IIssueSpan> spans, IStream stream, IJob job, IJobStepBatch batch, IJobStep step)
+		async Task<bool> TryUpdateSentinelsAsync(List<IIssueSpan> spans, StreamConfig streamConfig, IJob job, IJobStepBatch batch, IJobStep step)
 		{
 			foreach(IIssueSpan span in spans)
 			{
 				if (job.Change < span.FirstFailure.Change && (span.LastSuccess == null || job.Change > span.LastSuccess.Change))
 				{
 					NewIssueStepData newLastSuccess = new NewIssueStepData(job.Change, IssueSeverity.Unspecified, job.Name, job.Id, batch.Id, step.Id, step.StartTimeUtc ?? default, step.LogId, null, false);
-					List<NewIssueSpanSuspectData> newSuspects = await FindSuspectsForSpanAsync(stream, span.Fingerprint, job.Change + 1, span.FirstFailure.Change);
+					List<NewIssueSpanSuspectData> newSuspects = await FindSuspectsForSpanAsync(streamConfig, span.Fingerprint, job.Change + 1, span.FirstFailure.Change);
 
 					if (await _issueCollection.TryUpdateSpanAsync(span, newLastSuccess: newLastSuccess, newSuspects: newSuspects) == null)
 					{
diff --git a/Engine/Source/Programs/Horde/Horde.Build/Issues/IssuesController.cs b/Engine/Source/Programs/Horde/Horde.Build/Issues/IssuesController.cs
--- a/Engine/Source/Programs/Horde/Horde.Build/Issues/IssuesController.cs
+++ b/Engine/Source/Programs/Horde/Horde.Build/Issues/IssuesController.cs
@@ -7,7 +7,6 @@ using System.Threading;
 using System.Threading.Tasks;
 using Horde.Build.Acls;
 using Horde.Build.Auditing;
-using Horde.Build.Server;
 using Horde.Build.Utilities;
 using HordeCommon;
 using Microsoft.AspNetCore.Authorization;
@@ -20,6 +19,8 @@ using Horde.Build.Streams;
 using Horde.Build.Issues.External;
 using Horde.Build.Users;
 using Horde.Build.Jobs.Graphs;
+using Microsoft.Extensions.Options;
+using Horde.Build.Server;
 
 namespace Horde.Build.Issues
 {
@@ -42,24 +43,26 @@ namespace Horde.Build.Issues
 		private readonly IssueService _issueService;
 		private readonly IExternalIssueService _externalIssueService;
 		private readonly JobService _jobService;
-		private readonly StreamService _streamService;
+		private readonly IStreamCollection _streamCollection;
 		private readonly IUserCollection _userCollection;
 		private readonly ILogFileService _logFileService;
+		private readonly IOptionsSnapshot<GlobalConfig> _globalConfig;
 		private readonly ILogger<IssuesController> _logger;
 
 		/// <summary>
 		/// Constructor
 		/// </summary>
-		public IssuesController(AclService aclService, ILogger<IssuesController> logger, IIssueCollection issueCollection, IssueService issueService, JobService jobService, StreamService streamService, IUserCollection userCollection, ILogFileService logFileService, IExternalIssueService externalIssueService)
+		public IssuesController(AclService aclService, IIssueCollection issueCollection, IssueService issueService, JobService jobService, IStreamCollection streamCollection, IUserCollection userCollection, ILogFileService logFileService, IExternalIssueService externalIssueService, IOptionsSnapshot<GlobalConfig> globalConfig, ILogger<IssuesController> logger)
 		{
 			_aclService = aclService;
 			_issueCollection = issueCollection;
 			_issueService = issueService;
 			_jobService = jobService;
-			_streamService = streamService;
+			_streamCollection = streamCollection;
 			_userCollection = userCollection;
 			_logFileService = logFileService;
 			_externalIssueService = externalIssueService;
+			_globalConfig = globalConfig;
 			_logger = logger;
 		}
 
@@ -88,16 +91,17 @@ namespace Horde.Build.Issues
 			List<object> responses = new List<object>();
 			if (streamId != null)
 			{
-				if (!await _streamService.AuthorizeAsync(streamId.Value, AclAction.ViewStream, User, new StreamPermissionsCache()))
-				{
-					return Forbid();
-				}
+				GlobalConfig globalConfig = _globalConfig.Value;
 
-				IStream? stream = await _streamService.GetStreamAsync(streamId.Value);
-				if (stream == null)
+				StreamConfig? streamConfig;
+				if (!globalConfig.TryGetStream(streamId.Value, out streamConfig))
 				{
 					return NotFound(streamId.Value);
 				}
+				if (!streamConfig.Authorize(AclAction.ViewStream, User))
+				{
+					return Forbid(AclAction.ViewStream, streamId.Value);
+				}
 
 				List<IIssueSpan> spans = await _issueCollection.FindSpansAsync(null, ids, streamId.Value, minChange, maxChange, resolved);
 				if(spans.Count > 0)
@@ -247,7 +251,7 @@ namespace Horde.Build.Issues
 				{
 					return NotFound(jobId.Value);
 				}
-				if(!await _jobService.AuthorizeAsync(job, AclAction.ViewJob, User, null))
+				if(!_globalConfig.Value.Authorize(job, AclAction.ViewJob, User))
 				{
 					return Forbid(AclAction.ViewJob, jobId.Value);
 				}
@@ -256,16 +260,14 @@ namespace Horde.Build.Issues
 				issues = await _issueService.Collection.FindIssuesForJobAsync(job, graph, stepId?.ToSubResourceId(), batchId?.ToSubResourceId(), labelIdx, ownerIdValue, resolved, promoted, index, count);
 			}
 
-			StreamPermissionsCache permissionsCache = new StreamPermissionsCache();
-
 			List<object> responses = new List<object>();
 			foreach (IIssue issue in issues)
 			{
 				IIssueDetails details = await _issueService.GetIssueDetailsAsync(issue);
-				if (await AuthorizeIssue(details, permissionsCache))
+				if (AuthorizeIssue(details))
 				{
 					bool showDesktopAlerts = _issueService.ShowDesktopAlertsForIssue(issue, details.Spans);
-					GetIssueResponse response = await CreateIssueResponseAsync(details, showDesktopAlerts);
+					GetIssueResponse response = CreateIssueResponse(details, showDesktopAlerts);
 					responses.Add(PropertyFilter.Apply(response, filter));
 				}
 			}
@@ -288,13 +290,13 @@ namespace Horde.Build.Issues
 			{
 				return NotFound();
 			}
-			if (!await AuthorizeIssue(details, null))
+			if (!AuthorizeIssue(details))
 			{
 				return Forbid();
 			}
 
 			bool showDesktopAlerts = _issueService.ShowDesktopAlertsForIssue(details.Issue, details.Spans);
-			return PropertyFilter.Apply(await CreateIssueResponseAsync(details, showDesktopAlerts), filter);
+			return PropertyFilter.Apply(CreateIssueResponse(details, showDesktopAlerts), filter);
 		}
 
 		/// <summary>
@@ -322,15 +324,15 @@ namespace Horde.Build.Issues
 		/// <param name="details"></param>
 		/// <param name="showDesktopAlerts"></param>
 		/// <returns></returns>
-		async Task<GetIssueResponse> CreateIssueResponseAsync(IIssueDetails details, bool showDesktopAlerts)
+		GetIssueResponse CreateIssueResponse(IIssueDetails details, bool showDesktopAlerts)
 		{
 			List<GetIssueAffectedStreamResponse> affectedStreams = new List<GetIssueAffectedStreamResponse>();
 			foreach (IGrouping<StreamId, IIssueSpan> streamSpans in details.Spans.GroupBy(x => x.StreamId))
 			{
 				try
 				{
-					IStream? stream = await _streamService.GetCachedStream(streamSpans.Key);
-					affectedStreams.Add(new GetIssueAffectedStreamResponse(details, stream, streamSpans));
+					_globalConfig.Value.TryGetStream(streamSpans.Key, out StreamConfig? streamConfig);
+					affectedStreams.Add(new GetIssueAffectedStreamResponse(details, streamConfig, streamSpans));
 				}
 				catch (Exception ex)
 				{
@@ -356,9 +358,7 @@ namespace Horde.Build.Issues
 			{
 				return NotFound();
 			}
-
-			StreamPermissionsCache cache = new StreamPermissionsCache();
-			if (!await AuthorizeIssue(issue, cache))
+			if (!AuthorizeIssue(issue))
 			{
 				return Forbid();
 			}
@@ -366,15 +366,11 @@ namespace Horde.Build.Issues
 			List<object> responses = new List<object>();
 			foreach (IGrouping<StreamId, IIssueSpan> spanGroup in issue.Spans.GroupBy(x => x.StreamId))
 			{
-				if (await _streamService.AuthorizeAsync(spanGroup.Key, AclAction.ViewStream, User, cache))
+				if (_globalConfig.Value.TryGetStream(spanGroup.Key, out StreamConfig? streamConfig) && streamConfig.Authorize(AclAction.ViewStream, User))
 				{
-					IStream? stream = await _streamService.GetCachedStream(spanGroup.Key);
-					if (stream != null)
-					{
-						HashSet<ObjectId> spanIds = new HashSet<ObjectId>(spanGroup.Select(x => x.Id));
-						List<IIssueStep> steps = issue.Steps.Where(x => spanIds.Contains(x.SpanId)).ToList();
-						responses.Add(PropertyFilter.Apply(new GetIssueStreamResponse(stream, spanGroup.ToList(), steps), filter));
-					}
+					HashSet<ObjectId> spanIds = new HashSet<ObjectId>(spanGroup.Select(x => x.Id));
+					List<IIssueStep> steps = issue.Steps.Where(x => spanIds.Contains(x.SpanId)).ToList();
+					responses.Add(PropertyFilter.Apply(new GetIssueStreamResponse(spanGroup.Key, spanGroup.ToList(), steps), filter));
 				}
 			}
 			return responses;
@@ -390,27 +386,23 @@ namespace Horde.Build.Issues
 		[HttpGet]
 		[Route("/api/v1/issues/{issueId}/streams/{streamId}")]
 		[ProducesResponseType(typeof(List<GetIssueStreamResponse>), 200)]
-		public async Task<ActionResult<object>> GetIssueStreamAsync(int issueId, string streamId, [FromQuery] PropertyFilter? filter = null)
+		public async Task<ActionResult<object>> GetIssueStreamAsync(int issueId, StreamId streamId, [FromQuery] PropertyFilter? filter = null)
 		{
 			IIssueDetails? details = await _issueService.GetIssueDetailsAsync(issueId);
 			if (details == null)
 			{
 				return NotFound();
 			}
-
-			StreamId streamIdValue = new StreamId(streamId);
-			if (!await _streamService.AuthorizeAsync(streamIdValue, AclAction.ViewStream, User, null))
+			if (!_globalConfig.Value.TryGetStream(streamId, out StreamConfig? streamConfig))
 			{
-				return Forbid();
+				return NotFound(streamId);
 			}
-
-			IStream? stream = await _streamService.GetCachedStream(streamIdValue);
-			if (stream == null)
+			if (!streamConfig.Authorize(AclAction.ViewStream, User))
 			{
-				return NotFound();
+				return Forbid(AclAction.ViewStream, streamId);
 			}
 
-			List<IIssueSpan> spans = details.Spans.Where(x => x.StreamId == streamIdValue).ToList();
+			List<IIssueSpan> spans = details.Spans.Where(x => x.StreamId == streamId).ToList();
 			if(spans.Count == 0)
 			{
 				return NotFound();
@@ -419,7 +411,7 @@ namespace Horde.Build.Issues
 			HashSet<ObjectId> spanIds = new HashSet<ObjectId>(spans.Select(x => x.Id));
 			List<IIssueStep> steps = details.Steps.Where(x => spanIds.Contains(x.SpanId)).ToList();
 
-			return PropertyFilter.Apply(new GetIssueStreamResponse(stream, spans, steps), filter);
+			return PropertyFilter.Apply(new GetIssueStreamResponse(streamId, spans, steps), filter);
 		}
 
 		/// <summary>
@@ -519,7 +511,7 @@ namespace Horde.Build.Issues
 					logFile = await _logFileService.GetLogFileAsync(logEvent.LogId, cancellationToken);
 					logFiles[logEvent.LogId] = logFile;
 				}
-				if (logFile != null && await _jobService.AuthorizeAsync(logFile.JobId, AclAction.ViewLog, User, permissionsCache))
+				if (logFile != null && await _jobService.AuthorizeAsync(logFile.JobId, AclAction.ViewLog, User, _globalConfig.Value))
 				{
 					ILogEventData data = await _logFileService.GetEventDataAsync(logFile, logEvent.LineIndex, logEvent.LineCount, cancellationToken);
 					GetLogEventResponse response = new GetLogEventResponse(logEvent, data, issueId);
@@ -533,13 +525,12 @@ namespace Horde.Build.Issues
 		/// Authorize the current user to see an issue
 		/// </summary>
 		/// <param name="issue">The issue to authorize</param>
-		/// <param name="permissionsCache">Cache of permissions</param>
 		/// <returns>True if the user is authorized to see the issue</returns>
-		private async Task<bool> AuthorizeIssue(IIssueDetails issue, StreamPermissionsCache? permissionsCache)
+		private bool AuthorizeIssue(IIssueDetails issue)
 		{
 			foreach (StreamId streamId in issue.Spans.Select(x => x.StreamId).Distinct())
 			{
-				if (await _streamService.AuthorizeAsync(streamId, AclAction.ViewStream, User, permissionsCache))
+				if (_globalConfig.Value.TryGetStream(streamId, out StreamConfig? streamConfig) && streamConfig.Authorize(AclAction.ViewStream, User))
 				{
 					return true;
 				}
@@ -625,19 +616,19 @@ namespace Horde.Build.Issues
 		[HttpGet]
 		[Authorize]
 		[Route("/api/v1/issues/external")]
-		public async Task<ActionResult<List<GetExternalIssueResponse>>> GetExternalIssuesAsync([FromQuery] string streamId, [FromQuery] string[] keys)
+		public async Task<ActionResult<List<GetExternalIssueResponse>>> GetExternalIssuesAsync([FromQuery] StreamId streamId, [FromQuery] string[] keys)
 		{
 			if (_externalIssueService == null)
 			{
 				return NotFound();
 			}
-
-			StreamId streamIdValue = new StreamId(streamId);
-			StreamPermissionsCache cache = new StreamPermissionsCache();
-
-			if (!await _streamService.AuthorizeAsync(streamIdValue, AclAction.ViewStream, User, cache))
+			if (!_globalConfig.Value.TryGetStream(streamId, out StreamConfig? streamConfig))
 			{
-				return Forbid();
+				return NotFound(streamId);
+			}
+			if (!streamConfig.Authorize(AclAction.ViewStream, User))
+			{
+				return Forbid(AclAction.ViewStream, streamId);
 			}
 
 			List<GetExternalIssueResponse> response = new List<GetExternalIssueResponse>();
@@ -669,11 +660,9 @@ namespace Horde.Build.Issues
 			}
 
 			StreamId streamIdValue = new StreamId(issueRequest.StreamId);
-			StreamPermissionsCache cache = new StreamPermissionsCache();
-
-			if (!await _streamService.AuthorizeAsync(streamIdValue, AclAction.ViewStream, User, cache))
+			if (!_globalConfig.Value.Authorize(AclAction.ViewStream, User))
 			{
-				return Forbid();
+				return Forbid(AclAction.ViewStream, streamIdValue);
 			}
 
 			IUser? user = await _userCollection.GetUserAsync(User);
@@ -709,16 +698,16 @@ namespace Horde.Build.Issues
 			}
 
 			StreamId streamIdValue = new StreamId(streamId);
-			StreamPermissionsCache cache = new StreamPermissionsCache();
-
-			if (!await _streamService.AuthorizeAsync(streamIdValue, AclAction.ViewStream, User, cache))
+			if (!_globalConfig.Value.TryGetStream(streamIdValue, out StreamConfig? streamConfig))
 			{
-				return Forbid();
+				return NotFound(streamIdValue);
+			}
+			if (!streamConfig.Authorize(AclAction.ViewStream, User))
+			{
+				return Forbid(AclAction.ViewStream, streamIdValue);
 			}
 
-			IStream? stream = await _streamService.GetStreamAsync(streamIdValue);
-
-			List<IExternalIssueProject> projects = await _externalIssueService.GetProjects(stream!);
+			List<IExternalIssueProject> projects = await _externalIssueService.GetProjects(streamConfig);
 			List<GetExternalIssueProjectResponse> response = new List<GetExternalIssueProjectResponse>();
 
 			projects.ForEach(project =>
diff --git a/Engine/Source/Programs/Horde/Horde.Build/Jobs/Artifacts/ArtifactsController.cs b/Engine/Source/Programs/Horde/Horde.Build/Jobs/Artifacts/ArtifactsController.cs
--- a/Engine/Source/Programs/Horde/Horde.Build/Jobs/Artifacts/ArtifactsController.cs
+++ b/Engine/Source/Programs/Horde/Horde.Build/Jobs/Artifacts/ArtifactsController.cs
@@ -12,10 +12,12 @@ using System.Threading.Tasks;
 using Horde.Build.Acls;
 using Horde.Build.Jobs.Graphs;
 using Horde.Build.Server;
+using Horde.Build.Streams;
 using Horde.Build.Utilities;
 using Microsoft.AspNetCore.Authorization;
 using Microsoft.AspNetCore.Http;
 using Microsoft.AspNetCore.Mvc;
+using Microsoft.Extensions.Options;
 using Microsoft.IdentityModel.Tokens;
 using MongoDB.Bson;
 
@@ -30,35 +32,22 @@ namespace Horde.Build.Jobs.Artifacts
 	[Route("[controller]")]
 	public class ArtifactsController : ControllerBase
 	{
-		/// <summary>
-		/// Instance of the globals service
-		/// </summary>
 		private readonly GlobalsService _globalsService;
-
-		/// <summary>
-		/// Instance of the artifact collection
-		/// </summary>
 		private readonly IArtifactCollection _artifactCollection;
-
-		/// <summary>
-		/// Instance of the ACL service
-		/// </summary>
 		private readonly AclService _aclService;
-
-		/// <summary>
-		/// Instance of the Job service
-		/// </summary>
 		private readonly JobService _jobService;
+		private readonly IOptionsSnapshot<GlobalConfig> _globalConfig;
 
 		/// <summary>
 		/// Constructor
 		/// </summary>
-		public ArtifactsController(GlobalsService globalsService, IArtifactCollection artifactCollection, AclService aclService, JobService jobService)
+		public ArtifactsController(GlobalsService globalsService, IArtifactCollection artifactCollection, AclService aclService, JobService jobService, IOptionsSnapshot<GlobalConfig> globalConfig)
 		{
 			_globalsService = globalsService;
 			_artifactCollection = artifactCollection;
 			_aclService = aclService;
 			_jobService = jobService;
+			_globalConfig = globalConfig;
 		}
 
 		/// <summary>
@@ -76,10 +65,15 @@ namespace Horde.Build.Jobs.Artifacts
 			IJob? job = await _jobService.GetJobAsync(jobId);
 			if(job == null)
 			{
-				return NotFound();
+				return NotFound(jobId);
 			}
 
-			if (!await _jobService.AuthorizeAsync(job, AclAction.UploadArtifact, User, null))
+			StreamConfig? streamConfig;
+			if (!_globalConfig.Value.TryGetStream(job.StreamId, out streamConfig))
+			{
+				return NotFound(job.StreamId);
+			}
+			if (!streamConfig.Authorize(AclAction.UploadArtifact, User))
 			{
 				return Forbid();
 			}
@@ -121,7 +115,7 @@ namespace Horde.Build.Jobs.Artifacts
 			{
 				return NotFound();
 			}
-			if (!await _jobService.AuthorizeAsync(artifact.JobId, AclAction.UploadArtifact, User, null))
+			if (!await _jobService.AuthorizeAsync(artifact.JobId, AclAction.UploadArtifact, User, _globalConfig.Value))
 			{
 				return Forbid();
 			}
@@ -144,7 +138,7 @@ namespace Horde.Build.Jobs.Artifacts
 		[ProducesResponseType(typeof(List<GetArtifactResponse>), 200)]
 		public async Task<ActionResult<List<object>>> GetArtifacts([FromQuery] JobId jobId, [FromQuery] string? stepId = null, [FromQuery] bool code = false, [FromQuery] PropertyFilter? filter = null)
 		{
-			if (!await _jobService.AuthorizeAsync(jobId, AclAction.DownloadArtifact, User, null))
+			if (!await _jobService.AuthorizeAsync(jobId, AclAction.DownloadArtifact, User, _globalConfig.Value))
 			{
 				return Forbid();
 			}
@@ -194,7 +188,7 @@ namespace Horde.Build.Jobs.Artifacts
 			{
 				return NotFound();
 			}
-			if (!await _jobService.AuthorizeAsync(artifact.JobId, AclAction.DownloadArtifact, User, null))
+			if (!await _jobService.AuthorizeAsync(artifact.JobId, AclAction.DownloadArtifact, User, _globalConfig.Value))
 			{
 				return Forbid();
 			}
@@ -218,7 +212,7 @@ namespace Horde.Build.Jobs.Artifacts
 			{
 				return NotFound();
 			}
-			if (!await _jobService.AuthorizeAsync(artifact.JobId, AclAction.DownloadArtifact, User, null))
+			if (!await _jobService.AuthorizeAsync(artifact.JobId, AclAction.DownloadArtifact, User, _globalConfig.Value))
 			{
 				return Forbid();
 			}
@@ -240,7 +234,7 @@ namespace Horde.Build.Jobs.Artifacts
 		{
 			SubResourceId stepIdValue = stepId.ToSubResourceId();
 
-			if (!await _jobService.AuthorizeAsync(jobId, AclAction.DownloadArtifact, User, null))
+			if (!await _jobService.AuthorizeAsync(jobId, AclAction.DownloadArtifact, User, _globalConfig.Value))
 			{
 				return Forbid();
 			}
@@ -349,7 +343,7 @@ namespace Horde.Build.Jobs.Artifacts
 			{
 				return NotFound();
 			}
-			if (!await _jobService.AuthorizeAsync(job, AclAction.DownloadArtifact, User, null))
+			if (!_globalConfig.Value.Authorize(job, AclAction.DownloadArtifact, User))
 			{
 				return Forbid();
 			}
diff --git a/Engine/Source/Programs/Horde/Horde.Build/Jobs/IJob.cs b/Engine/Source/Programs/Horde/Horde.Build/Jobs/IJob.cs
--- a/Engine/Source/Programs/Horde/Horde.Build/Jobs/IJob.cs
+++ b/Engine/Source/Programs/Horde/Horde.Build/Jobs/IJob.cs
@@ -4,7 +4,9 @@ using System;
 using System.Collections.Generic;
 using System.Diagnostics.CodeAnalysis;
 using System.Linq;
+using System.Security.Claims;
 using EpicGames.Core;
+using Horde.Build.Acls;
 using Horde.Build.Agents;
 using Horde.Build.Agents.Leases;
 using Horde.Build.Agents.Pools;
@@ -12,6 +14,7 @@ using Horde.Build.Agents.Sessions;
 using Horde.Build.Jobs.Graphs;
 using Horde.Build.Jobs.Timing;
 using Horde.Build.Logs;
+using Horde.Build.Server;
 using Horde.Build.Streams;
 using Horde.Build.Ugs;
 using Horde.Build.Users;
@@ -741,6 +744,26 @@ namespace Horde.Build.Jobs
 	/// </summary>
 	public static class JobExtensions
 	{
+		/// <summary>
+		/// Determines if the user is authorized to perform an action on a particular stream
+		/// </summary>
+		/// <param name="globalConfig">Current global config instance</param>
+		/// <param name="job">Job to authorize for</param>
+		/// <param name="action">The action being performed</param>
+		/// <param name="user">The principal to authorize</param>
+		/// <returns>True if the action is authorized</returns>
+		public static bool Authorize(this GlobalConfig globalConfig, IJob? job, AclAction action, ClaimsPrincipal user)
+		{
+			if (job != null && globalConfig.TryGetStream(job.StreamId, out StreamConfig? streamConfig))
+			{
+				return streamConfig.Authorize(action, user);
+			}
+			else
+			{
+				return globalConfig.Authorize(action, user);
+			}
+		}
+
 		/// <summary>
 		/// Attempts to get a step with a given ID
 		/// </summary>
diff --git a/Engine/Source/Programs/Horde/Horde.Build/Jobs/JobMessages.cs b/Engine/Source/Programs/Horde/Horde.Build/Jobs/JobMessages.cs
--- a/Engine/Source/Programs/Horde/Horde.Build/Jobs/JobMessages.cs
+++ b/Engine/Source/Programs/Horde/Horde.Build/Jobs/JobMessages.cs
@@ -164,11 +164,6 @@ namespace Horde.Build.Jobs
 		/// New list of arguments for the job. Only -Target= arguments can be modified after the job has started.
 		/// </summary>
 		public List<string>? Arguments { get; set; }
-
-		/// <summary>
-		/// Custom permissions for this object
-		/// </summary>
-		public AclConfig? Acl { get; set; }
 	}
 
 	/// <summary>
@@ -353,11 +348,6 @@ namespace Horde.Build.Jobs
 		/// </summary>
 		public bool UpdateIssues { get; set; }
 
-		/// <summary>
-		/// Per-object permissions
-		/// </summary>
-		public GetAclResponse? Acl { get; set; }
-
 		/// <summary>
 		/// Constructor
 		/// </summary>
diff --git a/Engine/Source/Programs/Horde/Horde.Build/Jobs/JobService.cs b/Engine/Source/Programs/Horde/Horde.Build/Jobs/JobService.cs
--- a/Engine/Source/Programs/Horde/Horde.Build/Jobs/JobService.cs
+++ b/Engine/Source/Programs/Horde/Horde.Build/Jobs/JobService.cs
@@ -3,6 +3,7 @@
 using System;
 using System.Collections.Generic;
 using System.Globalization;
+using System.IO;
 using System.Linq;
 using System.Security.Claims;
 using System.Text.RegularExpressions;
@@ -21,6 +22,7 @@ using Horde.Build.Jobs.Timing;
 using Horde.Build.Logs;
 using Horde.Build.Notifications;
 using Horde.Build.Perforce;
+using Horde.Build.Server;
 using Horde.Build.Streams;
 using Horde.Build.Users;
 using Horde.Build.Utilities;
@@ -65,7 +67,7 @@ namespace Horde.Build.Jobs
 	/// <summary>
 	/// Cache of information about job ACLs
 	/// </summary>
-	public class JobPermissionsCache : StreamPermissionsCache
+	public class JobPermissionsCache
 	{
 		/// <summary>
 		/// Map of job id to permissions for that job
@@ -86,7 +88,7 @@ namespace Horde.Build.Jobs
 		readonly IUserCollection _userCollection;
 		readonly INotificationTriggerCollection _triggerCollection;
 		readonly JobTaskSource _jobTaskSource;
-		readonly StreamService _streamService;
+		readonly IStreamCollection _streamCollection;
 		readonly ITemplateCollection _templateCollection;
 		readonly IssueService? _issueService;
 		readonly IPerforceService _perforceService;
@@ -122,20 +124,7 @@ namespace Horde.Build.Jobs
 		/// <summary>
 		/// Constructor
 		/// </summary>
-		/// <param name="jobs">The jobs collection</param>
-		/// <param name="graphs">The graphs collection</param>
-		/// <param name="agents">The agents collection</param>
-		/// <param name="jobStepRefs">The jobsteprefs collection</param>
-		/// <param name="jobTimings">The job timing document collection</param>
-		/// <param name="userCollection">User profiles</param>
-		/// <param name="triggerCollection">Trigger collection</param>
-		/// <param name="jobTaskSource">The queue service</param>
-		/// <param name="streamService">The stream service</param>
-		/// <param name="templateCollection">The template service</param>
-		/// <param name="issueService">The issue service</param>
-		/// <param name="perforceService">The perforce service</param>
-		/// <param name="logger">Log output</param>
-		public JobService(IJobCollection jobs, IGraphCollection graphs, IAgentCollection agents, IJobStepRefCollection jobStepRefs, IJobTimingCollection jobTimings, IUserCollection userCollection, INotificationTriggerCollection triggerCollection, JobTaskSource jobTaskSource, StreamService streamService, ITemplateCollection templateCollection, IssueService issueService, IPerforceService perforceService, ILogger<JobService> logger)
+		public JobService(IJobCollection jobs, IGraphCollection graphs, IAgentCollection agents, IJobStepRefCollection jobStepRefs, IJobTimingCollection jobTimings, IUserCollection userCollection, INotificationTriggerCollection triggerCollection, JobTaskSource jobTaskSource, IStreamCollection streamCollection, ITemplateCollection templateCollection, IssueService issueService, IPerforceService perforceService, ILogger<JobService> logger)
 		{
 			_jobs = jobs;
 			_graphs = graphs;
@@ -145,7 +134,7 @@ namespace Horde.Build.Jobs
 			_userCollection = userCollection;
 			_triggerCollection = triggerCollection;
 			_jobTaskSource = jobTaskSource;
-			_streamService = streamService;
+			_streamCollection = streamCollection;
 			_templateCollection = templateCollection;
 			_issueService = issueService;
 			_perforceService = perforceService;
@@ -168,7 +157,7 @@ namespace Horde.Build.Jobs
 		/// Creates a new job
 		/// </summary>
 		/// <param name="jobId">A requested job id</param>
-		/// <param name="stream">The stream that this job belongs to</param>
+		/// <param name="streamConfig">The stream that this job belongs to</param>
 		/// <param name="templateRefId">Name of the template ref</param>
 		/// <param name="templateHash">Template for this job</param>
 		/// <param name="graph">The graph for the new job</param>
@@ -177,11 +166,11 @@ namespace Horde.Build.Jobs
 		/// <param name="codeChange">The corresponding code changelist</param>
 		/// <param name="options">Options for the new job</param>
 		/// <returns>Unique id representing the job</returns>
-		public async Task<IJob> CreateJobAsync(JobId? jobId, IStream stream, TemplateId templateRefId, ContentHash templateHash, IGraph graph, string name, int change, int codeChange, CreateJobOptions options)
+		public async Task<IJob> CreateJobAsync(JobId? jobId, StreamConfig streamConfig, TemplateId templateRefId, ContentHash templateHash, IGraph graph, string name, int change, int codeChange, CreateJobOptions options)
 		{
 			using IScope traceScope = GlobalTracer.Instance.BuildSpan("JobService.CreateJobAsync").StartActive();
 			traceScope.Span.SetTag("JobId", jobId);
-			traceScope.Span.SetTag("Stream", stream.Name);
+			traceScope.Span.SetTag("Stream", streamConfig.Name);
 			traceScope.Span.SetTag("TemplateRefId", templateRefId);
 			traceScope.Span.SetTag("TemplateHash", templateHash);
 			traceScope.Span.SetTag("GraphId", graph.Id);
@@ -213,9 +202,9 @@ namespace Horde.Build.Jobs
 			JobId jobIdValue = jobId ?? Horde.Build.Utilities.ObjectId<IJob>.GenerateNewId();
 			using IDisposable scope = _logger.BeginScope("CreateJobAsync({JobId})", jobIdValue);
 
-			if (options.PreflightChange != null && ShouldClonePreflightChange(stream.Id))
+			if (options.PreflightChange != null && ShouldClonePreflightChange(streamConfig.Id))
 			{
-				options.ClonedPreflightChange = await CloneShelvedChangeAsync(stream.Config.ClusterName, options.ClonedPreflightChange ?? options.PreflightChange.Value);
+				options.ClonedPreflightChange = await CloneShelvedChangeAsync(streamConfig.ClusterName, options.ClonedPreflightChange ?? options.PreflightChange.Value);
 			}
 
 			_logger.LogInformation("Creating job at CL {Change}, code CL {CodeChange}, preflight CL {PreflightChange}, cloned CL {ClonedPreflightChange}", change, codeChange, options.PreflightChange, options.ClonedPreflightChange);
@@ -225,7 +214,7 @@ namespace Horde.Build.Jobs
 			properties["CodeChange"] = codeChange.ToString(CultureInfo.InvariantCulture);
 			properties["PreflightChange"] = options.PreflightChange?.ToString(CultureInfo.InvariantCulture) ?? String.Empty;
 			properties["ClonedPreflightChange"] = options.ClonedPreflightChange?.ToString(CultureInfo.InvariantCulture) ?? String.Empty;
-			properties["StreamId"] = stream.Id.ToString();
+			properties["StreamId"] = streamConfig.Id.ToString();
 			properties["TemplateId"] = templateRefId.ToString();
 			properties["JobId"] = jobIdValue.ToString();
 
@@ -236,7 +225,7 @@ namespace Horde.Build.Jobs
 
 			name = StringUtils.ExpandProperties(name, properties);
 
-			IJob newJob = await _jobs.AddAsync(jobIdValue, stream.Id, templateRefId, templateHash, graph, name, change, codeChange, options);
+			IJob newJob = await _jobs.AddAsync(jobIdValue, streamConfig.Id, templateRefId, templateHash, graph, name, change, codeChange, options);
 			_jobTaskSource.UpdateQueuedJob(newJob, graph);
 
 			await _jobTaskSource.UpdateUgsBadges(newJob, graph, new List<(LabelState, LabelOutcome)>());
@@ -406,11 +395,11 @@ namespace Horde.Build.Jobs
 		/// Evaluate a change query to determine which CL to run a job at
 		/// </summary>
 		/// <returns></returns>
-		public async Task<int?> EvaluateChangeQueriesAsync(IStream stream, List<ChangeQueryConfig> queries, List<CommitTag>? commitTags, ICommitCollection commits, CancellationToken cancellationToken)
+		public async Task<int?> EvaluateChangeQueriesAsync(StreamId streamId, List<ChangeQueryConfig> queries, List<CommitTag>? commitTags, ICommitCollection commits, CancellationToken cancellationToken)
 		{
 			foreach (ChangeQueryConfig query in queries)
 			{
-				int? change = await EvaluateChangeQueryAsync(stream, query, commitTags, commits, cancellationToken);
+				int? change = await EvaluateChangeQueryAsync(streamId, query, commitTags, commits, cancellationToken);
 				if (change != null)
 				{
 					return change;
@@ -423,7 +412,7 @@ namespace Horde.Build.Jobs
 		/// Evaluate a change query to determine which CL to run a job at
 		/// </summary>
 		/// <returns></returns>
-		public async ValueTask<int?> EvaluateChangeQueryAsync(IStream stream, ChangeQueryConfig query, List<CommitTag>? commitTags, ICommitCollection commits, CancellationToken cancellationToken)
+		public async ValueTask<int?> EvaluateChangeQueryAsync(StreamId streamId, ChangeQueryConfig query, List<CommitTag>? commitTags, ICommitCollection commits, CancellationToken cancellationToken)
 		{
 			if (query.Condition == null || query.Condition.Evaluate(propertyName => GetTagPropertyValues(propertyName, commitTags)))
 			{
@@ -443,7 +432,7 @@ namespace Horde.Build.Jobs
 				{
 					List<JobStepOutcome> outcomes = query.Outcomes ?? new List<JobStepOutcome> { JobStepOutcome.Success };
 
-					IList<IJob> jobs = await FindJobsAsync(streamId: stream.Id, templates: new[] { query.TemplateId.Value }, target: query.Target, state: new[] { JobStepState.Completed }, outcome: outcomes.ToArray(), count: 1, excludeUserJobs: true);
+					IList<IJob> jobs = await FindJobsAsync(streamId: streamId, templates: new[] { query.TemplateId.Value }, target: query.Target, state: new[] { JobStepState.Completed }, outcome: outcomes.ToArray(), count: 1, excludeUserJobs: true);
 					if (jobs.Count > 0)
 					{
 						_logger.LogInformation("Last successful build of {TemplateId} target {Target} was job {JobId} at change {Change}", query.TemplateId, query.Target, jobs[0].Id, jobs[0].Change);
@@ -806,10 +795,11 @@ namespace Horde.Build.Jobs
 		/// </summary>
 		/// <param name="job">Job to update</param>
 		/// <param name="batchId">Unique id of the batch to update</param>
+		/// <param name="streamConfig">The current stream config</param>
 		/// <param name="newLogId">The new log file id</param>
 		/// <param name="newState">New state of the jobstep</param>
 		/// <returns>True if the job was updated, false if it was deleted</returns>
-		public async Task<IJob?> UpdateBatchAsync(IJob job, SubResourceId batchId, LogId? newLogId = null, JobStepBatchState? newState = null)
+		public async Task<IJob?> UpdateBatchAsync(IJob job, SubResourceId batchId, StreamConfig streamConfig, LogId? newLogId = null, JobStepBatchState? newState = null)
 		{
 			using IScope traceScope = GlobalTracer.Instance.BuildSpan("JobService.UpdateBatchAsync").StartActive();
 			traceScope.Span.SetTag("Job", job.Id);
@@ -846,7 +836,7 @@ namespace Horde.Build.Jobs
 					if (batch.Steps.Any(x => x.State == JobStepState.Waiting || x.State == JobStepState.Ready || x.State == JobStepState.Running))
 					{
 						// Check if the job is valid. If not, we will fail with a specific error code for it.
-						error ??= await CheckJobAsync(job) ?? JobStepBatchError.Incomplete;
+						error ??= await CheckJobAsync(job, streamConfig) ?? JobStepBatchError.Incomplete;
 						newError = error.Value;
 
 						// Find the agent and set the conform flag
@@ -888,20 +878,13 @@ namespace Horde.Build.Jobs
 			}
 		}
 
-		async Task<JobStepBatchError?> CheckJobAsync(IJob job)
+		async Task<JobStepBatchError?> CheckJobAsync(IJob job, StreamConfig streamConfig)
 		{
 			try
 			{
 				if (job.PreflightChange != 0)
 				{
-					IStream? stream = await _streamService.GetStreamAsync(job.StreamId);
-					if (stream == null)
-					{
-						_logger.LogWarning("Job {JobId} is no longer valid - stream {StreamId} does not exist.", job.Id, job.StreamId);
-						return JobStepBatchError.UnknownStream;
-					}
-
-					(CheckShelfResult result, _) = await _perforceService.CheckShelfAsync(stream, job.PreflightChange);
+					(CheckShelfResult result, _) = await _perforceService.CheckShelfAsync(streamConfig, job.PreflightChange);
 					if (result != CheckShelfResult.Ok)
 					{
 						_logger.LogWarning("Job {JobId} is no longer valid - check shelf returned {Result}", result);
@@ -938,6 +921,7 @@ namespace Horde.Build.Jobs
 		/// <param name="job">Job to update</param>
 		/// <param name="batchId">Unique id of the batch containing the step</param>
 		/// <param name="stepId">Unique id of the step to update</param>
+		/// <param name="streamConfig">Current global config</param>
 		/// <param name="newState">New state of the jobstep</param>
 		/// <param name="newOutcome">New outcome of the jobstep</param>
 		/// <param name="newError">New error for the step</param>
@@ -950,7 +934,7 @@ namespace Horde.Build.Jobs
 		/// <param name="newReports">New list of reports</param>
 		/// <param name="newProperties">Property changes. Any properties with a null value will be removed.</param>
 		/// <returns>True if the job was updated, false if it was deleted in the meantime</returns>
-		public async Task<IJob?> UpdateStepAsync(IJob job, SubResourceId batchId, SubResourceId stepId, JobStepState newState = JobStepState.Unspecified, JobStepOutcome newOutcome = JobStepOutcome.Unspecified, JobStepError? newError = null, bool? newAbortRequested = null, UserId? newAbortByUserId = null, LogId? newLogId = null, ObjectId? newNotificationTriggerId = null, UserId? newRetryByUserId = null, Priority? newPriority = null, List<Report>? newReports = null, Dictionary<string, string?>? newProperties = null)
+		public async Task<IJob?> UpdateStepAsync(IJob job, SubResourceId batchId, SubResourceId stepId, StreamConfig streamConfig, JobStepState newState = JobStepState.Unspecified, JobStepOutcome newOutcome = JobStepOutcome.Unspecified, JobStepError? newError = null, bool? newAbortRequested = null, UserId? newAbortByUserId = null, LogId? newLogId = null, ObjectId? newNotificationTriggerId = null, UserId? newRetryByUserId = null, Priority? newPriority = null, List<Report>? newReports = null, Dictionary<string, string?>? newProperties = null)
 		{
 			using IScope traceScope = GlobalTracer.Instance.BuildSpan("JobService.UpdateStepAsync").StartActive();
 			traceScope.Span.SetTag("Job", job.Id);
@@ -960,7 +944,7 @@ namespace Horde.Build.Jobs
 			using IDisposable scope = _logger.BeginScope("UpdateStepAsync({JobId}:{BatchId}:{StepId})", job.Id, batchId, stepId);
 			for (; ;)
 			{
-				IJob? newJob = await TryUpdateStepAsync(job, batchId, stepId, newState, newOutcome, newError, newAbortRequested, newAbortByUserId, newLogId, newNotificationTriggerId, newRetryByUserId, newPriority, newReports, newProperties);
+				IJob? newJob = await TryUpdateStepAsync(job, batchId, stepId, streamConfig, newState, newOutcome, newError, newAbortRequested, newAbortByUserId, newLogId, newNotificationTriggerId, newRetryByUserId, newPriority, newReports, newProperties);
 				if (newJob != null)
 				{
 					return newJob;
@@ -982,6 +966,7 @@ namespace Horde.Build.Jobs
 		/// <param name="job">Job to update</param>
 		/// <param name="batchId">Unique id of the batch containing the step</param>
 		/// <param name="stepId">Unique id of the step to update</param>
+		/// <param name="streamConfig">The current stream config</param>
 		/// <param name="newState">New state of the jobstep</param>
 		/// <param name="newOutcome">New outcome of the jobstep</param>
 		/// <param name="newError">New error for the step</param>
@@ -994,7 +979,7 @@ namespace Horde.Build.Jobs
 		/// <param name="newReports">New reports</param>
 		/// <param name="newProperties">Property changes. Any properties with a null value will be removed.</param>
 		/// <returns>True if the job was updated, false if it was deleted in the meantime</returns>
-		public async Task<IJob?> TryUpdateStepAsync(IJob job, SubResourceId batchId, SubResourceId stepId, JobStepState newState = JobStepState.Unspecified, JobStepOutcome newOutcome = JobStepOutcome.Unspecified, JobStepError? newError = null, bool? newAbortRequested = null, UserId? newAbortByUserId = null, LogId? newLogId = null, ObjectId? newTriggerId = null, UserId? newRetryByUserId = null, Priority? newPriority = null, List<Report>? newReports = null, Dictionary<string, string?>? newProperties = null)
+		public async Task<IJob?> TryUpdateStepAsync(IJob job, SubResourceId batchId, SubResourceId stepId, StreamConfig streamConfig, JobStepState newState = JobStepState.Unspecified, JobStepOutcome newOutcome = JobStepOutcome.Unspecified, JobStepError? newError = null, bool? newAbortRequested = null, UserId? newAbortByUserId = null, LogId? newLogId = null, ObjectId? newTriggerId = null, UserId? newRetryByUserId = null, Priority? newPriority = null, List<Report>? newReports = null, Dictionary<string, string?>? newProperties = null)
 		{
 			using IScope traceScope = GlobalTracer.Instance.BuildSpan("JobService.TryUpdateStepAsync").StartActive();
 			traceScope.Span.SetTag("Job", job.Id);
@@ -1057,15 +1042,11 @@ namespace Horde.Build.Jobs
 						{
 							if (job.AutoSubmit && outcome == JobStepOutcome.Success && job.AbortedByUserId == null)
 							{
-								job = await AutoSubmitChangeAsync(job, graph);
+								job = await AutoSubmitChangeAsync(streamConfig, job, graph);
 							}
 							else if (job.ClonedPreflightChange != 0)
 							{
-								IStream? stream = await _streamService.GetCachedStream(job.StreamId);
-								if (stream != null)
-								{
-									await DeleteShelvedChangeAsync(stream.Config.ClusterName, job.ClonedPreflightChange);
-								}
+								await DeleteShelvedChangeAsync(streamConfig.ClusterName, job.ClonedPreflightChange);
 							}
 						}
 					}
@@ -1089,7 +1070,7 @@ namespace Horde.Build.Jobs
 								JobStepOutcome jobTriggerOutcome = job.GetTargetOutcome(graph, jobTrigger.Target);
 								if (jobTriggerOutcome == JobStepOutcome.Success || jobTriggerOutcome == JobStepOutcome.Warnings)
 								{
-									job = await FireJobTriggerAsync(job, graph, jobTrigger) ?? job;
+									job = await FireJobTriggerAsync(job, graph, jobTrigger, streamConfig) ?? job;
 								}
 							}
 						}
@@ -1138,10 +1119,11 @@ namespace Horde.Build.Jobs
 		/// <summary>
 		/// Submit the given change for a preflight
 		/// </summary>
+		/// <param name="streamConfig">The current stream config</param>
 		/// <param name="job">The job being run</param>
 		/// <param name="graph">Graph for the job</param>
 		/// <returns></returns>
-		private async Task<IJob> AutoSubmitChangeAsync(IJob job, IGraph graph)
+		private async Task<IJob> AutoSubmitChangeAsync(StreamConfig streamConfig, IJob job, IGraph graph)
 		{
 			using IScope traceScope = GlobalTracer.Instance.BuildSpan("JobService.AutoSubmitChangeAsync").StartActive();
 			traceScope.Span.SetTag("Job", job.Id);
@@ -1151,54 +1133,46 @@ namespace Horde.Build.Jobs
 			string message;
 			try
 			{
-				IStream? stream = await _streamService.GetCachedStream(job.StreamId);
-				if (stream != null)
+				int clonedPreflightChange = job.ClonedPreflightChange;
+				if (clonedPreflightChange == 0)
 				{
-					int clonedPreflightChange = job.ClonedPreflightChange;
-					if (clonedPreflightChange == 0)
+					if (ShouldClonePreflightChange(job.StreamId))
 					{
-						if (ShouldClonePreflightChange(job.StreamId))
-						{
-							clonedPreflightChange = await CloneShelvedChangeAsync(stream.Config.ClusterName, job.PreflightChange);
-						}
-						else
-						{
-							clonedPreflightChange = job.PreflightChange;
-						}
+						clonedPreflightChange = await CloneShelvedChangeAsync(streamConfig.ClusterName, job.PreflightChange);
+					}
+					else
+					{
+						clonedPreflightChange = job.PreflightChange;
 					}
+				}
 
-					_logger.LogInformation("Updating description for {ClonedPreflightChange}", clonedPreflightChange);
+				_logger.LogInformation("Updating description for {ClonedPreflightChange}", clonedPreflightChange);
 
-					await _perforceService.UpdateChangelistDescription(stream.Config.ClusterName, clonedPreflightChange, x => x.TrimEnd() + $"\n#preflight {job.Id}");
+				await _perforceService.UpdateChangelistDescription(streamConfig.ClusterName, clonedPreflightChange, x => x.TrimEnd() + $"\n#preflight {job.Id}");
 
-					_logger.LogInformation("Submitting change {Change} (through {ChangeCopy}) after successful completion of {JobId}", job.PreflightChange, clonedPreflightChange, job.Id);
-					(change, message) = await _perforceService.SubmitShelvedChangeAsync(stream, clonedPreflightChange, job.PreflightChange);
+				_logger.LogInformation("Submitting change {Change} (through {ChangeCopy}) after successful completion of {JobId}", job.PreflightChange, clonedPreflightChange, job.Id);
+				(change, message) = await _perforceService.SubmitShelvedChangeAsync(streamConfig, clonedPreflightChange, job.PreflightChange);
 
-					_logger.LogInformation("Attempt to submit {Change} (through {ChangeCopy}): {Message}", job.PreflightChange, clonedPreflightChange, message);
+				_logger.LogInformation("Attempt to submit {Change} (through {ChangeCopy}): {Message}", job.PreflightChange, clonedPreflightChange, message);
 
-					if (!String.IsNullOrEmpty(message))
-					{
-						message = Regex.Replace(message, @"^Submit validation failed.*\n(?:\s*'[^']*' validation failed:\s*\n)?", "");
-					}
+				if (!String.IsNullOrEmpty(message))
+				{
+					message = Regex.Replace(message, @"^Submit validation failed.*\n(?:\s*'[^']*' validation failed:\s*\n)?", "");
+				}
 
-					if (ShouldClonePreflightChange(job.StreamId))
+				if (ShouldClonePreflightChange(job.StreamId))
+				{
+					if (change != null && job.ClonedPreflightChange != 0)
 					{
-						if (change != null && job.ClonedPreflightChange != 0)
-						{
-							await DeleteShelvedChangeAsync(stream.Config.ClusterName, job.PreflightChange);
-						}
-					}
-					else
-					{
-						if (change != null && job.PreflightChange != 0)
-						{
-							await DeleteShelvedChangeAsync(stream.Config.ClusterName, job.PreflightChange);
-						}
+						await DeleteShelvedChangeAsync(streamConfig.ClusterName, job.PreflightChange);
 					}
 				}
 				else
 				{
-					(change, message) = ((int?)null, "Stream not found");
+					if (change != null && job.PreflightChange != 0)
+					{
+						await DeleteShelvedChangeAsync(streamConfig.ClusterName, job.PreflightChange);
+					}
 				}
 			}
 			catch (Exception ex)
@@ -1278,8 +1252,9 @@ namespace Horde.Build.Jobs
 		/// <param name="job">The job containing the trigger</param>
 		/// <param name="graph">Graph for the job containing the trigger</param>
 		/// <param name="jobTrigger">The trigger object to fire</param>
+		/// <param name="streamConfig">Current config for the stream containing the job</param>
 		/// <returns>New job instance</returns>
-		private async Task<IJob?> FireJobTriggerAsync(IJob job, IGraph graph, IChainedJob jobTrigger)
+		private async Task<IJob?> FireJobTriggerAsync(IJob job, IGraph graph, IChainedJob jobTrigger, StreamConfig streamConfig)
 		{
 			using IScope traceScope = GlobalTracer.Instance.BuildSpan("JobService.FireJobTriggerAsync").StartActive();
 			traceScope.Span.SetTag("Job", job.Id);
@@ -1296,35 +1271,23 @@ namespace Horde.Build.Jobs
 				IJob? newJob = await _jobs.TryUpdateJobAsync(job, graph, jobTrigger: new KeyValuePair<TemplateId, JobId>(jobTrigger.TemplateRefId, chainedJobId));
 				if(newJob != null)
 				{
-					IStream? stream = await _streamService.GetStreamAsync(newJob.StreamId);
-					if(stream == null)
-					{
-						_logger.LogWarning("Cannot find stream {StreamId} for downstream job", newJob.StreamId);
-						break;
-					}
-
-					ITemplateRef? templateRef;
-					if (!stream.Templates.TryGetValue(jobTrigger.TemplateRefId, out templateRef))
+					TemplateRefConfig? templateRefConfig;
+					if (!streamConfig.TryGetTemplate(jobTrigger.TemplateRefId, out templateRefConfig))
 					{
 						_logger.LogWarning("Cannot find template {TemplateId} in stream {StreamId}", jobTrigger.TemplateRefId, newJob.StreamId);
 						break;
 					}
 
-					ITemplate? template = await _templateCollection.GetAsync(templateRef.Hash);
-					if (template == null)
-					{
-						_logger.LogWarning("Cannot find template {TemplateHash}", templateRef.Hash);
-						break;
-					}
+					ITemplate template = await _templateCollection.GetOrAddAsync(templateRefConfig);
 
-					IGraph triggerGraph = await _graphs.AddAsync(template, stream.Config.InitialAgentType);
+					IGraph triggerGraph = await _graphs.AddAsync(template, streamConfig.InitialAgentType);
 					_logger.LogInformation("Creating downstream job {ChainedJobId} from job {JobId}", chainedJobId, newJob.Id);
 
-					CreateJobOptions options = new CreateJobOptions(templateRef.Config);
+					CreateJobOptions options = new CreateJobOptions(templateRefConfig);
 					options.PreflightChange = newJob.PreflightChange;
 					options.PreflightDescription = newJob.PreflightDescription;
 
-					await CreateJobAsync(chainedJobId, stream, jobTrigger.TemplateRefId, templateRef.Hash, triggerGraph, templateRef.Config.Name, newJob.Change, newJob.CodeChange, options);
+					await CreateJobAsync(chainedJobId, streamConfig, jobTrigger.TemplateRefId, template.Hash, triggerGraph, templateRefConfig.Name, newJob.Change, newJob.CodeChange, options);
 					return newJob;
 				}
 
@@ -1340,42 +1303,18 @@ namespace Horde.Build.Jobs
 			return null;
 		}
 
-		/// <summary>
-		/// Determines if the user is authorized to perform an action on a particular stream
-		/// </summary>
-		/// <param name="job">The job to check</param>
-		/// <param name="action">The action being performed</param>
-		/// <param name="user">The principal to authorize</param>
-		/// <param name="cache">Cache of stream permissions</param>
-		/// <returns>True if the action is authorized</returns>
-		public Task<bool> AuthorizeAsync(IJob job, AclAction action, ClaimsPrincipal user, StreamPermissionsCache? cache)
-		{
-			return _streamService.AuthorizeAsync(job.StreamId, action, user, cache);
-		}
-
 		/// <summary>
 		/// Determines if the user is authorized to perform an action on a particular stream
 		/// </summary>
 		/// <param name="jobId">The job to check</param>
 		/// <param name="action">The action being performed</param>
 		/// <param name="user">The principal to authorize</param>
-		/// <param name="cache">Cache of job permissions</param>
+		/// <param name="globalConfig">Current global config instance</param>
 		/// <returns>True if the action is authorized</returns>
-		public async Task<bool> AuthorizeAsync(JobId jobId, AclAction action, ClaimsPrincipal user, JobPermissionsCache? cache)
+		public async Task<bool> AuthorizeAsync(JobId jobId, AclAction action, ClaimsPrincipal user, GlobalConfig globalConfig)
 		{
-			StreamId streamId;
-			if (cache == null)
-			{
-				IJob? job = await GetJobAsync(jobId);
-				streamId = job?.StreamId ?? StreamId.Empty;
-			}
-			else if (!cache.Jobs.TryGetValue(jobId, out streamId))
-			{
-				IJob? job = await GetJobAsync(jobId);
-				streamId = job?.StreamId ?? StreamId.Empty;
-				cache.Jobs.Add(jobId, streamId);
-			}
-			return !streamId.IsEmpty && await _streamService.AuthorizeAsync(streamId, action, user, cache);
+			IJob? job = await GetJobAsync(jobId);
+			return globalConfig.Authorize(job, action, user);
 		}
 
 		/// <summary>
diff --git a/Engine/Source/Programs/Horde/Horde.Build/Jobs/JobTaskSource.cs b/Engine/Source/Programs/Horde/Horde.Build/Jobs/JobTaskSource.cs
--- a/Engine/Source/Programs/Horde/Horde.Build/Jobs/JobTaskSource.cs
+++ b/Engine/Source/Programs/Horde/Horde.Build/Jobs/JobTaskSource.cs
@@ -56,11 +56,6 @@ namespace Horde.Build.Jobs
 		[DebuggerDisplay("{_job.Id}:{Batch.Id} ({_poolId})")]
 		internal class QueueItem
 		{
-			/// <summary>
-			/// The stream for this job
-			/// </summary>
-			public IStream _stream;
-
 			/// <summary>
 			/// The job instance
 			/// </summary>
@@ -104,15 +99,13 @@ namespace Horde.Build.Jobs
 			/// <summary>
 			/// Constructor
 			/// </summary>
-			/// <param name="stream">The stream containing this job</param>
 			/// <param name="job">The job instance</param>
 			/// <param name="batchIdx">The batch index to execute</param>
 			/// <param name="poolId">Unique id of the pool of machines to allocate from</param>
 			/// <param name="workspace">The workspace that this job should run in</param>
 			/// <param name="useAutoSdk">Whether or not to use the AutoSDK</param>
-			public QueueItem(IStream stream, IJob job, int batchIdx, PoolId poolId, AgentWorkspace workspace, bool useAutoSdk)
+			public QueueItem(IJob job, int batchIdx, PoolId poolId, AgentWorkspace workspace, bool useAutoSdk)
 			{
-				_stream = stream;
 				_job = job;
 				_batchIdx = batchIdx;
 				_poolId = poolId;
@@ -187,7 +180,7 @@ namespace Horde.Build.Jobs
 		}
 
 		readonly GlobalsService _globalsService;
-		readonly StreamService _streamService;
+		readonly IStreamCollection _streamCollection;
 		readonly ILogFileService _logFileService;
 		readonly IAgentCollection _agentsCollection;
 		readonly IJobCollection _jobs;
@@ -198,6 +191,7 @@ namespace Horde.Build.Jobs
 		readonly PerforceLoadBalancer _perforceLoadBalancer;
 		readonly IClock _clock;
 		readonly IOptionsMonitor<ServerSettings> _settings;
+		readonly IOptionsMonitor<GlobalConfig> _globalConfig;
 		readonly ILogger<JobTaskSource> _logger;
 		readonly ITicker _ticker;
 
@@ -216,9 +210,6 @@ namespace Horde.Build.Jobs
 		// During a background queue refresh operation, any updated batches are added to this dictionary for merging into the updated queue.
 		List<QueueItem>? _newQueueItemsDuringUpdate;
 
-		// Cache of stream objects. Used to resolve agent types.
-		private Dictionary<StreamId, IStream> _streams = new Dictionary<StreamId, IStream>();
-
 		/// <summary>
 		/// Delegate for job schedule events
 		/// </summary>
@@ -235,7 +226,7 @@ namespace Horde.Build.Jobs
 		/// <summary>
 		/// Constructor
 		/// </summary>
-		public JobTaskSource(GlobalsService globalsService, IAgentCollection agents, IJobCollection jobs, IJobStepRefCollection jobStepRefs, IGraphCollection graphs, IPoolCollection pools, IUgsMetadataCollection ugsMetadataCollection, StreamService streamService, ILogFileService logFileService, PerforceLoadBalancer perforceLoadBalancer, IClock clock, IOptionsMonitor<ServerSettings> settings, ILogger<JobTaskSource> logger)
+		public JobTaskSource(GlobalsService globalsService, IAgentCollection agents, IJobCollection jobs, IJobStepRefCollection jobStepRefs, IGraphCollection graphs, IPoolCollection pools, IUgsMetadataCollection ugsMetadataCollection, IStreamCollection streamCollection, ILogFileService logFileService, PerforceLoadBalancer perforceLoadBalancer, IClock clock, IOptionsMonitor<ServerSettings> settings, IOptionsMonitor<GlobalConfig> globalConfig, ILogger<JobTaskSource> logger)
 		{
 			_globalsService = globalsService;
 			_agentsCollection = agents;
@@ -244,11 +235,12 @@ namespace Horde.Build.Jobs
 			_graphs = graphs;
 			_poolCollection = pools;
 			_ugsMetadataCollection = ugsMetadataCollection;
-			_streamService = streamService;
+			_streamCollection = streamCollection;
 			_logFileService = logFileService;
 			_perforceLoadBalancer = perforceLoadBalancer;
 			_clock = clock;
 			_ticker = clock.AddTicker<JobTaskSource>(s_refreshInterval, TickAsync, logger);
+			_globalConfig = globalConfig;
 			_settings = settings;
 			_logger = logger;
 
@@ -362,8 +354,9 @@ namespace Horde.Build.Jobs
 			}
 
 			// Query all the current streams
-			List<IStream> streamsList = await _streamService.GetStreamsAsync();
-			_streams = streamsList.ToDictionary(x => x.Id, x => x);
+			GlobalConfig globalConfig = _globalConfig.CurrentValue;
+			List<IStream> streamsList = await _streamCollection.GetAsync(globalConfig.Streams);
+			Dictionary<StreamId, IStream> streams = streamsList.ToDictionary(x => x.Id, x => x);
 
 			// Find all the pools which are valid (ie. have at least one online agent)
 			List<IAgent> agents = await _agentsCollection.FindAsync();
@@ -398,7 +391,7 @@ namespace Horde.Build.Jobs
 
 				// Get the stream. If it fails, skip the whole job.
 				IStream? stream;
-				if (!_streams.TryGetValue(newJob.StreamId, out stream))
+				if (!streams.TryGetValue(newJob.StreamId, out stream))
 				{
 					newJob = await _jobs.SkipAllBatchesAsync(newJob, graph, JobStepBatchError.UnknownStream);
 					continue;
@@ -429,7 +422,7 @@ namespace Horde.Build.Jobs
 					{
 						newJob = await SkipBatchAsync(newJob, batch.Id, graph, JobStepBatchError.NoAgentsInPool);
 					}
-					else if (!stream.TryGetAgentWorkspace(agentType, out (AgentWorkspace, bool)? workspaceResult))
+					else if (!stream.Config.TryGetAgentWorkspace(agentType, out (AgentWorkspace, bool)? workspaceResult))
 					{
 						newJob = await SkipBatchAsync(newJob, batch.Id, graph, JobStepBatchError.UnknownWorkspace);
 					}
@@ -466,7 +459,7 @@ namespace Horde.Build.Jobs
 						if (newJob != null)
 						{
 							(AgentWorkspace workspace, bool useAutoSdk) = workspaceResult.Value;
-							QueueItem newQueueItem = new QueueItem(stream, newJob, batchIdx, agentType.Pool, workspace, useAutoSdk);
+							QueueItem newQueueItem = new QueueItem(newJob, batchIdx, agentType.Pool, workspace, useAutoSdk);
 							newQueue.Add(newQueueItem);
 							newBatchIdToQueueItem[(newJob.Id, batch.Id)] = newQueueItem;
 
@@ -545,19 +538,6 @@ namespace Horde.Build.Jobs
 			}
 		}
 
-		/// <summary>
-		/// Updates the current state of a job
-		/// </summary>
-		/// <param name="job">The job that has been updated</param>
-		/// <param name="graph">Graph for the job</param>
-		/// <returns>Async task</returns>
-		public void UpdateQueuedJob(IJob job, IGraph graph)
-		{
-			IStream? stream;
-			_streams.TryGetValue(job.StreamId, out stream);
-			UpdateQueuedJob(job, graph, stream);
-		}
-
 		void AssignAnyQueueItemToWaiter(QueueWaiter waiter)
 		{
 			lock (_waiters)
@@ -625,16 +605,18 @@ namespace Horde.Build.Jobs
 		/// </summary>
 		/// <param name="job">The job that has been updated</param>
 		/// <param name="graph">Graph for the job</param>
-		/// <param name="stream">The stream containing the job</param>
-		public void UpdateQueuedJob(IJob job, IGraph graph, IStream? stream)
+		public void UpdateQueuedJob(IJob job, IGraph graph)
 		{
+			StreamConfig? streamConfig;
+			_globalConfig.CurrentValue.TryGetStream(job.StreamId, out streamConfig);
+
 			List<TaskCompletionSource<bool>> completeWaiters = new List<TaskCompletionSource<bool>>();
 			lock (_lockObject)
 			{
 				for (int batchIdx = 0; batchIdx < job.Batches.Count; batchIdx++)
 				{
 					IJobStepBatch batch = job.Batches[batchIdx];
-					if (batch.State == JobStepBatchState.Ready && stream != null && batch.AgentId == null)
+					if (batch.State == JobStepBatchState.Ready && streamConfig != null && batch.AgentId == null)
 					{
 						// Check if this item is already in the list.
 						QueueItem? existingItem;
@@ -651,7 +633,7 @@ namespace Horde.Build.Jobs
 								else
 								{
 									RemoveQueueItem(existingItem);
-									InsertQueueItem(stream, job, batchIdx, existingItem._poolId, existingItem._workspace, existingItem._useAutoSdk);
+									InsertQueueItem(job, batchIdx, existingItem._poolId, existingItem._workspace, existingItem._useAutoSdk);
 								}
 							}
 							continue;
@@ -662,13 +644,13 @@ namespace Horde.Build.Jobs
 
 						// Get the requirements for the new queue item
 						AgentConfig? agentType;
-						if (stream.Config.AgentTypes.TryGetValue(group.AgentType, out agentType))
+						if (streamConfig.AgentTypes.TryGetValue(group.AgentType, out agentType))
 						{
 							(AgentWorkspace, bool)? result;
-							if (stream.TryGetAgentWorkspace(agentType, out result))
+							if (streamConfig.TryGetAgentWorkspace(agentType, out result))
 							{
 								(AgentWorkspace agentWorkspace, bool useAutoSdk) = result.Value;
-								InsertQueueItem(stream, job, batchIdx, agentType.Pool, agentWorkspace, useAutoSdk);
+								InsertQueueItem(job, batchIdx, agentType.Pool, agentWorkspace, useAutoSdk);
 							}
 						}
 					}
@@ -766,8 +748,14 @@ namespace Horde.Build.Jobs
 			{
 				job = newJob;
 
+				StreamConfig? streamConfig;
+				if (!_globalConfig.CurrentValue.TryGetStream(newJob.StreamId, out streamConfig))
+				{
+					return null;
+				}
+
 				// Get the lease name
-				StringBuilder leaseName = new StringBuilder($"{item._stream.Name} - ");
+				StringBuilder leaseName = new StringBuilder($"{streamConfig.Name} - ");
 				if (job.PreflightChange > 0)
 				{
 					leaseName.Append((job.Change > 0) ? $"Preflight CL {job.PreflightChange} against CL {job.Change}" : $"Preflight CL {job.PreflightChange} against latest");
@@ -782,7 +770,7 @@ namespace Horde.Build.Jobs
 				IGlobals globals = await _globalsService.GetAsync();
 
 				// Encode the payload
-				ExecuteJobTask? task = await CreateExecuteJobTaskAsync(item._stream, job, batch, agent, item._workspace, item._useAutoSdk, logId);
+				ExecuteJobTask? task = await CreateExecuteJobTaskAsync(streamConfig, job, batch, agent, item._workspace, item._useAutoSdk, logId);
 				if (task != null)
 				{
 					byte[] payload = Any.Pack(task).ToByteArray();
@@ -832,10 +820,10 @@ namespace Horde.Build.Jobs
 			return null;
 		}
 
-		async Task<ExecuteJobTask?> CreateExecuteJobTaskAsync(IStream stream, IJob job, IJobStepBatch batch, IAgent agent, AgentWorkspace workspace, bool useAutoSdk, LogId logId)
+		async Task<ExecuteJobTask?> CreateExecuteJobTaskAsync(StreamConfig streamConfig, IJob job, IJobStepBatch batch, IAgent agent, AgentWorkspace workspace, bool useAutoSdk, LogId logId)
 		{
 			// Get the lease name
-			StringBuilder leaseName = new StringBuilder($"{stream.Name} - ");
+			StringBuilder leaseName = new StringBuilder($"{streamConfig.Name} - ");
 			if (job.PreflightChange > 0)
 			{
 				leaseName.Append((job.Change > 0) ? $"Preflight CL {job.PreflightChange} against CL {job.Change}" : $"Preflight CL {job.PreflightChange} against latest");
@@ -847,7 +835,7 @@ namespace Horde.Build.Jobs
 			leaseName.Append(CultureInfo.InvariantCulture, $" - {job.Name}");
 
 			// Get the global settings
-			IGlobals globals = await _globalsService.GetAsync();
+			GlobalConfig globalConfig = _globalConfig.CurrentValue;
 
 			// Encode the payload
 			ExecuteJobTask task = new ExecuteJobTask();
@@ -859,7 +847,7 @@ namespace Horde.Build.Jobs
 
 			List<HordeCommon.Rpc.Messages.AgentWorkspace> workspaces = new ();
 
-			PerforceCluster? cluster = globals.Config.FindPerforceCluster(workspace.Cluster);
+			PerforceCluster? cluster = globalConfig.FindPerforceCluster(workspace.Cluster);
 			if (cluster == null)
 			{
 				return null;
@@ -935,7 +923,7 @@ namespace Horde.Build.Jobs
 			}
 
 			// Cached stream for this job
-			IStream? stream = null;
+			StreamConfig? streamConfig = null;
 
 			// Send all the updates
 			Dictionary<int, IUgsMetadata> metadataCache = new Dictionary<int, IUgsMetadata>();
@@ -956,10 +944,9 @@ namespace Horde.Build.Jobs
 				}
 
 				// Get the stream
-				if (stream == null)
+				if (streamConfig == null)
 				{
-					stream = await _streamService.GetStreamAsync(job.StreamId);
-					if (stream == null)
+					if (!_globalConfig.CurrentValue.TryGetStream(job.StreamId, out streamConfig))
 					{
 						_logger.LogError("Unable to fetch definition for stream {StreamId}", job.StreamId);
 						break;
@@ -981,7 +968,7 @@ namespace Horde.Build.Jobs
 				IUgsMetadata? metadata;
 				if (!metadataCache.TryGetValue(change, out metadata))
 				{
-					metadata = await _ugsMetadataCollection.FindOrAddAsync(stream.Name, change, label.UgsProject);
+					metadata = await _ugsMetadataCollection.FindOrAddAsync(streamConfig.Name, change, label.UgsProject);
 					metadataCache[change] = metadata;
 				}
 
@@ -1099,18 +1086,17 @@ namespace Horde.Build.Jobs
 		/// <summary>
 		/// Inserts an item into the queue
 		/// </summary>
-		/// <param name="stream">The stream containing the job</param>
 		/// <param name="job"></param>
 		/// <param name="batchIdx"></param>
 		/// <param name="poolId">The pool to use</param>
 		/// <param name="workspace">The workspace for this item to run in</param>
 		/// <param name="useAutoSdk">Whether or not to use the AutoSDK</param>
 		/// <returns></returns>
-		void InsertQueueItem(IStream stream, IJob job, int batchIdx, PoolId poolId, AgentWorkspace workspace, bool useAutoSdk)
+		void InsertQueueItem(IJob job, int batchIdx, PoolId poolId, AgentWorkspace workspace, bool useAutoSdk)
 		{
 			_logger.LogDebug("Adding queued job {JobId}, batch {BatchId} [Pool: {Pool}, Workspace: {Workspace}]", job.Id, job.Batches[batchIdx].Id, poolId, workspace.Identifier);
 
-			QueueItem newItem = new QueueItem(stream, job, batchIdx, poolId, workspace, useAutoSdk);
+			QueueItem newItem = new QueueItem(job, batchIdx, poolId, workspace, useAutoSdk);
 			_batchIdToQueueItem[newItem.Id] = newItem;
 			_queue.Add(newItem);
 
