Horde: Refactor codebase to use new config system.

diff --git a/Engine/Source/Programs/Horde/Horde.Build/Jobs/JobsController.cs b/Engine/Source/Programs/Horde/Horde.Build/Jobs/JobsController.cs
--- a/Engine/Source/Programs/Horde/Horde.Build/Jobs/JobsController.cs
+++ b/Engine/Source/Programs/Horde/Horde.Build/Jobs/JobsController.cs
@@ -15,6 +15,7 @@ using Horde.Build.Jobs.Timing;
 using Horde.Build.Logs;
 using Horde.Build.Notifications;
 using Horde.Build.Perforce;
+using Horde.Build.Server;
 using Horde.Build.Streams;
 using Horde.Build.Users;
 using Horde.Build.Utilities;
@@ -22,6 +23,7 @@ using HordeCommon;
 using Microsoft.AspNetCore.Authorization;
 using Microsoft.AspNetCore.Mvc;
 using Microsoft.Extensions.Logging;
+using Microsoft.Extensions.Options;
 using MongoDB.Bson;
 using OpenTracing;
 using OpenTracing.Util;
@@ -44,30 +46,32 @@ namespace Horde.Build.Jobs
 		private readonly IGraphCollection _graphs;
 		private readonly ICommitService _commitService;
 		private readonly IPerforceService _perforce;
-		private readonly StreamService _streamService;
+		private readonly IStreamCollection _streamCollection;
 		private readonly JobService _jobService;
 		private readonly ITemplateCollection _templateCollection;
 		private readonly IArtifactCollection _artifactCollection;
 		private readonly IUserCollection _userCollection;
 		private readonly INotificationService _notificationService;
 		private readonly AgentService _agentService;
+		private readonly IOptionsSnapshot<GlobalConfig> _globalConfig;
 		private readonly ILogger<JobsController> _logger;
 
 		/// <summary>
 		/// Constructor
 		/// </summary>
-		public JobsController(IGraphCollection graphs, ICommitService commitService, IPerforceService perforce, StreamService streamService, JobService jobService, ITemplateCollection templateCollection, IArtifactCollection artifactCollection, IUserCollection userCollection, INotificationService notificationService, AgentService agentService, ILogger<JobsController> logger)
+		public JobsController(IGraphCollection graphs, ICommitService commitService, IPerforceService perforce, IStreamCollection streamCollection, JobService jobService, ITemplateCollection templateCollection, IArtifactCollection artifactCollection, IUserCollection userCollection, INotificationService notificationService, AgentService agentService, IOptionsSnapshot<GlobalConfig> globalConfig, ILogger<JobsController> logger)
 		{
 			_graphs = graphs;
 			_commitService = commitService;
 			_perforce = perforce;
-			_streamService = streamService;
+			_streamCollection = streamCollection;
 			_jobService = jobService;
 			_templateCollection = templateCollection;
 			_artifactCollection = artifactCollection;
 			_userCollection = userCollection;
 			_notificationService = notificationService;
 			_agentService = agentService;
+			_globalConfig = globalConfig;
 			_logger = logger;
 		}
 
@@ -80,31 +84,31 @@ namespace Horde.Build.Jobs
 		[Route("/api/v1/jobs")]
 		public async Task<ActionResult<CreateJobResponse>> CreateJobAsync([FromBody] CreateJobRequest create)
 		{
-			IStream? stream = await _streamService.GetStreamAsync(new StreamId(create.StreamId));
-			if (stream == null)
+			StreamConfig? streamConfig;
+			if (!_globalConfig.Value.TryGetStream(new StreamId(create.StreamId), out streamConfig))
 			{
 				return NotFound(create.StreamId);
 			}
-			if (!await _streamService.AuthorizeAsync(stream, AclAction.CreateJob, User, null))
+			if (!streamConfig.Authorize(AclAction.CreateJob, User))
 			{
-				return Forbid(AclAction.CreateJob, stream.Id);
+				return Forbid(AclAction.CreateJob, streamConfig.Id);
 			}
 
 			// Get the name of the template ref
 			TemplateId templateRefId = create.TemplateId;
 
 			// Augment the request with template properties
-			ITemplateRef? templateRef;
-			if (!stream.Templates.TryGetValue(templateRefId, out templateRef))
+			TemplateRefConfig? templateRefConfig;
+			if (!streamConfig.TryGetTemplate(templateRefId, out templateRefConfig))
 			{
-				return BadRequest($"Template {create.TemplateId} is not available for stream {stream.Id}");
+				return BadRequest($"Template {create.TemplateId} is not available for stream {streamConfig.Id}");
 			}
-			if (!await _streamService.AuthorizeAsync(stream, templateRef, AclAction.CreateJob, User, null))
+			if (!templateRefConfig.Authorize(AclAction.CreateJob, User))
 			{
-				return Forbid(AclAction.CreateJob, stream.Id);
+				return Forbid(AclAction.CreateJob, streamConfig.Id);
 			}
 
-			ITemplate? template = await _templateCollection.GetAsync(templateRef.Hash);
+			ITemplate? template = await _templateCollection.GetOrAddAsync(templateRefConfig);
 			if (template == null)
 			{
 				return BadRequest("Missing template referenced by {TemplateId}", create.TemplateId);
@@ -136,7 +140,7 @@ namespace Horde.Build.Jobs
 			ShelfInfo? shelfInfo = null;
 			if (create.PreflightChange != null)
 			{
-				(CheckShelfResult result, shelfInfo) = await _perforce.CheckShelfAsync(stream, create.PreflightChange.Value, HttpContext.RequestAborted);
+				(CheckShelfResult result, shelfInfo) = await _perforce.CheckShelfAsync(streamConfig, create.PreflightChange.Value, HttpContext.RequestAborted);
 				switch (result)
 				{
 					case CheckShelfResult.Ok:
@@ -146,7 +150,7 @@ namespace Horde.Build.Jobs
 					case CheckShelfResult.NoShelvedFiles:
 						return BadRequest(KnownLogEvents.Horde_InvalidPreflight, "CL {Change} does not contain any shelved files", create.PreflightChange);
 					case CheckShelfResult.WrongStream:
-						return BadRequest(KnownLogEvents.Horde_InvalidPreflight, "CL {Change} does not contain files in {Stream}", create.PreflightChange, stream.Name);
+						return BadRequest(KnownLogEvents.Horde_InvalidPreflight, "CL {Change} does not contain files in {Stream}", create.PreflightChange, streamConfig.Name);
 					case CheckShelfResult.MixedStream:
 						return BadRequest(KnownLogEvents.Horde_InvalidPreflight, "CL {Change} contains files from multiple streams", create.PreflightChange);
 					default:
@@ -164,13 +168,13 @@ namespace Horde.Build.Jobs
 			Priority priority = create.Priority ?? template.Priority ?? Priority.Normal;
 
 			// New groups for the job
-			IGraph graph = await _graphs.AddAsync(template, stream.Config.InitialAgentType);
+			IGraph graph = await _graphs.AddAsync(template, streamConfig.InitialAgentType);
 
 			// Get the commits for this stream
-			ICommitCollection commits = _commitService.GetCollection(stream);
+			ICommitCollection commits = _commitService.GetCollection(streamConfig);
 
 			// Get the change to build
-			int change = await GetChangeToBuildAsync(create, stream, template, shelfInfo, commits, HttpContext.RequestAborted);
+			int change = await GetChangeToBuildAsync(create, streamConfig.Id, template, shelfInfo, commits, HttpContext.RequestAborted);
 
 			// And get the matching code changelist
 			ICommit? lastCodeCommit = await commits.GetLastCodeChange(change, HttpContext.RequestAborted);
@@ -190,7 +194,7 @@ namespace Horde.Build.Jobs
 			}
 
 			// Create options for the new job
-			CreateJobOptions options = new CreateJobOptions(templateRef.Config);
+			CreateJobOptions options = new CreateJobOptions(templateRefConfig);
 			options.PreflightChange = create.PreflightChange;
 			options.PreflightDescription = shelfInfo?.Description;
 			options.StartedByUserId = User.GetUserId();
@@ -205,12 +209,12 @@ namespace Horde.Build.Jobs
 			}
 
 			// Create the job
-			IJob job = await _jobService.CreateJobAsync(null, stream, templateRefId, template.Id, graph, name, change, codeChange, options);
+			IJob job = await _jobService.CreateJobAsync(null, streamConfig, templateRefId, template.Hash, graph, name, change, codeChange, options);
 			await UpdateNotificationsAsync(job.Id, new UpdateNotificationsRequest { Slack = true });
 			return new CreateJobResponse(job.Id.ToString());
 		}
 
-		async ValueTask<int> GetChangeToBuildAsync(CreateJobRequest create, IStream stream, ITemplate template, ShelfInfo? shelfInfo, ICommitCollection commits, CancellationToken cancellationToken)
+		async ValueTask<int> GetChangeToBuildAsync(CreateJobRequest create, StreamId streamId, ITemplate template, ShelfInfo? shelfInfo, ICommitCollection commits, CancellationToken cancellationToken)
 		{
 			// If there's an explicit change specified, use that
 			if (create.Change.HasValue && create.Change.Value != -1)
@@ -221,7 +225,7 @@ namespace Horde.Build.Jobs
 			// Evaluate the change queries
 			if (create.ChangeQueries != null && create.ChangeQueries.Count > 0)
 			{
-				int? change = await _jobService.EvaluateChangeQueriesAsync(stream, create.ChangeQueries, shelfInfo?.Tags, commits, cancellationToken);
+				int? change = await _jobService.EvaluateChangeQueriesAsync(streamId, create.ChangeQueries, shelfInfo?.Tags, commits, cancellationToken);
 				if (change != null)
 				{
 					return change.Value;
@@ -252,7 +256,7 @@ namespace Horde.Build.Jobs
 			{
 				return NotFound(jobId);
 			}
-			if (!await _jobService.AuthorizeAsync(job, AclAction.DeleteJob, User, null))
+			if (!_globalConfig.Value.Authorize(job, AclAction.DeleteJob, User))
 			{
 				return Forbid(AclAction.DeleteJob, jobId);
 			}
@@ -279,14 +283,14 @@ namespace Horde.Build.Jobs
 				return NotFound(jobId);
 			}
 
-			StreamPermissionsCache permissionsCache = new StreamPermissionsCache();
-			if (!await _jobService.AuthorizeAsync(job, AclAction.UpdateJob, User, permissionsCache))
+			StreamConfig? streamConfig;
+			if (!_globalConfig.Value.TryGetStream(job.StreamId, out streamConfig))
 			{
-				return Forbid(AclAction.UpdateJob, jobId);
+				return NotFound(job.StreamId);
 			}
-			if (request.Acl != null && !await _jobService.AuthorizeAsync(job, AclAction.ChangePermissions, User, permissionsCache))
+			if (!streamConfig.Authorize(AclAction.UpdateJob, User))
 			{
-				return Forbid(AclAction.ChangePermissions, jobId);
+				return Forbid(AclAction.UpdateJob, jobId);
 			}
 
 			// Convert legacy behavior of clearing out the argument to setting the aborted flag
@@ -325,7 +329,13 @@ namespace Horde.Build.Jobs
 			{
 				return NotFound(jobId);
 			}
-			if (!await _jobService.AuthorizeAsync(job, AclAction.CreateSubscription, User, null))
+
+			StreamConfig? streamConfig;
+			if (!_globalConfig.Value.TryGetStream(job.StreamId, out streamConfig))
+			{
+				return NotFound(job.StreamId);
+			}
+			if (!streamConfig.Authorize(AclAction.CreateSubscription, User))
 			{
 				return Forbid(AclAction.CreateSubscription, jobId);
 			}
@@ -356,7 +366,13 @@ namespace Horde.Build.Jobs
 			{
 				return NotFound(jobId);
 			}
-			if (!await _jobService.AuthorizeAsync(job, AclAction.CreateSubscription, User, null))
+
+			StreamConfig? streamConfig;
+			if (!_globalConfig.Value.TryGetStream(job.StreamId, out streamConfig))
+			{
+				return NotFound(job.StreamId);
+			}
+			if (!streamConfig.Authorize(AclAction.CreateSubscription, User))
 			{
 				return Forbid(AclAction.CreateSubscription, jobId);
 			}
@@ -391,10 +407,14 @@ namespace Horde.Build.Jobs
 				return NotFound(jobId);
 			}
 
-			StreamPermissionsCache cache = new StreamPermissionsCache();
-			if (!await _jobService.AuthorizeAsync(job, AclAction.ViewJob, User, cache))
+			StreamConfig? streamConfig;
+			if (!_globalConfig.Value.TryGetStream(job.StreamId, out streamConfig))
 			{
-				return Forbid(AclAction.ViewJob, jobId);
+				return NotFound(job.StreamId);
+			}
+			if (!streamConfig.Authorize(AclAction.ViewJob, User))
+			{
+				return Forbid(AclAction.ViewJob, job.StreamId);
 			}
 			if (modifiedAfter != null && job.UpdateTimeUtc <= modifiedAfter.Value)
 			{
@@ -402,9 +422,8 @@ namespace Horde.Build.Jobs
 			}
 
 			IGraph graph = await _jobService.GetGraphAsync(job);
-			bool includeAcl = await _jobService.AuthorizeAsync(job, AclAction.ViewPermissions, User, cache);
-			bool includeCosts = await _jobService.AuthorizeAsync(job, AclAction.ViewCosts, User, cache);
-			return await CreateJobResponseAsync(job, graph, includeAcl, includeCosts, filter);
+			bool includeCosts = streamConfig.Authorize(AclAction.ViewCosts, User);
+			return await CreateJobResponseAsync(job, graph, includeCosts, filter);
 		}
 
 		/// <summary>
@@ -412,19 +431,18 @@ namespace Horde.Build.Jobs
 		/// </summary>
 		/// <param name="job">The job document</param>
 		/// <param name="graph">The graph for this job</param>
-		/// <param name="includeAcl">Whether to include the ACL in the response</param>
 		/// <param name="includeCosts">Whether to include costs in the response</param>
 		/// <param name="filter">Filter for the properties to return</param>
 		/// <returns>Object containing the requested properties</returns>
-		async Task<object> CreateJobResponseAsync(IJob job, IGraph graph, bool includeAcl, bool includeCosts, PropertyFilter? filter)
+		async Task<object> CreateJobResponseAsync(IJob job, IGraph graph, bool includeCosts, PropertyFilter? filter)
 		{
 			if (filter == null)
 			{
-				return await CreateJobResponseAsync(job, graph, true, true, includeAcl, includeCosts);
+				return await CreateJobResponseAsync(job, graph, true, true, includeCosts);
 			}
 			else
 			{
-				return filter.ApplyTo(await CreateJobResponseAsync(job, graph, filter.Includes(nameof(GetJobResponse.Batches)), filter.Includes(nameof(GetJobResponse.Labels)) || filter.Includes(nameof(GetJobResponse.DefaultLabel)), includeAcl, includeCosts));
+				return filter.ApplyTo(await CreateJobResponseAsync(job, graph, filter.Includes(nameof(GetJobResponse.Batches)), filter.Includes(nameof(GetJobResponse.Labels)) || filter.Includes(nameof(GetJobResponse.DefaultLabel)), includeCosts));
 			}
 		}
 
@@ -435,10 +453,9 @@ namespace Horde.Build.Jobs
 		/// <param name="graph">The graph definition</param>
 		/// <param name="includeBatches">Whether to include the job batches in the response</param>
 		/// <param name="includeLabels">Whether to include the job aggregates in the response</param>
-		/// <param name="includeAcl">Whether to include the ACL in the response</param>
 		/// <param name="includeCosts">Whether to include costs of running particular agents</param>
 		/// <returns>The response object</returns>
-		async ValueTask<GetJobResponse> CreateJobResponseAsync(IJob job, IGraph graph, bool includeBatches, bool includeLabels, bool includeAcl, bool includeCosts)
+		async ValueTask<GetJobResponse> CreateJobResponseAsync(IJob job, IGraph graph, bool includeBatches, bool includeLabels, bool includeCosts)
 		{
 			GetThinUserInfoResponse? startedByUserInfo = null;
 			if (job.StartedByUserId != null)
@@ -533,9 +550,15 @@ namespace Horde.Build.Jobs
 			{
 				return NotFound(jobId);
 			}
-			if (!await _jobService.AuthorizeAsync(job, AclAction.ViewJob, User, null))
+
+			StreamConfig? streamConfig;
+			if (!_globalConfig.Value.TryGetStream(job.StreamId, out streamConfig))
+			{
+				return NotFound(job.StreamId);
+			}
+			if (!streamConfig.Authorize(AclAction.ViewJob, User))
 			{
-				return Forbid(AclAction.ViewJob, job.Id);
+				return Forbid(AclAction.ViewJob, jobId);
 			}
 
 			IGraph graph = await _jobService.GetGraphAsync(job);
@@ -558,7 +581,13 @@ namespace Horde.Build.Jobs
 			{
 				return NotFound(jobId);
 			}
-			if (!await _jobService.AuthorizeAsync(job, AclAction.ViewJob, User, null))
+
+			StreamConfig? streamConfig;
+			if (!_globalConfig.Value.TryGetStream(job.StreamId, out streamConfig))
+			{
+				return NotFound(job.StreamId);
+			}
+			if (!streamConfig.Authorize(AclAction.ViewJob, User))
 			{
 				return Forbid(AclAction.ViewJob, jobId);
 			}
@@ -592,7 +621,7 @@ namespace Horde.Build.Jobs
 			GetJobResponse? jobResponse = null;
 			if (includeJobResponse)
 			{
-				jobResponse = await CreateJobResponseAsync(job, graph, true, true, false, true);
+				jobResponse = await CreateJobResponseAsync(job, graph, true, true, true);
 			}
 
 			return new GetJobTimingResponse(jobResponse, steps, labels);
@@ -629,7 +658,7 @@ namespace Horde.Build.Jobs
 			List<IJob> jobs = await _jobService.FindJobsByStreamWithTemplatesAsync(new StreamId(streamId), templateRefIds, count: count, consistentRead: false);
 
 			Dictionary<string, GetJobTimingResponse> jobTimings = await jobs.ToAsyncEnumerable()
-				.WhereAwait(async job => await _jobService.AuthorizeAsync(job, AclAction.ViewJob, User, null))
+				.Where(job => _globalConfig.Value.Authorize(job, AclAction.ViewJob, User))
 				.ToDictionaryAwaitAsync(x => ValueTask.FromResult(x.Id.ToString()), async job =>
 				{
 					IJobTiming jobTiming = await _jobService.GetJobTimingAsync(job);
@@ -656,7 +685,13 @@ namespace Horde.Build.Jobs
 			{
 				return NotFound(jobId);
 			}
-			if (!await _jobService.AuthorizeAsync(job, AclAction.ViewJob, User, null))
+
+			StreamConfig? streamConfig;
+			if (!_globalConfig.Value.TryGetStream(job.StreamId, out streamConfig))
+			{
+				return NotFound(job.StreamId);
+			}
+			if (!streamConfig.Authorize(AclAction.ViewJob, User))
 			{
 				return Forbid(AclAction.ViewJob, jobId);
 			}
@@ -753,53 +788,7 @@ namespace Horde.Build.Jobs
 					modifiedBefore, modifiedAfter, index, count, false);
 			}
 
-			StreamPermissionsCache permissionsCache = new StreamPermissionsCache();
-
-			List<object> responses = new List<object>();
-			foreach (IJob job in jobs)
-			{
-				using IScope jobScope = GlobalTracer.Instance.BuildSpan("JobIteration").StartActive();
-				jobScope.Span.SetTag("jobId", job.Id.ToString());
-
-				if (job.GraphHash == null)
-				{
-					_logger.LogWarning("Job {JobId} has no GraphHash", job.Id);
-					continue;
-				}
-
-				bool viewJobAuthorized;
-				using (IScope _ = GlobalTracer.Instance.BuildSpan("AuthorizeViewJob").StartActive())
-				{
-					viewJobAuthorized = await _jobService.AuthorizeAsync(job, AclAction.ViewJob, User, permissionsCache);
-				}
-				
-				if (viewJobAuthorized)
-				{
-					IGraph graph;
-					using (IScope _ = GlobalTracer.Instance.BuildSpan("GetGraph").StartActive())
-					{
-						graph = await _jobService.GetGraphAsync(job);
-					}
-
-					bool includeAcl;
-					using (IScope _ = GlobalTracer.Instance.BuildSpan("AuthorizeViewPermissions").StartActive())
-					{
-						includeAcl = await _jobService.AuthorizeAsync(job, AclAction.ViewPermissions, User, permissionsCache);
-					}
-
-					bool includeCosts;
-					using (IScope _ = GlobalTracer.Instance.BuildSpan("AuthorizeViewCosts").StartActive())
-					{
-						includeCosts = await _jobService.AuthorizeAsync(job, AclAction.ViewCosts, User, permissionsCache);
-					}
-
-					using (IScope _ = GlobalTracer.Instance.BuildSpan("CreateResponse").StartActive())
-					{
-						responses.Add(await CreateJobResponseAsync(job, graph, includeAcl, includeCosts, filter));
-					}
-				}
-			}
-			return responses;
+			return await CreateAuthorizedJobResponsesAsync(jobs, filter);
 		}
 
 		/// <summary>
@@ -835,52 +824,25 @@ namespace Horde.Build.Jobs
 			count = Math.Min(1000, count);
 
 			List<IJob> jobs = await _jobService.FindJobsByStreamWithTemplatesAsync(streamIdValue, templateRefIds, preflightStartedByUserIdValue, maxCreateTime, modifiedAfter, index, count, consistentRead);
-			return await CreateAuthorizedJobResponses(jobs, filter);
+			return await CreateAuthorizedJobResponsesAsync(jobs, filter);
 		}
 
-		private async Task<List<object>> CreateAuthorizedJobResponses(List<IJob> jobs, PropertyFilter? filter = null)
+		private async Task<List<object>> CreateAuthorizedJobResponsesAsync(List<IJob> jobs, PropertyFilter? filter = null)
 		{
-			StreamPermissionsCache permissionsCache = new ();
 			List<object> responses = new ();
-			foreach (IJob job in jobs)
+			foreach(IGrouping<StreamId, IJob> grouping in jobs.GroupBy(x => x.StreamId))
 			{
-				using IScope jobScope = GlobalTracer.Instance.BuildSpan("JobIteration").StartActive();
-				jobScope.Span.SetTag("jobId", job.Id.ToString());
-				
-				bool viewJobAuthorized;
-				using (IScope _ = GlobalTracer.Instance.BuildSpan("AuthorizeViewJob").StartActive())
+				StreamConfig? streamConfig;
+				if (_globalConfig.Value.TryGetStream(grouping.Key, out streamConfig) && streamConfig.Authorize(AclAction.ViewJob, User))
 				{
-					viewJobAuthorized = await _jobService.AuthorizeAsync(job, AclAction.ViewJob, User, permissionsCache);
-				}
-				
-				if (viewJobAuthorized)
-				{
-					IGraph graph;
-					using (IScope _ = GlobalTracer.Instance.BuildSpan("GetGraph").StartActive())
+					bool includeCosts = streamConfig.Authorize(AclAction.ViewCosts, User);
+					foreach (IJob job in grouping)
 					{
-						graph = await _jobService.GetGraphAsync(job);
-					}
-
-					bool includeAcl;
-					using (IScope _ = GlobalTracer.Instance.BuildSpan("AuthorizeViewPermissions").StartActive())
-					{
-						includeAcl = await _jobService.AuthorizeAsync(job, AclAction.ViewPermissions, User, permissionsCache);
-					}
-
-					bool includeCosts;
-					using (IScope _ = GlobalTracer.Instance.BuildSpan("AuthorizeViewCosts").StartActive())
-					{
-						includeCosts = await _jobService.AuthorizeAsync(job, AclAction.ViewCosts, User, permissionsCache);
-					}
-
-					using (IScope _ = GlobalTracer.Instance.BuildSpan("CreateResponse").StartActive())
-					{
-						responses.Add(await CreateJobResponseAsync(job, graph, includeAcl, includeCosts, filter));
+						IGraph graph = await _jobService.GetGraphAsync(job);
+						responses.Add(await CreateJobResponseAsync(job, graph, includeCosts, filter));
 					}
 				}
 			}
-
-			
 			return responses;
 		}
 
@@ -901,7 +863,13 @@ namespace Horde.Build.Jobs
 				{
 					return NotFound(jobId);
 				}
-				if (!await _jobService.AuthorizeAsync(job, AclAction.ExecuteJob, User, null))
+
+				StreamConfig? streamConfig;
+				if (!_globalConfig.Value.TryGetStream(job.StreamId, out streamConfig))
+				{
+					return NotFound(job.StreamId);
+				}
+				if (!streamConfig.Authorize(AclAction.ExecuteJob, User))
 				{
 					return Forbid(AclAction.ExecuteJob, jobId);
 				}
@@ -933,7 +901,13 @@ namespace Horde.Build.Jobs
 			{
 				return NotFound(jobId);
 			}
-			if (!await _jobService.AuthorizeAsync(job, AclAction.ViewJob, User, null))
+
+			StreamConfig? streamConfig;
+			if (!_globalConfig.Value.TryGetStream(job.StreamId, out streamConfig))
+			{
+				return NotFound(job.StreamId);
+			}
+			if (!streamConfig.Authorize(AclAction.ViewJob, User))
 			{
 				return Forbid(AclAction.ViewJob, jobId);
 			}
@@ -959,7 +933,13 @@ namespace Horde.Build.Jobs
 			{
 				return NotFound(jobId);
 			}
-			if (!await _jobService.AuthorizeAsync(job, AclAction.ViewJob, User, null))
+
+			StreamConfig? streamConfig;
+			if (!_globalConfig.Value.TryGetStream(job.StreamId, out streamConfig))
+			{
+				return NotFound(job.StreamId);
+			}
+			if (!streamConfig.Authorize(AclAction.ViewJob, User))
 			{
 				return Forbid(AclAction.ViewJob, jobId);
 			}
@@ -989,7 +969,13 @@ namespace Horde.Build.Jobs
 			{
 				return NotFound(jobId);
 			}
-			if (!await _jobService.AuthorizeAsync(job, AclAction.ViewJob, User, null))
+
+			StreamConfig? streamConfig;
+			if (!_globalConfig.Value.TryGetStream(job.StreamId, out streamConfig))
+			{
+				return NotFound(job.StreamId);
+			}
+			if (!streamConfig.Authorize(AclAction.ViewJob, User))
 			{
 				return Forbid(AclAction.ViewJob, jobId);
 			}
@@ -1020,7 +1006,13 @@ namespace Horde.Build.Jobs
 			{
 				return NotFound(jobId);
 			}
-			if (!await _jobService.AuthorizeAsync(job, AclAction.ViewJob, User, null))
+
+			StreamConfig? streamConfig;
+			if (!_globalConfig.Value.TryGetStream(job.StreamId, out streamConfig))
+			{
+				return NotFound(job.StreamId);
+			}
+			if (!streamConfig.Authorize(AclAction.ViewJob, User))
 			{
 				return Forbid(AclAction.ViewJob, jobId);
 			}
@@ -1055,13 +1047,17 @@ namespace Horde.Build.Jobs
 				return NotFound(jobId);
 			}
 
-			StreamPermissionsCache cache = new StreamPermissionsCache();
-			if (!await _jobService.AuthorizeAsync(job, AclAction.ViewJob, User, cache))
+			StreamConfig? streamConfig;
+			if (!_globalConfig.Value.TryGetStream(job.StreamId, out streamConfig))
 			{
-				return Forbid(AclAction.ViewJob, jobId);
+				return NotFound(job.StreamId);
+			}
+			if (!streamConfig.Authorize(AclAction.ViewJob, User))
+			{
+				return Forbid(AclAction.ViewJob, job.StreamId);
 			}
 
-			bool includeCosts = await _jobService.AuthorizeAsync(job, AclAction.ViewCosts, User, cache);
+			bool includeCosts = streamConfig.Authorize(AclAction.ViewCosts, User);
 
 			List<object> responses = new List<object>();
 			foreach (IJobStepBatch batch in job.Batches)
@@ -1087,6 +1083,10 @@ namespace Horde.Build.Jobs
 			{
 				return NotFound(jobId);
 			}
+			if (!_globalConfig.Value.TryGetStream(job.StreamId, out StreamConfig? streamConfig))
+			{
+				return NotFound(job.StreamId);
+			}
 
 			IJobStepBatch? batch = job.Batches.FirstOrDefault(x => x.Id == batchId);
 			if (batch == null)
@@ -1098,7 +1098,7 @@ namespace Horde.Build.Jobs
 				return Forbid("Missing session claim for job {JobId} batch {BatchId}", jobId, batchId);
 			}
 
-			IJob? newJob = await _jobService.UpdateBatchAsync(job, batchId, request.LogId?.ToObjectId<ILogFile>(), request.State);
+			IJob? newJob = await _jobService.UpdateBatchAsync(job, batchId, streamConfig, request.LogId?.ToObjectId<ILogFile>(), request.State);
 			if (newJob == null)
 			{
 				return NotFound(jobId);
@@ -1124,14 +1124,17 @@ namespace Horde.Build.Jobs
 				return NotFound(jobId);
 			}
 
-			StreamPermissionsCache cache = new StreamPermissionsCache();
-			if (!await _jobService.AuthorizeAsync(job, AclAction.ViewJob, User, cache))
+			StreamConfig? streamConfig;
+			if (!_globalConfig.Value.TryGetStream(job.StreamId, out streamConfig))
 			{
-				return Forbid(AclAction.ViewJob, jobId);
+				return NotFound(job.StreamId);
+			}
+			if (!streamConfig.Authorize(AclAction.ViewJob, User))
+			{
+				return Forbid(AclAction.ViewJob, job.StreamId);
 			}
 
-			bool includeCosts = await _jobService.AuthorizeAsync(job, AclAction.ViewCosts, User, cache);
-
+			bool includeCosts = streamConfig.Authorize(AclAction.ViewCosts, User);
 			foreach (IJobStepBatch batch in job.Batches)
 			{
 				if (batch.Id == batchId)
@@ -1161,7 +1164,13 @@ namespace Horde.Build.Jobs
 			{
 				return NotFound(jobId);
 			}
-			if (!await _jobService.AuthorizeAsync(job, AclAction.ViewJob, User, null))
+
+			StreamConfig? streamConfig;
+			if (!_globalConfig.Value.TryGetStream(job.StreamId, out streamConfig))
+			{
+				return NotFound(job.StreamId);
+			}
+			if (!streamConfig.Authorize(AclAction.ViewJob, User))
 			{
 				return Forbid(AclAction.ViewJob, jobId);
 			}
@@ -1200,6 +1209,12 @@ namespace Horde.Build.Jobs
 				return NotFound(jobId);
 			}
 
+			StreamConfig? streamConfig;
+			if (!_globalConfig.Value.TryGetStream(job.StreamId, out streamConfig))
+			{
+				return NotFound(job.StreamId);
+			}
+
 			// Check permissions for updating this step. Only the agent executing the step can modify the state of it.
 			if (request.State != JobStepState.Unspecified || request.Outcome != JobStepOutcome.Unspecified)
 			{
@@ -1213,16 +1228,17 @@ namespace Horde.Build.Jobs
 					return Forbid();
 				}
 			}
+
 			if (request.Retry != null || request.Priority != null)
 			{
-				if (!await _jobService.AuthorizeAsync(job, AclAction.RetryJobStep, User, null))
+				if (!streamConfig.Authorize(AclAction.RetryJobStep, User))
 				{
 					return Forbid(AclAction.RetryJobStep, jobId);
 				}
 			}
 			if (request.Properties != null)
 			{
-				if (!await _jobService.AuthorizeAsync(job, AclAction.UpdateJob, User, null))
+				if (!streamConfig.Authorize(AclAction.UpdateJob, User))
 				{
 					return Forbid(AclAction.UpdateJob, jobId);
 				}
@@ -1233,7 +1249,7 @@ namespace Horde.Build.Jobs
 
 			try
 			{
-				IJob? newJob = await _jobService.UpdateStepAsync(job, batchId, stepId, request.State, request.Outcome, null, request.AbortRequested, abortByUser, request.LogId?.ToObjectId<ILogFile>(), null, retryByUser, request.Priority, null, request.Properties);
+				IJob? newJob = await _jobService.UpdateStepAsync(job, batchId, stepId, streamConfig, request.State, request.Outcome, null, request.AbortRequested, abortByUser, request.LogId?.ToObjectId<ILogFile>(), null, retryByUser, request.Priority, null, request.Properties);
 				if (newJob == null)
 				{
 					return NotFound(jobId);
@@ -1305,7 +1321,13 @@ namespace Horde.Build.Jobs
 			{
 				return NotFound(jobId);
 			}
-			if (!await _jobService.AuthorizeAsync(job, AclAction.ViewJob, User, null))
+
+			StreamConfig? streamConfig;
+			if (!_globalConfig.Value.TryGetStream(job.StreamId, out streamConfig))
+			{
+				return NotFound(job.StreamId);
+			}
+			if (!streamConfig.Authorize(AclAction.ViewJob, User))
 			{
 				return Forbid(AclAction.ViewJob, jobId);
 			}
@@ -1346,7 +1368,13 @@ namespace Horde.Build.Jobs
 			{
 				return NotFound(jobId);
 			}
-			if (!await _jobService.AuthorizeAsync(job, AclAction.CreateSubscription, User, null))
+
+			StreamConfig? streamConfig;
+			if (!_globalConfig.Value.TryGetStream(job.StreamId, out streamConfig))
+			{
+				return NotFound(job.StreamId);
+			}
+			if (!streamConfig.Authorize(AclAction.CreateSubscription, User))
 			{
 				return Forbid(AclAction.CreateSubscription, jobId);
 			}
@@ -1365,7 +1393,7 @@ namespace Horde.Build.Jobs
 			if (triggerId == null)
 			{
 				triggerId = ObjectId.GenerateNewId();
-				if (await _jobService.UpdateStepAsync(job, batchId, stepId, JobStepState.Unspecified, JobStepOutcome.Unspecified, newNotificationTriggerId: triggerId) == null)
+				if (await _jobService.UpdateStepAsync(job, batchId, stepId, streamConfig, JobStepState.Unspecified, JobStepOutcome.Unspecified, newNotificationTriggerId: triggerId) == null)
 				{
 					return NotFound(jobId, batchId, stepId);
 				}
@@ -1401,7 +1429,13 @@ namespace Horde.Build.Jobs
 			{
 				return NotFound(jobId, batchId, stepId);
 			}
-			if (!await _jobService.AuthorizeAsync(job, AclAction.CreateSubscription, User, null))
+
+			StreamConfig? streamConfig;
+			if (!_globalConfig.Value.TryGetStream(job.StreamId, out streamConfig))
+			{
+				return NotFound(job.StreamId);
+			}
+			if (!streamConfig.Authorize(AclAction.CreateSubscription, User))
 			{
 				return Forbid(AclAction.CreateSubscription, jobId);
 			}
@@ -1435,10 +1469,17 @@ namespace Horde.Build.Jobs
 			{
 				return NotFound(jobId);
 			}
-			if (!await _jobService.AuthorizeAsync(job, AclAction.ViewJob, User, null))
+
+			StreamConfig? streamConfig;
+			if (!_globalConfig.Value.TryGetStream(job.StreamId, out streamConfig))
+			{
+				return NotFound(job.StreamId);
+			}
+			if (!streamConfig.Authorize(AclAction.ViewJob, User))
 			{
 				return Forbid(AclAction.ViewJob, jobId);
 			}
+
 			if (!job.TryGetBatch(batchId, out IJobStepBatch? batch))
 			{
 				return NotFound(jobId, batchId);
@@ -1474,10 +1515,17 @@ namespace Horde.Build.Jobs
 			{
 				return NotFound(jobId);
 			}
-			if (!await _jobService.AuthorizeAsync(job, AclAction.ViewJob, User, null))
+
+			StreamConfig? streamConfig;
+			if (!_globalConfig.Value.TryGetStream(job.StreamId, out streamConfig))
+			{
+				return NotFound(job.StreamId);
+			}
+			if (!streamConfig.Authorize(AclAction.ViewJob, User))
 			{
 				return Forbid(AclAction.ViewJob, jobId);
 			}
+
 			if (!job.TryGetBatch(batchId, out IJobStepBatch? batch))
 			{
 				return NotFound(jobId, batchId);
@@ -1516,7 +1564,13 @@ namespace Horde.Build.Jobs
 				{
 					return NotFound(jobId);
 				}
-				if (!await _jobService.AuthorizeAsync(job, AclAction.CreateSubscription, User, null))
+
+				StreamConfig? streamConfig;
+				if (!_globalConfig.Value.TryGetStream(job.StreamId, out streamConfig))
+				{
+					return NotFound(job.StreamId);
+				}
+				if (!streamConfig.Authorize(AclAction.CreateSubscription, User))
 				{
 					return Forbid(AclAction.CreateSubscription, jobId);
 				}
diff --git a/Engine/Source/Programs/Horde/Horde.Build/Jobs/Schedules/ScheduleService.cs b/Engine/Source/Programs/Horde/Horde.Build/Jobs/Schedules/ScheduleService.cs
--- a/Engine/Source/Programs/Horde/Horde.Build/Jobs/Schedules/ScheduleService.cs
+++ b/Engine/Source/Programs/Horde/Horde.Build/Jobs/Schedules/ScheduleService.cs
@@ -25,7 +25,6 @@ using MongoDB.Driver;
 using OpenTracing;
 using OpenTracing.Util;
 using StackExchange.Redis;
-using TimeZoneConverter;
 
 namespace Horde.Build.Jobs.Schedules
 {
@@ -65,7 +64,7 @@ namespace Horde.Build.Jobs.Schedules
 		readonly ICommitService _commitService;
 		readonly IJobCollection _jobCollection;
 		readonly JobService _jobService;
-		readonly StreamService _streamService;
+		readonly IStreamCollection _streamCollection;
 		readonly ITemplateCollection _templateCollection;
 		readonly IClock _clock;
 		readonly ILogger _logger;
@@ -73,20 +72,22 @@ namespace Horde.Build.Jobs.Schedules
 		static readonly RedisKey s_baseLockKey = "scheduler/locks";
 		static readonly RedisKey s_tickLockKey = s_baseLockKey.Append("/tick"); // Lock to tick the queue
 		static readonly RedisSortedSetKey<QueueItem> s_queueKey = "scheduler/queue"; // Items to tick, ordered by time
+		readonly IOptionsMonitor<GlobalConfig> _globalConfig;
 		readonly ITicker _ticker;
 
 		/// <summary>
 		/// Constructor
 		/// </summary>
-		public ScheduleService(RedisService redis, IGraphCollection graphs, ICommitService commitService, IJobCollection jobCollection, JobService jobService, StreamService streamService, ITemplateCollection templateCollection, MongoService mongoService, IClock clock, IOptionsMonitor<ServerSettings> settings, ILogger<ScheduleService> logger)
+		public ScheduleService(RedisService redis, IGraphCollection graphs, ICommitService commitService, IJobCollection jobCollection, JobService jobService, IStreamCollection streamCollection, ITemplateCollection templateCollection, MongoService mongoService, IClock clock, IOptionsMonitor<GlobalConfig> globalConfig, ILogger<ScheduleService> logger)
 		{
 			_graphs = graphs;
 			_commitService = commitService;
 			_jobCollection = jobCollection;
 			_jobService = jobService;
-			_streamService = streamService;
+			_streamCollection = streamCollection;
 			_templateCollection = templateCollection;
 			_clock = clock;
+			_globalConfig = globalConfig;
 			_logger = logger;
 
 			_redis = redis;
@@ -200,9 +201,10 @@ namespace Horde.Build.Jobs.Schedules
 		{
 			List<SortedSetEntry<QueueItem>> queueItems = new List<SortedSetEntry<QueueItem>>();
 
-			List<IStream> streams = await _streamService.GetStreamsAsync();
-			foreach (IStream stream in streams)
+			GlobalConfig globalConfig = _globalConfig.CurrentValue;
+			foreach(StreamConfig streamConfig in globalConfig.Streams)
 			{
+				IStream stream = await _streamCollection.GetAsync(streamConfig);
 				foreach ((TemplateId templateId, ITemplateRef templateRef) in stream.Templates)
 				{
 					if (templateRef.Schedule != null)
@@ -215,7 +217,7 @@ namespace Horde.Build.Jobs.Schedules
 								double score = QueueItem.GetScoreFromTime(nextTriggerTimeUtc.Value);
 								queueItems.Add(new SortedSetEntry<QueueItem>(new QueueItem(stream.Id, templateId), score));
 
-								await _streamService.UpdateScheduleTriggerAsync(stream, templateId, utcNow);
+								await _streamCollection.UpdateScheduleTriggerAsync(stream, templateId, utcNow);
 							}
 						}
 					}
@@ -235,7 +237,13 @@ namespace Horde.Build.Jobs.Schedules
 		/// <returns>Async task</returns>
 		internal async Task<bool> TriggerAsync(StreamId streamId, TemplateId templateId, DateTime utcNow, CancellationToken cancellationToken)
 		{
-			IStream? stream = await _streamService.GetStreamAsync(streamId);
+			GlobalConfig globalConfig = _globalConfig.CurrentValue;
+			if (!globalConfig.TryGetStream(streamId, out StreamConfig? streamConfig))
+			{
+				return false;
+			}
+
+			IStream? stream = await _streamCollection.GetAsync(streamConfig);
 			if (stream == null || !stream.Templates.TryGetValue(templateId, out ITemplateRef? templateRef))
 			{
 				return false;
@@ -265,7 +273,7 @@ namespace Horde.Build.Jobs.Schedules
 					removeJobIds.Add(activeJobId);
 				}
 			}
-			await _streamService.UpdateScheduleTriggerAsync(stream, templateId, removeJobs: removeJobIds);
+			await _streamCollection.UpdateScheduleTriggerAsync(stream, templateId, removeJobs: removeJobIds);
 
 			// If the stream is paused, bail out
 			if (stream.IsPaused(utcNow))
@@ -346,7 +354,7 @@ namespace Horde.Build.Jobs.Schedules
 			}
 
 			// Cache the Perforce history as we're iterating through changes to improve query performance
-			ICommitCollection commits = _commitService.GetCollection(stream);
+			ICommitCollection commits = _commitService.GetCollection(stream.Config);
 			IAsyncEnumerable<ICommit> commitEnumerable = commits.FindAsync(minChangeNumber, null, null, schedule.Config.Commits, cancellationToken);
 			await using IAsyncEnumerator<ICommit> commitEnumerator = commitEnumerable.GetAsyncEnumerator(cancellationToken);
 
@@ -362,7 +370,7 @@ namespace Horde.Build.Jobs.Schedules
 
 				if (schedule.Config.Gate != null)
 				{
-					change = await GetNextChangeForGateAsync(stream, templateId, schedule.Config.Gate, minChangeNumber, maxChangeNumber, cancellationToken);
+					change = await GetNextChangeForGateAsync(stream.Id, templateId, schedule.Config.Gate, minChangeNumber, maxChangeNumber, cancellationToken);
 					commit = await commits.FindAsync(change, change, 1, null, cancellationToken).FirstOrDefaultAsync(cancellationToken); // May be a change in a different stream
 				}
 				else if (await commitEnumerator.MoveNextAsync(cancellationToken))
@@ -431,12 +439,7 @@ namespace Horde.Build.Jobs.Schedules
 			}
 
 			// Get the matching template
-			ITemplate? template = await _templateCollection.GetAsync(templateRef.Hash);
-			if (template == null)
-			{
-				_logger.LogWarning("Unable to find template '{TemplateHash}' for '{TemplateId}'", templateRef.Hash, templateId);
-				return;
-			}
+			ITemplate template = await _templateCollection.GetOrAddAsync(templateRef.Config);
 
 			// Register the graph for it
 			IGraph graph = await _graphs.AddAsync(template, stream.Config.InitialAgentType);
@@ -460,9 +463,9 @@ namespace Horde.Build.Jobs.Schedules
 				options.Priority = template.Priority;
 				options.Arguments.AddRange(template.GetDefaultArguments());
 
-				IJob newJob = await _jobService.CreateJobAsync(null, stream, templateId, template.Id, graph, template.Name, change, codeChange, options);
+				IJob newJob = await _jobService.CreateJobAsync(null, stream.Config, templateId, template.Hash, graph, template.Name, change, codeChange, options);
 				_logger.LogInformation("Started new job for {StreamId} template {TemplateId} at CL {Change} (Code CL {CodeChange}): {JobId}", stream.Id, templateId, change, codeChange, newJob.Id);
-				await _streamService.UpdateScheduleTriggerAsync(stream, templateId, utcNow, change, new List<JobId> { newJob.Id }, new List<JobId>());
+				await _streamCollection.UpdateScheduleTriggerAsync(stream, templateId, utcNow, change, new List<JobId> { newJob.Id }, new List<JobId>());
 			}
 		}
 
@@ -504,13 +507,13 @@ namespace Horde.Build.Jobs.Schedules
 		/// Gets the next change to build for a schedule on a gate
 		/// </summary>
 		/// <returns></returns>
-		private async Task<int> GetNextChangeForGateAsync(IStream stream, TemplateId templateRefId, ScheduleGateConfig gate, int? minChange, int? maxChange, CancellationToken cancellationToken)
+		private async Task<int> GetNextChangeForGateAsync(StreamId streamId, TemplateId templateRefId, ScheduleGateConfig gate, int? minChange, int? maxChange, CancellationToken cancellationToken)
 		{
 			for (; ; )
 			{
 				cancellationToken.ThrowIfCancellationRequested();
 
-				List<IJob> jobs = await _jobCollection.FindAsync(streamId: stream.Id, templates: new[] { gate.TemplateId }, minChange: minChange, maxChange: maxChange, count: 1);
+				List<IJob> jobs = await _jobCollection.FindAsync(streamId: streamId, templates: new[] { gate.TemplateId }, minChange: minChange, maxChange: maxChange, count: 1);
 				if (jobs.Count == 0)
 				{
 					return 0;
@@ -529,7 +532,7 @@ namespace Horde.Build.Jobs.Schedules
 						{
 							return job.Change;
 						}
-						_logger.LogInformation("Skipping trigger of {StreamName} template {TemplateId} - last {OtherTemplateRefId} job ({JobId}) ended with errors", stream.Id, templateRefId, gate.TemplateId, job.Id);
+						_logger.LogInformation("Skipping trigger of {StreamName} template {TemplateId} - last {OtherTemplateRefId} job ({JobId}) ended with errors", streamId, templateRefId, gate.TemplateId, job.Id);
 					}
 				}
 
diff --git a/Engine/Source/Programs/Horde/Horde.Build/Jobs/Templates/ITemplate.cs b/Engine/Source/Programs/Horde/Horde.Build/Jobs/Templates/ITemplate.cs
--- a/Engine/Source/Programs/Horde/Horde.Build/Jobs/Templates/ITemplate.cs
+++ b/Engine/Source/Programs/Horde/Horde.Build/Jobs/Templates/ITemplate.cs
@@ -425,7 +425,7 @@ namespace Horde.Build.Jobs.Templates
 		/// <summary>
 		/// Hash of this template
 		/// </summary>
-		public ContentHash Id { get; }
+		public ContentHash Hash { get; }
 
 		/// <summary>
 		/// Name of the template.
diff --git a/Engine/Source/Programs/Horde/Horde.Build/Jobs/Templates/ITemplateCollection.cs b/Engine/Source/Programs/Horde/Horde.Build/Jobs/Templates/ITemplateCollection.cs
--- a/Engine/Source/Programs/Horde/Horde.Build/Jobs/Templates/ITemplateCollection.cs
+++ b/Engine/Source/Programs/Horde/Horde.Build/Jobs/Templates/ITemplateCollection.cs
@@ -3,7 +3,6 @@
 using System.Collections.Generic;
 using System.Threading.Tasks;
 using EpicGames.Core;
-using HordeCommon;
 
 namespace Horde.Build.Jobs.Templates
 {
@@ -12,32 +11,18 @@ namespace Horde.Build.Jobs.Templates
 	/// </summary>
 	public interface ITemplateCollection
 	{
-		/// <summary>
-		/// Public constructor
-		/// </summary>
-		/// <param name="name">Name of the template</param>
-		/// <param name="priority">Priority of this template</param>
-		/// <param name="allowPreflights">Whether to allow preflights of this job</param>
-		/// <param name="updateIssues"> Whether to update issues for all jobs using this template</param>
-		/// <param name="promoteIssuesByDefault">Whether to promote issues by default for all jobs using this template</param>
-		/// <param name="initialAgentType">The agent type to parse the buildgraph script</param>
-		/// <param name="submitNewChange">Path to a file within the stream to submit to generate a new changelist for jobs</param>
-		/// <param name="submitDescription">Description for new changes submitted to the stream</param>
-		/// <param name="arguments">List of arguments which are always specified</param>
-		/// <param name="parameters">List of template parameters</param>
-		Task<ITemplate> AddAsync(string name, Priority? priority = null, bool allowPreflights = true, bool updateIssues = false, bool promoteIssuesByDefault = false, string? initialAgentType = null, string? submitNewChange = null, string? submitDescription = null, List<string>? arguments = null, List<Parameter>? parameters = null);
-
-		/// <summary>
-		/// Gets all the available templates
-		/// </summary>
-		/// <returns>List of template documents</returns>
-		Task<List<ITemplate>> FindAllAsync();
-
 		/// <summary>
 		/// Gets a template by ID
 		/// </summary>
 		/// <param name="templateId">Unique id of the template</param>
 		/// <returns>The template document</returns>
 		Task<ITemplate?> GetAsync(ContentHash templateId);
+
+		/// <summary>
+		/// Gets a template instance from its configuration
+		/// </summary>
+		/// <param name="templateConfig">The template configuration</param>
+		/// <returns>Template instance</returns>
+		Task<ITemplate> GetOrAddAsync(TemplateConfig templateConfig);
 	}
 }
diff --git a/Engine/Source/Programs/Horde/Horde.Build/Jobs/Templates/TemplateCollection.cs b/Engine/Source/Programs/Horde/Horde.Build/Jobs/Templates/TemplateCollection.cs
--- a/Engine/Source/Programs/Horde/Horde.Build/Jobs/Templates/TemplateCollection.cs
+++ b/Engine/Source/Programs/Horde/Horde.Build/Jobs/Templates/TemplateCollection.cs
@@ -48,6 +48,7 @@ namespace Horde.Build.Jobs.Templates
 			public List<string> Arguments { get; private set; } = new List<string>();
 			public List<Parameter> Parameters { get; private set; } = new List<Parameter>();
 
+			ContentHash ITemplate.Hash => Id;
 			IReadOnlyList<string> ITemplate.Arguments => Arguments;
 			IReadOnlyList<Parameter> ITemplate.Parameters => Parameters;
 
@@ -57,19 +58,19 @@ namespace Horde.Build.Jobs.Templates
 				Name = null!;
 			}
 
-			public TemplateDocument(string name, Priority? priority, bool allowPreflights, bool updateIssues, bool promoteIssuesByDefault, string? initialAgentType, string? submitNewChange, string? submitDescription, List<string>? arguments, List<Parameter>? parameters)
+			public TemplateDocument(TemplateConfig config)
 			{
-				Name = name;
-				Priority = priority;
-				AllowPreflights = allowPreflights;
-				UpdateIssues = updateIssues;
-				PromoteIssuesByDefault = promoteIssuesByDefault;
-				InitialAgentType = initialAgentType;
-				SubmitNewChange = submitNewChange;
-				SubmitDescription = submitDescription;
-				Arguments = arguments ?? new List<string>();
-				Parameters = parameters ?? new List<Parameter>();
-
+				Name = config.Name;
+				Priority = config.Priority;
+				AllowPreflights = config.AllowPreflights;
+				UpdateIssues = config.UpdateIssues;
+				PromoteIssuesByDefault = config.PromoteIssuesByDefault;
+				InitialAgentType = config.InitialAgentType;
+				SubmitNewChange = config.SubmitNewChange;
+				SubmitDescription = config.SubmitDescription;
+				Arguments = config.Arguments ?? new List<string>();
+				Parameters = config.Parameters.ConvertAll(x => x.ToModel());
+			
 				// Compute the hash once all other fields have been set
 				Id = ContentHash.SHA1(BsonExtensionMethods.ToBson(this));
 			}
@@ -103,24 +104,6 @@ namespace Horde.Build.Jobs.Templates
 			_templateCache.Dispose();
 		}
 
-		/// <inheritdoc/>
-		public async Task<ITemplate> AddAsync(string name, Priority? priority, bool allowPreflights, bool updateIssues, bool promoteIssuesByDefault, string? initialAgentType, string? submitNewChange, string? submitDescription, List<string>? arguments, List<Parameter>? parameters)
-		{
-			TemplateDocument template = new TemplateDocument(name, priority, allowPreflights, updateIssues, promoteIssuesByDefault, initialAgentType, submitNewChange, submitDescription, arguments, parameters);
-			if (await GetAsync(template.Id) == null)
-			{
-				await _templates.ReplaceOneAsync(x => x.Id == template.Id, template, new ReplaceOptions { IsUpsert = true });
-			}
-			return template;
-		}
-
-		/// <inheritdoc/>
-		public async Task<List<ITemplate>> FindAllAsync()
-		{
-			List<TemplateDocument> results = await _templates.Find(FilterDefinition<TemplateDocument>.Empty).ToListAsync();
-			return results.ConvertAll<ITemplate>(x => x);
-		}
-
 		/// <inheritdoc/>
 		public async Task<ITemplate?> GetAsync(ContentHash templateId)
 		{
@@ -142,5 +125,25 @@ namespace Horde.Build.Jobs.Templates
 
 			return template;
 		}
+
+		/// <inheritdoc/>
+		public async Task<ITemplate> GetOrAddAsync(TemplateConfig config)
+		{
+			if (config.CachedHash != null)
+			{
+				ITemplate? existingTemplate = await GetAsync(config.CachedHash);
+				if (existingTemplate != null)
+				{
+					return existingTemplate;
+				}
+			}
+
+			TemplateDocument template = new TemplateDocument(config);
+			if (await GetAsync(template.Id) == null)
+			{
+				await _templates.ReplaceOneAsync(x => x.Id == template.Id, template, new ReplaceOptions { IsUpsert = true });
+			}
+			return template;
+		}
 	}
 }
diff --git a/Engine/Source/Programs/Horde/Horde.Build/Jobs/Templates/TemplateMessages.cs b/Engine/Source/Programs/Horde/Horde.Build/Jobs/Templates/TemplateMessages.cs
--- a/Engine/Source/Programs/Horde/Horde.Build/Jobs/Templates/TemplateMessages.cs
+++ b/Engine/Source/Programs/Horde/Horde.Build/Jobs/Templates/TemplateMessages.cs
@@ -3,6 +3,7 @@
 using System;
 using System.Collections.Generic;
 using System.ComponentModel.DataAnnotations;
+using System.Text.Json.Serialization;
 using EpicGames.Core;
 using Horde.Build.Streams;
 using HordeCommon;
@@ -450,6 +451,12 @@ namespace Horde.Build.Jobs.Templates
 		/// The executor to use for this job. Mainly used for debugging.
 		/// </summary>
 		public string? Executor { get; set; }
+
+		/// <summary>
+		/// The cached hash of this template.
+		/// </summary>
+		[JsonIgnore]
+		internal ContentHash? CachedHash { get; set; }
 	}
 
 	/// <summary>
@@ -551,7 +558,7 @@ namespace Horde.Build.Jobs.Templates
 		public GetTemplateResponse(ITemplate template)
 			: base(template)
 		{
-			Id = template.Id.ToString();
+			Id = template.Hash.ToString();
 		}
 	}
 }
diff --git a/Engine/Source/Programs/Horde/Horde.Build/Jobs/TestData/TestDataController.cs b/Engine/Source/Programs/Horde/Horde.Build/Jobs/TestData/TestDataController.cs
--- a/Engine/Source/Programs/Horde/Horde.Build/Jobs/TestData/TestDataController.cs
+++ b/Engine/Source/Programs/Horde/Horde.Build/Jobs/TestData/TestDataController.cs
@@ -15,6 +15,8 @@ using System.Text;
 using System.Text.Json.Serialization;
 using System.Text.Json;
 using System.Threading;
+using Microsoft.Extensions.Options;
+using Horde.Build.Server;
 
 namespace Horde.Build.Jobs.TestData
 {
@@ -46,21 +48,20 @@ namespace Horde.Build.Jobs.TestData
 
 		readonly TestDataService _testDataService;
 
-		readonly StreamService _streamService;
+		readonly IStreamCollection _streamCollection;
+	
+		readonly IOptionsSnapshot<GlobalConfig> _globalConfig;
 
 		/// <summary>
 		/// Constructor
 		/// </summary>
-		/// <param name="testDataService"></param>
-		/// <param name="streamService"></param>
-		/// <param name="jobService"></param>
-		/// <param name="testDataCollection"></param>
-		public TestDataController(TestDataService testDataService, StreamService streamService, JobService jobService, ITestDataCollection testDataCollection)
+		public TestDataController(TestDataService testDataService, IStreamCollection streamCollection, JobService jobService, ITestDataCollection testDataCollection, IOptionsSnapshot<GlobalConfig> globalConfig)
 		{
 			_jobService = jobService;
 			_testDataCollection = testDataCollection;
 			_testDataService = testDataService;
-			_streamService = streamService;
+			_streamCollection = streamCollection;
+			_globalConfig = globalConfig;
 		}
 
 		/// <summary>
@@ -136,13 +137,9 @@ namespace Horde.Build.Jobs.TestData
 			// authorize streams
 			foreach (StreamId streamId in streamIdValues)
 			{
-				IStream? stream = await _streamService.GetCachedStream(streamId);
-				if (stream != null)
+				if (_globalConfig.Value.TryGetStream(streamId, out StreamConfig? streamConfig) && streamConfig.Authorize(AclAction.ViewJob, User))
 				{
-					if (await _streamService.AuthorizeAsync(stream, AclAction.ViewJob, User, null))
-					{
-						queryStreams.Add(stream.Id);
-					}
+					queryStreams.Add(streamId);
 				}
 			}
 
@@ -280,13 +277,9 @@ namespace Horde.Build.Jobs.TestData
 			// authorize streams
 			foreach (StreamId streamId in streamIdValues)
 			{
-				IStream? stream = await _streamService.GetCachedStream(streamId);
-				if (stream != null)
+				if (_globalConfig.Value.TryGetStream(streamId, out StreamConfig? streamConfig) && streamConfig.Authorize(AclAction.ViewJob, User))
 				{
-					if (await _streamService.AuthorizeAsync(stream, AclAction.ViewJob, User, null))
-					{
-						queryStreams.Add(stream.Id);
-					}
+					queryStreams.Add(streamId);
 				}
 			}
 
@@ -315,7 +308,7 @@ namespace Horde.Build.Jobs.TestData
 			{
 				return NotFound();
 			}
-			if (!await _jobService.AuthorizeAsync(job, AclAction.UpdateJob, User, null))
+			if (!_globalConfig.Value.Authorize(job, AclAction.UpdateJob, User))
 			{
 				return Forbid();
 			}
@@ -354,14 +347,12 @@ namespace Horde.Build.Jobs.TestData
 				streamIdValue = new StreamId(streamId);
 			}
 
-			JobPermissionsCache cache = new JobPermissionsCache();
-
 			List<object> results = new List<object>();
 
 			List<ITestData> documents = await _testDataCollection.FindAsync(streamIdValue, minChange, maxChange, jobId?.ToObjectId<IJob>(), jobStepId?.ToSubResourceId(), key, index, count);
 			foreach (ITestData document in documents)
 			{
-				if (await _jobService.AuthorizeAsync(document.JobId, AclAction.ViewJob, User, cache))
+				if (await _jobService.AuthorizeAsync(document.JobId, AclAction.ViewJob, User, _globalConfig.Value))
 				{
 					results.Add(PropertyFilter.Apply(new GetTestDataResponse(document), filter));
 				}
@@ -386,7 +377,7 @@ namespace Horde.Build.Jobs.TestData
 			{
 				return NotFound();
 			}
-			if (!await _jobService.AuthorizeAsync(testData.JobId, AclAction.ViewJob, User, null))
+			if (!await _jobService.AuthorizeAsync(testData.JobId, AclAction.ViewJob, User, _globalConfig.Value))
 			{
 				return Forbid();
 			}
diff --git a/Engine/Source/Programs/Horde/Horde.Build/Logs/LogsController.cs b/Engine/Source/Programs/Horde/Horde.Build/Logs/LogsController.cs
--- a/Engine/Source/Programs/Horde/Horde.Build/Logs/LogsController.cs
+++ b/Engine/Source/Programs/Horde/Horde.Build/Logs/LogsController.cs
@@ -15,6 +15,7 @@ using Horde.Build.Acls;
 using Horde.Build.Issues;
 using Horde.Build.Jobs;
 using Horde.Build.Logs.Data;
+using Horde.Build.Server;
 using Horde.Build.Storage;
 using Horde.Build.Utilities;
 using Microsoft.AspNetCore.Authorization;
@@ -57,19 +58,19 @@ namespace Horde.Build.Logs
 		private readonly AclService _aclService;
 		private readonly JobService _jobService;
 		private readonly StorageService _storageService;
-		private readonly IOptions<ServerSettings> _settings;
+		private readonly IOptionsSnapshot<GlobalConfig> _globalConfig;
 
 		/// <summary>
 		/// Constructor
 		/// </summary>
-		public LogsController(ILogFileService logFileService, IIssueCollection issueCollection, AclService aclService, JobService jobService, StorageService storageService, IOptions<ServerSettings> settings)
+		public LogsController(ILogFileService logFileService, IIssueCollection issueCollection, AclService aclService, JobService jobService, StorageService storageService, IOptionsSnapshot<GlobalConfig> globalConfig)
 		{
 			_logFileService = logFileService;
 			_issueCollection = issueCollection;
 			_aclService = aclService;
 			_jobService = jobService;
 			_storageService = storageService;
-			_settings = settings;
+			_globalConfig = globalConfig;
  		}
 
 		/// <summary>
@@ -84,7 +85,7 @@ namespace Horde.Build.Logs
 		[ProducesResponseType(typeof(GetLogFileResponse), 200)]
 		public async Task<ActionResult<object>> CreateLog(CreateLogFileRequest request, CancellationToken cancellationToken = default)
 		{
-			if (!await _aclService.AuthorizeAsync(AclAction.CreateLog, User, null))
+			if (!_globalConfig.Value.Authorize(AclAction.CreateLog, User))
 			{
 				return Forbid();
 			}
@@ -110,7 +111,7 @@ namespace Horde.Build.Logs
 			{
 				return NotFound();
 			}
-			if (!await AuthorizeAsync(logFile, AclAction.ViewLog, User, null))
+			if (!await AuthorizeAsync(logFile, AclAction.ViewLog, User))
 			{
 				return Forbid();
 			}
@@ -136,7 +137,7 @@ namespace Horde.Build.Logs
 			{
 				return NotFound();
 			}
-			if (!await AuthorizeAsync(logFile, AclAction.WriteLogData, User, null))
+			if (!await AuthorizeAsync(logFile, AclAction.WriteLogData, User))
 			{
 				return Forbid();
 			}
@@ -167,7 +168,7 @@ namespace Horde.Build.Logs
 			{
 				return NotFound();
 			}
-			if (!await AuthorizeAsync(logFile, AclAction.ViewLog, User, null))
+			if (!await AuthorizeAsync(logFile, AclAction.ViewLog, User))
 			{
 				return Forbid();
 			}
@@ -202,7 +203,7 @@ namespace Horde.Build.Logs
 			{
 				return NotFound();
 			}
-			if (!await AuthorizeAsync(logFile, AclAction.ViewLog, User, null))
+			if (!await AuthorizeAsync(logFile, AclAction.ViewLog, User))
 			{
 				return Forbid();
 			}
@@ -297,7 +298,7 @@ namespace Horde.Build.Logs
 			{
 				return NotFound();
 			}
-			if (!await AuthorizeAsync(logFile, AclAction.ViewLog, User, null))
+			if (!await AuthorizeAsync(logFile, AclAction.ViewLog, User))
 			{
 				return Forbid();
 			}
@@ -326,7 +327,7 @@ namespace Horde.Build.Logs
 			{
 				return NotFound();
 			}
-			if (!await AuthorizeAsync(logFile, AclAction.ViewLog, User, null))
+			if (!await AuthorizeAsync(logFile, AclAction.ViewLog, User))
 			{
 				return Forbid();
 			}
@@ -370,7 +371,7 @@ namespace Horde.Build.Logs
 			{
 				return NotFound();
 			}
-			if (!await AuthorizeAsync(logFile, AclAction.WriteLogData, User, null))
+			if (!await AuthorizeAsync(logFile, AclAction.WriteLogData, User))
 			{
 				return Forbid();
 			}
@@ -389,19 +390,19 @@ namespace Horde.Build.Logs
 		/// <param name="logFile">The template to check</param>
 		/// <param name="action">The action being performed</param>
 		/// <param name="user">The principal to authorize</param>
-		/// <param name="permissionsCache">Permissions cache</param>
 		/// <returns>True if the action is authorized</returns>
-		async Task<bool> AuthorizeAsync(ILogFile logFile, AclAction action, ClaimsPrincipal user, JobPermissionsCache? permissionsCache)
+		async Task<bool> AuthorizeAsync(ILogFile logFile, AclAction action, ClaimsPrincipal user)
 		{
-			if (logFile.JobId != JobId.Empty && await _jobService.AuthorizeAsync(logFile.JobId, action, user, permissionsCache))
+			GlobalConfig globalConfig = _globalConfig.Value;
+			if (logFile.JobId != JobId.Empty && await _jobService.AuthorizeAsync(logFile.JobId, action, user, globalConfig))
 			{
 				return true;
 			}
-			if (action == AclAction.ViewLog && logFile.SessionId != null && await _aclService.AuthorizeAsync(AclAction.ViewSession, user, permissionsCache))
+			if (action == AclAction.ViewLog && logFile.SessionId != null && globalConfig.Authorize(AclAction.ViewSession, user))
 			{
 				return true;
 			}
-			if (await _aclService.AuthorizeAsync(action, user, permissionsCache))
+			if (globalConfig.Authorize(action, user))
 			{
 				return true;
 			}
diff --git a/Engine/Source/Programs/Horde/Horde.Build/Notifications/INotificationService.cs b/Engine/Source/Programs/Horde/Horde.Build/Notifications/INotificationService.cs
--- a/Engine/Source/Programs/Horde/Horde.Build/Notifications/INotificationService.cs
+++ b/Engine/Source/Programs/Horde/Horde.Build/Notifications/INotificationService.cs
@@ -143,12 +143,12 @@ namespace Horde.Build.Notifications
 		/// <param name="message">The message to send</param>
 		/// <param name="device">The device</param>
 		/// <param name="pool">The pool</param>
-		/// <param name="stream"></param>
+		/// <param name="streamConfig"></param>
 		/// <param name="job"></param>
 		/// <param name="step"></param>
 		/// <param name="node"></param>
 		/// <param name="user"></param>
-		void NotifyDeviceService(string message, IDevice? device = null, IDevicePool? pool = null, IStream? stream = null, IJob? job = null, IJobStep? step = null, INode? node = null, IUser? user = null);
+		void NotifyDeviceService(string message, IDevice? device = null, IDevicePool? pool = null, StreamConfig? streamConfig = null, IJob? job = null, IJobStep? step = null, INode? node = null, IUser? user = null);
 
 		/// <summary>
 		/// Post a notification for the open issues in a stream
diff --git a/Engine/Source/Programs/Horde/Horde.Build/Notifications/INotificationSink.cs b/Engine/Source/Programs/Horde/Horde.Build/Notifications/INotificationSink.cs
--- a/Engine/Source/Programs/Horde/Horde.Build/Notifications/INotificationSink.cs
+++ b/Engine/Source/Programs/Horde/Horde.Build/Notifications/INotificationSink.cs
@@ -30,49 +30,45 @@ namespace Horde.Build.Notifications
 		/// <summary>
 		/// Send notifications that a job has completed
 		/// </summary>
-		/// <param name="jobStream"></param>
 		/// <param name="job">The job containing the step</param>
 		/// <param name="graph"></param>
 		/// <param name="outcome"></param>
 		/// <returns>Async task</returns>
-		Task NotifyJobCompleteAsync(IStream jobStream, IJob job, IGraph graph, LabelOutcome outcome);
+		Task NotifyJobCompleteAsync(IJob job, IGraph graph, LabelOutcome outcome);
 
 		/// <summary>
 		/// Send notifications that a job has completed
 		/// </summary>
 		/// <param name="user">User to notify</param>
-		/// <param name="jobStream"></param>
 		/// <param name="job">The job containing the step</param>
 		/// <param name="graph"></param>
 		/// <param name="outcome"></param>
 		/// <returns>Async task</returns>
-		Task NotifyJobCompleteAsync(IUser user, IStream jobStream, IJob job, IGraph graph, LabelOutcome outcome);
+		Task NotifyJobCompleteAsync(IUser user, IJob job, IGraph graph, LabelOutcome outcome);
 
 		/// <summary>
 		/// Send notifications that a job step has completed
 		/// </summary>
 		/// <param name="user">User to notify</param>
-		/// <param name="jobStream">Stream containing the job</param>
 		/// <param name="job">The job containing the step</param>
 		/// <param name="batch">Unique id of the batch</param>
 		/// <param name="step">The step id</param>
 		/// <param name="node">Corresponding node for the step</param>
 		/// <param name="jobStepEventData"></param>
 		/// <returns>Async task</returns>
-		Task NotifyJobStepCompleteAsync(IUser user, IStream jobStream, IJob job, IJobStepBatch batch, IJobStep step, INode node, List<ILogEventData> jobStepEventData);
+		Task NotifyJobStepCompleteAsync(IUser user, IJob job, IJobStepBatch batch, IJobStep step, INode node, List<ILogEventData> jobStepEventData);
 
 		/// <summary>
 		/// Send notifications that a job step has completed
 		/// </summary>
 		/// <param name="user">User to notify</param>
 		/// <param name="job">The job containing the step</param>
-		/// <param name="stream"></param>
 		/// <param name="label"></param>
 		/// <param name="labelIdx"></param>
 		/// <param name="outcome"></param>
 		/// <param name="stepData"></param>
 		/// <returns>Async task</returns>
-		Task NotifyLabelCompleteAsync(IUser user, IJob job, IStream stream, ILabel label, int labelIdx, LabelOutcome outcome, List<(string, JobStepOutcome, Uri)> stepData);
+		Task NotifyLabelCompleteAsync(IUser user, IJob job, ILabel label, int labelIdx, LabelOutcome outcome, List<(string, JobStepOutcome, Uri)> stepData);
 
 		#region Issues
 
@@ -115,13 +111,13 @@ namespace Horde.Build.Notifications
 		/// <param name="message"></param>
 		/// <param name="device"></param>
 		/// <param name="pool"></param>
-		/// <param name="stream"></param>
+		/// <param name="streamConfig"></param>
 		/// <param name="job"></param>
 		/// <param name="step"></param>
 		/// <param name="node"></param>
 		/// <param name="user"></param>
 		/// <returns></returns>
-		Task NotifyDeviceServiceAsync(string message, IDevice? device = null, IDevicePool? pool = null, IStream? stream = null, IJob? job = null, IJobStep? step = null, INode? node = null, IUser? user = null);
+		Task NotifyDeviceServiceAsync(string message, IDevice? device = null, IDevicePool? pool = null, StreamConfig? streamConfig = null, IJob? job = null, IJobStep? step = null, INode? node = null, IUser? user = null);
     }
 }
 
diff --git a/Engine/Source/Programs/Horde/Horde.Build/Notifications/NotificationService.cs b/Engine/Source/Programs/Horde/Horde.Build/Notifications/NotificationService.cs
--- a/Engine/Source/Programs/Horde/Horde.Build/Notifications/NotificationService.cs
+++ b/Engine/Source/Programs/Horde/Horde.Build/Notifications/NotificationService.cs
@@ -10,6 +10,7 @@ using System.Threading.Tasks;
 using EpicGames.Core;
 using EpicGames.Redis;
 using Horde.Build.Agents.Pools;
+using Horde.Build.Configuration;
 using Horde.Build.Devices;
 using Horde.Build.Issues;
 using Horde.Build.Jobs;
@@ -70,9 +71,9 @@ namespace Horde.Build.Notifications
 		private readonly JobService _jobService;
 
 		/// <summary>
-		/// Instance of the <see cref="_streamService"/>.
+		/// Instance of the <see cref="IStreamCollection"/>.
 		/// </summary>
-		private readonly StreamService _streamService;
+		private readonly IStreamCollection _streamCollection;
 
 		/// <summary>
 		/// 
@@ -134,7 +135,9 @@ namespace Horde.Build.Notifications
 		/// Interval at which queued notifications should be sent as a batch 
 		/// </summary>
 		internal TimeSpan _notificationBatchInterval = TimeSpan.FromHours(12);
-		
+
+		readonly IOptionsMonitor<GlobalConfig> _globalConfig;
+
 		static string RedisQueueListKey(string notificationType) => "NotificationService.queued." + notificationType;
 
 		/// <summary>
@@ -149,12 +152,14 @@ namespace Horde.Build.Notifications
 			INotificationTriggerCollection triggerCollection,
 			IUserCollection userCollection,
 			JobService jobService,
-			StreamService streamService,
+			IStreamCollection streamCollection,
 			IssueService issueService,
 			ILogFileService logFileService,
 			IDogStatsd dogStatsd,
 			IMemoryCache cache,
 			RedisService redisService,
+			ConfigService configService,
+			IOptionsMonitor<GlobalConfig> globalConfig,
 			IClock clock)
 		{
 			_sinks = sinks.ToList();
@@ -165,18 +170,20 @@ namespace Horde.Build.Notifications
 			_triggerCollection = triggerCollection;
 			_userCollection = userCollection;
 			_jobService = jobService;
-			_streamService = streamService;
+			_streamCollection = streamCollection;
 			_issueService = issueService;
 			_logFileService = logFileService;
 			_dogStatsd = dogStatsd;
 			_cache = cache;
 			_redisConnectionPool = redisService.ConnectionPool;
+			_globalConfig = globalConfig;
 
 			issueService.OnIssueUpdated += NotifyIssueUpdated;
 			jobService.OnJobStepComplete += NotifyJobStepComplete;
 			jobService.OnJobScheduled += NotifyJobScheduled;
 			jobService.OnLabelUpdate += NotifyLabelUpdate;
-			
+			configService.OnConfigUpdate += NotifyConfigUpdate;
+
 			_ticker = clock.AddSharedTicker<NotificationService>(_notificationBatchInterval, TickEveryTwelveHoursAsync, logger);
 		}
 
@@ -290,9 +297,9 @@ namespace Horde.Build.Notifications
 		}
 
 		/// <inheritdoc/>
-		public void NotifyDeviceService(string message, IDevice? device = null, IDevicePool? pool = null, IStream? stream = null, IJob? job = null, IJobStep? step = null, INode? node = null, IUser? user = null)
+		public void NotifyDeviceService(string message, IDevice? device = null, IDevicePool? pool = null, StreamConfig? streamConfig = null, IJob? job = null, IJobStep? step = null, INode? node = null, IUser? user = null)
 		{
-			EnqueueTasks(sink => sink.NotifyDeviceServiceAsync(message, device, pool, stream, job, step, node, user));
+			EnqueueTasks(sink => sink.NotifyDeviceServiceAsync(message, device, pool, streamConfig, job, step, node, user));
 		}
 
 		/// <summary>
@@ -489,13 +496,6 @@ namespace Horde.Build.Notifications
 			job.GetJobState(job.GetStepForNodeMap(), out _, out LabelOutcome outcome);
 			JobCompleteEventRecord jobCompleteEvent = new JobCompleteEventRecord(job.StreamId, job.TemplateId, outcome);
 
-			IStream? jobStream = await _streamService.GetStreamAsync(job.StreamId);
-			if (jobStream == null)
-			{
-				_logger.LogError("Unable to get stream {StreamId}", job.StreamId);
-				return;
-			}
-
 			List<IUser> usersToNotify = await GetUsersToNotify(jobCompleteEvent, job.NotificationTriggerId, true);
 			foreach (IUser userToNotify in usersToNotify)
 			{
@@ -506,12 +506,12 @@ namespace Horde.Build.Notifications
 						continue;
 					}
 				}
-				EnqueueTasks(sink => sink.NotifyJobCompleteAsync(userToNotify, jobStream, job, graph, outcome));
+				EnqueueTasks(sink => sink.NotifyJobCompleteAsync(userToNotify, job, graph, outcome));
 			}
 
 			if (job.PreflightChange == 0)
 			{
-				EnqueueTasks(sink => sink.NotifyJobCompleteAsync(jobStream, job, graph, outcome));
+				EnqueueTasks(sink => sink.NotifyJobCompleteAsync(job, graph, outcome));
 			}
 
 			_logger.LogDebug("Finished sending notifications for job {JobId}", job.Id);
@@ -644,14 +644,7 @@ namespace Horde.Build.Notifications
 				return;
 			}
 
-			IStream? jobStream = await _streamService.GetStreamAsync(job.StreamId);
-			if (jobStream == null)
-			{
-				_logger.LogError("Unable to find stream {StreamId}", job.StreamId);
-				return;
-			}
-
-			if(step.LogId == null)
+			if (step.LogId == null)
 			{
 				_logger.LogError("Step does not have a log file");
 				return;
@@ -681,20 +674,13 @@ namespace Horde.Build.Notifications
 						continue;
 					}
 				}
-				EnqueueTasks(sink => sink.NotifyJobStepCompleteAsync(slackUser, jobStream, job, batch, step, node, jobStepEventData));
+				EnqueueTasks(sink => sink.NotifyJobStepCompleteAsync(slackUser, job, batch, step, node, jobStepEventData));
 			}
 			_logger.LogDebug("Finished sending notifications for step {JobId}:{BatchId}:{StepId}", job.Id, batchId, stepId);
 		}
 
 		private async Task SendAllLabelNotificationsAsync(IJob job, IReadOnlyList<(LabelState State, LabelOutcome Outcome)> oldLabelStates, IReadOnlyList<(LabelState, LabelOutcome)> newLabelStates)
 		{
-			IStream? stream = await _streamService.GetStreamAsync(job.StreamId);
-			if (stream == null)
-			{
-				_logger.LogError("Unable to find stream {StreamId} for job {JobId}", job.StreamId, job.Id);
-				return;
-			}
-
 			IGraph? graph = await _graphCollection.GetAsync(job.GraphHash);
 			if (graph == null)
 			{
@@ -755,7 +741,7 @@ namespace Horde.Build.Notifications
 
 					if (usersToNotify.Count > 0)
 					{
-						SendLabelUpdateNotifications(job, stream, graph, stepForNode, graph.Labels[labelIdx], labelIdx, newLabel.Outcome, usersToNotify);
+						SendLabelUpdateNotifications(job, graph, stepForNode, graph.Labels[labelIdx], labelIdx, newLabel.Outcome, usersToNotify);
 					}
 					else
 					{
@@ -765,7 +751,7 @@ namespace Horde.Build.Notifications
 			}
 		}
 
-		private void SendLabelUpdateNotifications(IJob job, IStream stream, IGraph graph, IReadOnlyDictionary<NodeRef, IJobStep> stepForNode, ILabel label, int labelIdx, LabelOutcome outcome, List<IUser> slackUsers)
+		private void SendLabelUpdateNotifications(IJob job, IGraph graph, IReadOnlyDictionary<NodeRef, IJobStep> stepForNode, ILabel label, int labelIdx, LabelOutcome outcome, List<IUser> slackUsers)
 		{
 			List<(string, JobStepOutcome, Uri)> stepData = new List<(string, JobStepOutcome, Uri)>();
 			if (outcome != LabelOutcome.Success)
@@ -780,7 +766,7 @@ namespace Horde.Build.Notifications
 
 			foreach (IUser slackUser in slackUsers)
 			{
-				EnqueueTasks(sink => sink.NotifyLabelCompleteAsync(slackUser, job, stream, label, labelIdx, outcome, stepData));
+				EnqueueTasks(sink => sink.NotifyLabelCompleteAsync(slackUser, job, label, labelIdx, outcome, stepData));
 			}
 
 			_logger.LogDebug("Finished sending label notifications for label {DashboardName}/{UgsName} in job {JobId}", label.DashboardName, label.UgsName, job.Id);
@@ -797,7 +783,7 @@ namespace Horde.Build.Notifications
 				}
 				catch (Exception e)
 				{
-					string streamsWithChannel = String.Join(", ", report.Reports.Select(x => x.Stream.Name + " " + x.TriageChannel));
+					string streamsWithChannel = String.Join(", ", report.Reports.Select(x => $"{x.StreamId} {x.TriageChannel}"));
 					_logger.LogError(e, "Failed sending issue report to {Channel} for streams/channels {StreamsWithChannels})", report.Channel, streamsWithChannel);
 				}
 			}
