Horde: Refactor codebase to use new config system.

diff --git a/Engine/Source/Programs/Horde/Horde.Build/Notifications/Sinks/SlackNotificationSink.cs b/Engine/Source/Programs/Horde/Horde.Build/Notifications/Sinks/SlackNotificationSink.cs
--- a/Engine/Source/Programs/Horde/Horde.Build/Notifications/Sinks/SlackNotificationSink.cs
+++ b/Engine/Source/Programs/Horde/Horde.Build/Notifications/Sinks/SlackNotificationSink.cs
@@ -42,7 +42,6 @@ using Microsoft.Extensions.Options;
 using MongoDB.Bson;
 using MongoDB.Bson.Serialization.Attributes;
 using MongoDB.Driver;
-using StackExchange.Redis;
 
 namespace Horde.Build.Notifications.Sinks
 {
@@ -215,7 +214,7 @@ namespace Horde.Build.Notifications.Sinks
 		readonly IssueService _issueService;
 		readonly IUserCollection _userCollection;
 		readonly ILogFileService _logFileService;
-		readonly StreamService _streamService;
+		readonly IStreamCollection _streamCollection;
 		readonly IWebHostEnvironment _environment;
 		readonly ServerSettings _settings;
 		readonly IMongoCollection<MessageStateDocument> _messageStates;
@@ -226,6 +225,7 @@ namespace Horde.Build.Notifications.Sinks
 		readonly ITicker _escalateTicker;
 		static readonly RedisSortedSetKey<int> _escalateIssues = "slack/escalate";
 		readonly IClock _clock;
+		readonly IOptionsMonitor<GlobalConfig> _globalConfig;
 		readonly ILogger _logger;
 
 		readonly ITicker _issueQueueTicker;
@@ -246,19 +246,20 @@ namespace Horde.Build.Notifications.Sinks
 		/// <summary>
 		/// Constructor
 		/// </summary>
-		public SlackNotificationSink(MongoService mongoService, RedisService redisService, IssueService issueService, IUserCollection userCollection, ILogFileService logFileService, StreamService streamService, IExternalIssueService externalIssueService, IWebHostEnvironment environment, IOptions<ServerSettings> settings, IClock clock, ILogger<SlackNotificationSink> logger)
+		public SlackNotificationSink(MongoService mongoService, RedisService redisService, IssueService issueService, IUserCollection userCollection, ILogFileService logFileService, IStreamCollection streamCollection, IExternalIssueService externalIssueService, IWebHostEnvironment environment, IOptions<ServerSettings> settings, IClock clock, IOptionsMonitor<GlobalConfig> globalConfig, ILogger<SlackNotificationSink> logger)
 		{
 			_redisService = redisService;
 			_issueService = issueService;
 			_userCollection = userCollection;
 			_logFileService = logFileService;
-			_streamService = streamService;
+			_streamCollection = streamCollection;
 			_externalIssueService = externalIssueService;
 			_environment = environment;
 			_settings = settings.Value;
 			_messageStates = mongoService.GetCollection<MessageStateDocument>("SlackV2", keys => keys.Ascending(x => x.Recipient).Ascending(x => x.EventId), unique: true);
 			_slackUsers = mongoService.GetCollection<SlackUserDocument>("Slack.UsersV2");
 			_clock = clock;
+			_globalConfig = globalConfig;
 			_logger = logger;
 
 			_escalateTicker = clock.AddSharedTicker<SlackNotificationSink>(TimeSpan.FromMinutes(1.0), EscalateAsync, _logger);
@@ -293,6 +294,7 @@ namespace Horde.Build.Notifications.Sinks
 
 			_userCache.Dispose();
 			_httpClient.Dispose();
+			_adminHttpClient?.Dispose();
 			_issueQueueTicker.Dispose();
 			_escalateTicker.Dispose();
 		}
@@ -428,19 +430,24 @@ namespace Horde.Build.Notifications.Sinks
 		#region Job Complete
 
 		/// <inheritdoc/>
-		public async Task NotifyJobCompleteAsync(IStream jobStream, IJob job, IGraph graph, LabelOutcome outcome)
+		public async Task NotifyJobCompleteAsync(IJob job, IGraph graph, LabelOutcome outcome)
 		{
+			StreamConfig? streamConfig;
+			if (!_globalConfig.CurrentValue.TryGetStream(job.StreamId, out streamConfig))
+			{
+				return;
+			}
 			if (job.NotificationChannel != null)
 			{
-				await SendJobCompleteNotificationToChannelAsync(job.NotificationChannel, job.NotificationChannelFilter, jobStream, job, graph, outcome);
+				await SendJobCompleteNotificationToChannelAsync(job.NotificationChannel, job.NotificationChannelFilter, streamConfig, job, graph, outcome);
 			}
-			if (jobStream.Config.NotificationChannel != null)
+			if (streamConfig.NotificationChannel != null)
 			{
-				await SendJobCompleteNotificationToChannelAsync(jobStream.Config.NotificationChannel, jobStream.Config.NotificationChannelFilter, jobStream, job, graph, outcome);
+				await SendJobCompleteNotificationToChannelAsync(streamConfig.NotificationChannel, streamConfig.NotificationChannelFilter, streamConfig, job, graph, outcome);
 			}
 		}
 
-		async Task SendJobCompleteNotificationToChannelAsync(string notificationChannel, string? notificationFilter, IStream jobStream, IJob job, IGraph graph, LabelOutcome outcome)
+		async Task SendJobCompleteNotificationToChannelAsync(string notificationChannel, string? notificationFilter, StreamConfig streamConfig, IJob job, IGraph graph, LabelOutcome outcome)
 		{
 			if (notificationFilter != null)
 			{
@@ -464,21 +471,27 @@ namespace Horde.Build.Notifications.Sinks
 			}
 			foreach (string channel in notificationChannel.Split(';'))
 			{
-				await SendJobCompleteMessageAsync(channel, jobStream, job, graph);
+				await SendJobCompleteMessageAsync(channel, streamConfig, job, graph);
 			}
 		}
 
 		/// <inheritdoc/>
-		public async Task NotifyJobCompleteAsync(IUser slackUser, IStream jobStream, IJob job, IGraph graph, LabelOutcome outcome)
+		public async Task NotifyJobCompleteAsync(IUser slackUser, IJob job, IGraph graph, LabelOutcome outcome)
 		{
+			StreamConfig? streamConfig;
+			if (!_globalConfig.CurrentValue.TryGetStream(job.StreamId, out streamConfig))
+			{
+				return;
+			}
+
 			string? slackUserId = await GetSlackUserId(slackUser);
 			if (slackUserId != null)
 			{
-				await SendJobCompleteMessageAsync(slackUserId, jobStream, job, graph);
+				await SendJobCompleteMessageAsync(slackUserId, streamConfig, job, graph);
 			}
 		}
 
-		private Task SendJobCompleteMessageAsync(string recipient, IStream stream, IJob job, IGraph graph)
+		private Task SendJobCompleteMessageAsync(string recipient, StreamConfig streamConfig, IJob job, IGraph graph)
 		{
 			JobStepOutcome jobOutcome = job.Batches.SelectMany(x => x.Steps).Min(x => x.Outcome);
 			_logger.LogInformation("Sending Slack notification for job {JobId} outcome {Outcome} to {SlackUser}", job.Id, jobOutcome, recipient);
@@ -488,9 +501,9 @@ namespace Horde.Build.Notifications.Sinks
 			string outcomeColor = jobOutcome == JobStepOutcome.Failure ? ErrorColor : jobOutcome == JobStepOutcome.Warnings ? WarningColor : SuccessColor;
 
 			SlackAttachment attachment = new SlackAttachment();
-			attachment.FallbackText = $"{stream.Name} - {GetJobChangeText(job)} - {job.Name} - {jobOutcome}";
+			attachment.FallbackText = $"{streamConfig.Name} - {GetJobChangeText(job)} - {job.Name} - {jobOutcome}";
 			attachment.Color = outcomeColor;
-			attachment.AddSection($"*<{jobLink}|{stream.Name} - {GetJobChangeText(job)} - {job.Name}>*");
+			attachment.AddSection($"*<{jobLink}|{streamConfig.Name} - {GetJobChangeText(job)} - {job.Name}>*");
 
 			if (!String.IsNullOrEmpty(job.PreflightDescription))
 			{
@@ -570,14 +583,14 @@ namespace Horde.Build.Notifications.Sinks
 		#region Job step complete
 
 		/// <inheritdoc/>
-		public async Task NotifyJobStepCompleteAsync(IUser slackUser, IStream jobStream, IJob job, IJobStepBatch batch, IJobStep step, INode node, List<ILogEventData> jobStepEventData)
+		public async Task NotifyJobStepCompleteAsync(IUser slackUser, IJob job, IJobStepBatch batch, IJobStep step, INode node, List<ILogEventData> jobStepEventData)
 		{
 			_logger.LogInformation("Sending Slack notification for job {JobId}, batch {BatchId}, step {StepId}, outcome {Outcome} to {SlackUser} ({UserId})", job.Id, batch.Id, step.Id, step.Outcome, slackUser.Name, slackUser.Id);
 
 			string? slackUserId = await GetSlackUserId(slackUser);
 			if (slackUserId != null)
 			{
-				await SendJobStepCompleteMessageAsync(slackUserId, jobStream, job, step, node, jobStepEventData);
+				await SendJobStepCompleteMessageAsync(slackUserId, job, step, node, jobStepEventData);
 			}
 		}
 
@@ -585,22 +598,27 @@ namespace Horde.Build.Notifications.Sinks
 		/// Creates a Slack message about a completed step job.
 		/// </summary>
 		/// <param name="recipient"></param>
-		/// <param name="stream"></param>
 		/// <param name="job">The job that contains the step that completed.</param>
 		/// <param name="step">The job step that completed.</param>
 		/// <param name="node">The node for the job step.</param>
 		/// <param name="events">Any events that occurred during the job step.</param>
-		private Task SendJobStepCompleteMessageAsync(string recipient, IStream stream, IJob job, IJobStep step, INode node, List<ILogEventData> events)
+		private Task SendJobStepCompleteMessageAsync(string recipient, IJob job, IJobStep step, INode node, List<ILogEventData> events)
 		{
+			GlobalConfig globalConfig = _globalConfig.CurrentValue;
+			if (!globalConfig.TryGetStream(job.StreamId, out StreamConfig? streamConfig))
+			{
+				return Task.CompletedTask;
+			}
+
 			Uri jobStepLink = new Uri($"{_settings.DashboardUrl}job/{job.Id}?step={step.Id}");
 			Uri jobStepLogLink = new Uri($"{_settings.DashboardUrl}log/{step.LogId}");
 
 			string outcomeColor = step.Outcome == JobStepOutcome.Failure ? ErrorColor : step.Outcome == JobStepOutcome.Warnings ? WarningColor : SuccessColor;
 
 			SlackAttachment attachment = new SlackAttachment();
-			attachment.FallbackText = $"{stream.Name} - {GetJobChangeText(job)} - {job.Name} - {node.Name} - {step.Outcome}";
+			attachment.FallbackText = $"{streamConfig.Name} - {GetJobChangeText(job)} - {job.Name} - {node.Name} - {step.Outcome}";
 			attachment.Color = outcomeColor;
-			attachment.AddSection($"*<{jobStepLink}|{stream.Name} - {GetJobChangeText(job)} - {job.Name} - {node.Name}>*");
+			attachment.AddSection($"*<{jobStepLink}|{streamConfig.Name} - {GetJobChangeText(job)} - {job.Name} - {node.Name}>*");
 			if (step.Outcome == JobStepOutcome.Success)
 			{
 				attachment.AddSection($"*Job Step Succeeded*");
@@ -640,27 +658,32 @@ namespace Horde.Build.Notifications.Sinks
 		#region Label complete
 
 		/// <inheritdoc/>
-		public async Task NotifyLabelCompleteAsync(IUser user, IJob job, IStream stream, ILabel label, int labelIdx, LabelOutcome outcome, List<(string, JobStepOutcome, Uri)> stepData)
+		public async Task NotifyLabelCompleteAsync(IUser user, IJob job, ILabel label, int labelIdx, LabelOutcome outcome, List<(string, JobStepOutcome, Uri)> stepData)
 		{
+			if (!_globalConfig.CurrentValue.TryGetStream(job.StreamId, out StreamConfig? streamConfig))
+			{
+				return;
+			}
+
 			_logger.LogInformation("Sending Slack notification for job {JobId} outcome {Outcome} to {Name} ({UserId})", job.Id, outcome, user.Name, user.Id);
 
 			string? slackUserId = await GetSlackUserId(user);
 			if (slackUserId != null)
 			{
-				await SendLabelUpdateMessageAsync(slackUserId, stream, job, label, labelIdx, outcome, stepData);
+				await SendLabelUpdateMessageAsync(slackUserId, streamConfig, job, label, labelIdx, outcome, stepData);
 			}
 		}
 
-		Task SendLabelUpdateMessageAsync(string recipient, IStream stream, IJob job, ILabel label, int labelIdx, LabelOutcome outcome, List<(string, JobStepOutcome, Uri)> jobStepData)
+		Task SendLabelUpdateMessageAsync(string recipient, StreamConfig streamConfig, IJob job, ILabel label, int labelIdx, LabelOutcome outcome, List<(string, JobStepOutcome, Uri)> jobStepData)
 		{
 			Uri labelLink = new Uri($"{_settings.DashboardUrl}job/{job.Id}?label={labelIdx}");
 
 			string outcomeColor = outcome == LabelOutcome.Failure ? ErrorColor : outcome == LabelOutcome.Warnings ? WarningColor : SuccessColor;
 
 			SlackAttachment attachment = new SlackAttachment();
-			attachment.FallbackText = $"{stream.Name} - {GetJobChangeText(job)} - {job.Name} - Label {label.DashboardName} - {outcome}";
+			attachment.FallbackText = $"{streamConfig.Name} - {GetJobChangeText(job)} - {job.Name} - Label {label.DashboardName} - {outcome}";
 			attachment.Color = outcomeColor;
-			attachment.AddSection($"*<{labelLink}|{stream.Name} - {GetJobChangeText(job)} - {job.Name} - Label {label.DashboardName}>*");
+			attachment.AddSection($"*<{labelLink}|{streamConfig.Name} - {GetJobChangeText(job)} - {job.Name} - Label {label.DashboardName}>*");
 			if (outcome == LabelOutcome.Success)
 			{
 				attachment.AddSection($"*Label Succeeded*");
@@ -720,6 +743,7 @@ namespace Horde.Build.Notifications.Sinks
 		/// <param name="cancellationToken">Cancellation token for the operation</param>
 		async ValueTask ProcessIssueQueueAsync(CancellationToken cancellationToken)
 		{
+			GlobalConfig globalConfig = _globalConfig.CurrentValue;
 			HashSet<int> testedIssueIds = new HashSet<int>();
 
 			// Execute loop number of times based on the length of the queue at the start. This should bound the number of iterations while allowing us
@@ -728,7 +752,7 @@ namespace Horde.Build.Notifications.Sinks
 			for (; count > 0; count--)
 			{
 				int issueId = await _redisService.GetDatabase().ListLeftPopAsync(_redisIssueQueue);
-				if (!testedIssueIds.Add(issueId) || !await TryUpdateIssueAsync(issueId))
+				if (!testedIssueIds.Add(issueId) || !await TryUpdateIssueAsync(globalConfig, issueId))
 				{
 					await _redisService.GetDatabase().ListRightPushAsync(_redisIssueQueue, issueId);
 				}
@@ -736,13 +760,13 @@ namespace Horde.Build.Notifications.Sinks
 			}
 		}
 
-		async ValueTask<bool> TryUpdateIssueAsync(int issueId)
+		async ValueTask<bool> TryUpdateIssueAsync(GlobalConfig globalConfig, int issueId)
 		{
 			using (RedisLock issueLock = new RedisLock(_redisService.GetDatabase(), $"{_redisIssueLockPrefix}/{issueId}"))
 			{
 				if (await issueLock.AcquireAsync(TimeSpan.FromSeconds(30.0)))
 				{
-					await NotifyIssueUpdatedInternalAsync(issueId);
+					await NotifyIssueUpdatedInternalAsync(globalConfig, issueId);
 					return true;
 				}
 				else
@@ -753,7 +777,7 @@ namespace Horde.Build.Notifications.Sinks
 			}
 		}
 
-		async Task NotifyIssueUpdatedInternalAsync(int issueId)
+		async Task NotifyIssueUpdatedInternalAsync(GlobalConfig globalConfig, int issueId)
 		{
 			using IDisposable scope = _logger.BeginScope("Slack notifications for issue {IssueId}", issueId);
 			_logger.LogInformation("Updating Slack notifications for issue {IssueId}", issueId);
@@ -777,8 +801,8 @@ namespace Horde.Build.Notifications.Sinks
 				WorkflowId? workflowId = span.LastFailure.Annotations.WorkflowId;
 				if (workflowId != null)
 				{
-					IStream? stream = await _streamService.GetStreamAsync(span.StreamId);
-					if (stream != null && stream.Config.TryGetWorkflow(workflowId.Value, out workflow) && workflow.TriageChannel != null)
+					StreamConfig? streamConfig;
+					if (globalConfig.TryGetStream(span.StreamId, out streamConfig) && streamConfig.TryGetWorkflow(workflowId.Value, out workflow) && workflow.TriageChannel != null)
 					{
 						await CreateOrUpdateWorkflowThreadAsync(workflow.TriageChannel, issue, span, details.Spans, workflow);
 						notifyOwner = notifySuspects = false;
@@ -820,17 +844,17 @@ namespace Horde.Build.Notifications.Sinks
 
 			foreach (IIssueSpan span in details.Spans)
 			{
-				IStream? stream = await _streamService.GetCachedStream(span.StreamId);
-				if (stream != null)
+				StreamConfig? streamConfig;
+				if (globalConfig.TryGetStream(span.StreamId, out streamConfig))
 				{
-					ITemplateRef? templateRef;
-					if (stream.Templates.TryGetValue(span.TemplateRefId, out templateRef) && templateRef.Config.TriageChannel != null)
+					TemplateRefConfig? templateRefConfig;
+					if (streamConfig.TryGetTemplate(span.TemplateRefId, out templateRefConfig) && templateRefConfig.TriageChannel != null)
 					{
-						channels.Add(templateRef.Config.TriageChannel);
+						channels.Add(templateRefConfig.TriageChannel);
 					}
-					else if (stream.Config.TriageChannel != null)
+					else if (streamConfig.TriageChannel != null)
 					{
-						channels.Add(stream.Config.TriageChannel);
+						channels.Add(streamConfig.TriageChannel);
 					}
 				}
 			}
@@ -846,7 +870,7 @@ namespace Horde.Build.Notifications.Sinks
 					}
 					else
 					{
-						await NotifyIssueUpdatedAsync(user, issue, details);
+						await NotifyIssueUpdatedAsync(globalConfig, user, issue, details);
 					}
 				}
 			}
@@ -855,11 +879,11 @@ namespace Horde.Build.Notifications.Sinks
 			{
 				foreach (string channel in channels)
 				{
-					await SendIssueMessageAsync(channel, issue, details, null, DefaultAllowMentions);
+					await SendIssueMessageAsync(globalConfig, channel, issue, details, null, DefaultAllowMentions);
 				}
 			}
 
-			await UpdateReportsAsync(issue, details.Spans);
+			await UpdateReportsAsync(globalConfig, issue, details.Spans);
 		}
 
 		static IIssueSpan? GetFixFailedSpan(IIssue issue, IReadOnlyList<IIssueSpan> spans)
@@ -1222,7 +1246,7 @@ namespace Horde.Build.Notifications.Sinks
 			}
 		}
 
-		async Task NotifyIssueUpdatedAsync(IUser user, IIssue issue, IIssueDetails details)
+		async Task NotifyIssueUpdatedAsync(GlobalConfig globalConfig, IUser user, IIssue issue, IIssueDetails details)
 		{
 			string? slackUserId = await GetSlackUserId(user);
 			if (slackUserId == null)
@@ -1230,7 +1254,7 @@ namespace Horde.Build.Notifications.Sinks
 				return;
 			}
 
-			await SendIssueMessageAsync(slackUserId, issue, details, user.Id, DefaultAllowMentions);
+			await SendIssueMessageAsync(globalConfig, slackUserId, issue, details, user.Id, DefaultAllowMentions);
 		}
 
 		Uri GetJobUrl(JobId jobId)
@@ -1248,7 +1272,7 @@ namespace Horde.Build.Notifications.Sinks
 			return new Uri(_settings.DashboardUrl, $"job/{step.JobId}?step={step.StepId}&issue={issue.Id}");
 		}
 
-		async Task SendIssueMessageAsync(string recipient, IIssue issue, IIssueDetails details, UserId? userId, bool allowMentions)
+		async Task SendIssueMessageAsync(GlobalConfig globalConfig, string recipient, IIssue issue, IIssueDetails details, UserId? userId, bool allowMentions)
 		{
 			using IDisposable scope = _logger.BeginScope("SendIssueMessageAsync (User: {SlackUser}, Issue: {IssueId})", recipient, issue.Id);
 
@@ -1269,8 +1293,8 @@ namespace Horde.Build.Notifications.Sinks
 			WorkflowId? workflowId = span.LastFailure.Annotations.WorkflowId;
 			if (workflowId != null)
 			{
-				IStream? stream = await _streamService.GetStreamAsync(span.StreamId);
-				if (stream != null && stream.Config.TryGetWorkflow(workflowId.Value, out WorkflowConfig? workflow) && workflow.TriageChannel != null && workflow.AllowMentions)
+				StreamConfig? streamConfig;
+				if(globalConfig.TryGetStream(span.StreamId, out streamConfig) && streamConfig.TryGetWorkflow(workflowId.Value, out WorkflowConfig? workflow) && workflow.TriageChannel != null && workflow.AllowMentions)
 				{
 					MessageStateDocument? state = await GetMessageStateAsync(workflow.TriageChannel, GetTriageThreadEventId(issue.Id));
 					if (state != null)
@@ -1532,9 +1556,11 @@ namespace Horde.Build.Notifications.Sinks
 		/// <inheritdoc/>
 		public async Task SendIssueReportAsync(IssueReportGroup group)
 		{
-			foreach (IssueReport report in group.Reports.OrderBy(x => x.WorkflowId).ThenBy(x => x.Stream.Id))
+			GlobalConfig globalConfig = _globalConfig.CurrentValue;
+
+			foreach (IssueReport report in group.Reports.OrderBy(x => x.WorkflowId).ThenBy(x => x.StreamId))
 			{
-				await SendIssueReportForStreamAsync(group.Channel, group.Time, report);
+				await SendIssueReportForStreamAsync(globalConfig, group.Channel, group.Time, report);
 			}
 
 			SlackMessage message = new SlackMessage();
@@ -1543,10 +1569,16 @@ namespace Horde.Build.Notifications.Sinks
 			await SendMessageAsync(group.Channel, message);
 		}
 
-		async Task SendIssueReportForStreamAsync(string channel, DateTime time, IssueReport report)
+		async Task SendIssueReportForStreamAsync(GlobalConfig globalConfig, string channel, DateTime time, IssueReport report)
 		{
 			const int MaxIssuesPerMessage = 8;
 
+			StreamConfig? streamConfig;
+			if (!globalConfig.TryGetStream(report.StreamId, out streamConfig))
+			{
+				return;
+			}
+
 			ReportState state = new ReportState();
 			state.Time = time;
 
@@ -1562,11 +1594,11 @@ namespace Horde.Build.Notifications.Sinks
 				foreach (IGrouping<TemplateId, IIssueSpan> group in report.IssueSpans.GroupBy(x => x.TemplateRefId).OrderBy(x => x.Key.ToString()))
 				{
 					TemplateRefConfig? templateConfig;
-					if (!report.Stream.Config.TryGetTemplate(group.Key, out templateConfig))
+					if (!streamConfig.TryGetTemplate(group.Key, out templateConfig))
 					{
 						continue;
 					}
-					if (!IsIssueOpenForWorkflow(report.Stream.Id, group.Key, report.WorkflowId, group))
+					if (!IsIssueOpenForWorkflow(report.StreamId, group.Key, report.WorkflowId, group))
 					{
 						continue;
 					}
@@ -1587,7 +1619,7 @@ namespace Horde.Build.Notifications.Sinks
 						{
 							ReportBlock block = new ReportBlock();
 							block.TemplateHeader = templateHeader;
-							block.StreamId = report.Stream.Id;
+							block.StreamId = report.StreamId;
 							block.TemplateId = group.Key;
 							block.IssueIds.AddRange(batch.Select(x => x.Key.Id));
 							state.Blocks.Add(block);
@@ -1603,7 +1635,7 @@ namespace Horde.Build.Notifications.Sinks
 				foreach (IReadOnlyList<IIssue> batch in report.Issues.OrderByDescending(x => x.Id).Batch(MaxIssuesPerMessage))
 				{
 					ReportBlock block = new ReportBlock();
-					block.StreamId = report.Stream.Id;
+					block.StreamId = report.StreamId;
 					block.IssueIds.AddRange(batch.Select(x => x.Id));
 					state.Blocks.Add(block);
 
@@ -1612,12 +1644,12 @@ namespace Horde.Build.Notifications.Sinks
 			}
 
 			SlackMessage headerMessage = new SlackMessage();
-			headerMessage.AddHeader($"Summary for {report.Stream.Name}");
+			headerMessage.AddHeader($"Summary for {streamConfig.Name}");
 
 			SlackMessageId? messageId = await SendMessageAsync(channel, headerMessage);
 			if (messageId != null)
 			{
-				string reportEventId = GetReportEventId(report.Stream.Id, report.WorkflowId);
+				string reportEventId = GetReportEventId(streamConfig.Id, report.WorkflowId);
 				string json = JsonSerializer.Serialize(state, _jsonSerializerOptions);
 				await AddOrUpdateMessageStateAsync(channel, reportEventId, null, json, messageId);
 
@@ -1630,7 +1662,7 @@ namespace Horde.Build.Notifications.Sinks
 				for (int idx = 0; idx < state.Blocks.Count; idx++)
 				{
 					string blockEventId = GetReportBlockEventId(messageId.Ts, idx);
-					await UpdateReportBlockAsync(channel, blockEventId, time, report.Stream, state.Blocks[idx].TemplateId, issuesByBlock[idx], report.TriageChannel, state.Blocks[idx].TemplateHeader);
+					await UpdateReportBlockAsync(channel, blockEventId, time, streamConfig, state.Blocks[idx].TemplateId, issuesByBlock[idx], report.TriageChannel, state.Blocks[idx].TemplateHeader);
 				}
 
 				if (report.WorkflowStats.NumSteps > 0)
@@ -1657,7 +1689,7 @@ namespace Horde.Build.Notifications.Sinks
 			return false;
 		}
 
-		async Task UpdateReportsAsync(IIssue issue, IReadOnlyList<IIssueSpan> spans)
+		async Task UpdateReportsAsync(GlobalConfig globalConfig, IIssue issue, IReadOnlyList<IIssueSpan> spans)
 		{
 			_logger.LogInformation("Checking for report updates to issue {IssueId}", issue.Id);
 
@@ -1676,19 +1708,19 @@ namespace Horde.Build.Notifications.Sinks
 					continue;
 				}
 
-				IStream? stream = await _streamService.GetStreamAsync(span.StreamId);
-				if (stream == null)
+				StreamConfig? streamConfig;
+				if (!globalConfig.TryGetStream(span.StreamId, out streamConfig))
 				{
 					continue;
 				}
 
-				WorkflowConfig? workflow;
-				if (!stream.Config.TryGetWorkflow(workflowId.Value, out workflow) || workflow.ReportChannel == null)
+				WorkflowConfig? workflowConfig;
+				if (!streamConfig.TryGetWorkflow(workflowId.Value, out workflowConfig) || workflowConfig.ReportChannel == null)
 				{
 					continue;
 				}
 
-				MessageStateDocument? messageState = await GetMessageStateAsync(workflow.ReportChannel, reportEventId);
+				MessageStateDocument? messageState = await GetMessageStateAsync(workflowConfig.ReportChannel, reportEventId);
 				if (messageState == null)
 				{
 					continue;
@@ -1731,30 +1763,30 @@ namespace Horde.Build.Notifications.Sinks
 									otherSpan = details.Spans[0];
 								}
 
-								bool open = IsIssueOpenForWorkflow(stream.Id, block.TemplateId, workflowId.Value, details.Spans);
+								bool open = IsIssueOpenForWorkflow(streamConfig.Id, block.TemplateId, workflowId.Value, details.Spans);
 								issues.Add((details.Issue, otherSpan, open));
 							}
 						}
 
-						await UpdateReportBlockAsync(workflow.ReportChannel, blockEventId, state.Time, stream, block.TemplateId, issues, workflow.TriageChannel, block.TemplateHeader);
+						await UpdateReportBlockAsync(workflowConfig.ReportChannel, blockEventId, state.Time, streamConfig, block.TemplateId, issues, workflowConfig.TriageChannel, block.TemplateHeader);
 					}
 				}
 			}
 		}
 
-		async Task UpdateReportBlockAsync(string channel, string eventId, DateTime reportTime, IStream stream, TemplateId templateId, List<(IIssue, IIssueSpan?, bool)> issues, string? triageChannel, bool templateHeader)
+		async Task UpdateReportBlockAsync(string channel, string eventId, DateTime reportTime, StreamConfig streamConfig, TemplateId templateId, List<(IIssue, IIssueSpan?, bool)> issues, string? triageChannel, bool templateHeader)
 		{
 			StringBuilder body = new StringBuilder();
 
 			if (templateHeader && !templateId.IsEmpty)
 			{
 				TemplateRefConfig? templateConfig;
-				if (stream.Config.TryGetTemplate(templateId, out templateConfig))
+				if (streamConfig.TryGetTemplate(templateId, out templateConfig))
 				{
-					JobsTabConfig? tab = stream.Config.Tabs.OfType<JobsTabConfig>().FirstOrDefault(x => x.Templates != null && x.Templates.Contains(templateId));
+					JobsTabConfig? tab = streamConfig.Tabs.OfType<JobsTabConfig>().FirstOrDefault(x => x.Templates != null && x.Templates.Contains(templateId));
 					if (tab != null)
 					{
-						Uri templateUrl = new Uri(_settings.DashboardUrl, $"stream/{stream.Id}?tab={tab.Title}&template={templateId}");
+						Uri templateUrl = new Uri(_settings.DashboardUrl, $"stream/{streamConfig.Id}?tab={tab.Title}&template={templateId}");
 						body.Append($"*<{templateUrl}|{templateConfig.Name}>*:");
 					}
 				}
@@ -2055,7 +2087,7 @@ namespace Horde.Build.Notifications.Sinks
 		#region Device notifications
 
 		/// <inheritdoc/>
-		public async Task NotifyDeviceServiceAsync(string message, IDevice? device = null, IDevicePool? pool = null, IStream? stream = null, IJob? job = null, IJobStep? step = null, INode? node = null, IUser? user = null)
+		public async Task NotifyDeviceServiceAsync(string message, IDevice? device = null, IDevicePool? pool = null, StreamConfig? streamConfig = null, IJob? job = null, IJobStep? step = null, INode? node = null, IUser? user = null)
 		{
 			string? recipient = null;
 
@@ -2075,7 +2107,7 @@ namespace Horde.Build.Notifications.Sinks
 			if (recipient != null)
 			{
 				_logger.LogDebug("Sending device service notification to {Recipient}", recipient);
-				await SendDeviceServiceMessage(recipient, message, device, pool, stream, job, step, node, user);
+				await SendDeviceServiceMessage(recipient, message, device, pool, streamConfig, job, step, node, user);
 			}
 		}
 
@@ -2086,12 +2118,12 @@ namespace Horde.Build.Notifications.Sinks
 		/// <param name="message"></param>
 		/// <param name="device"></param>
 		/// <param name="pool"></param>
-		/// <param name="stream"></param>
+		/// <param name="streamConfig"></param>
 		/// <param name="job">The job that contains the step that completed.</param>
 		/// <param name="step">The job step that completed.</param>
 		/// <param name="node">The node for the job step.</param>
 		/// <param name="user">The user to notify.</param>
-		private Task SendDeviceServiceMessage(string recipient, string message, IDevice? device = null, IDevicePool? pool = null, IStream? stream = null, IJob? job = null, IJobStep? step = null, INode? node = null, IUser? user = null)
+		private Task SendDeviceServiceMessage(string recipient, string message, IDevice? device = null, IDevicePool? pool = null, StreamConfig? streamConfig = null, IJob? job = null, IJobStep? step = null, INode? node = null, IUser? user = null)
 		{
 
 			if (user != null)
@@ -2116,13 +2148,13 @@ namespace Horde.Build.Notifications.Sinks
 				
 			attachment.AddHeader(message, false);
 
-			if (stream != null && job != null && step != null && node != null)
+			if (streamConfig != null && job != null && step != null && node != null)
 			{
 				Uri jobStepLink = new Uri($"{_settings.DashboardUrl}job/{job.Id}?step={step.Id}");
 				Uri jobStepLogLink = new Uri($"{_settings.DashboardUrl}log/{step.LogId}");
 
-				attachment.FallbackText += $" - {stream.Name} - {GetJobChangeText(job)} - {job.Name} - {node.Name}";
-				attachment.AddSection($"*<{jobStepLink}|{stream.Name} - {GetJobChangeText(job)} - {job.Name} - {node.Name}>*");
+				attachment.FallbackText += $" - {streamConfig.Name} - {GetJobChangeText(job)} - {job.Name} - {node.Name}";
+				attachment.AddSection($"*<{jobStepLink}|{streamConfig.Name} - {GetJobChangeText(job)} - {job.Name} - {node.Name}>*");
 				attachment.AddSection($"<{jobStepLogLink}|View Job Step Log>");
 			}
 			else
@@ -2304,6 +2336,8 @@ namespace Horde.Build.Notifications.Sinks
 
 		async ValueTask EscalateAsync(CancellationToken cancellationToken)
 		{
+			GlobalConfig globalConfig = _globalConfig.CurrentValue;
+
 			DateTime utcNow = DateTime.UtcNow;
 			double time = (utcNow - DateTime.UnixEpoch).TotalSeconds;
 
@@ -2316,7 +2350,7 @@ namespace Horde.Build.Notifications.Sinks
 					cancellationToken.ThrowIfCancellationRequested();
 					try
 					{
-						double? nextTime = await EscalateSingleIssueAsync(issueId, utcNow);
+						double? nextTime = await EscalateSingleIssueAsync(globalConfig, issueId, utcNow);
 						if (nextTime == null)
 						{
 							_logger.LogInformation("Cancelling escalation for issue {IssueId}", issueId);
@@ -2337,7 +2371,7 @@ namespace Horde.Build.Notifications.Sinks
 			}
 		}
 
-		async Task<double?> EscalateSingleIssueAsync(int issueId, DateTime utcNow)
+		async Task<double?> EscalateSingleIssueAsync(GlobalConfig globalConfig, int issueId, DateTime utcNow)
 		{
 			IIssue? issue = await _issueService.Collection.GetIssueAsync(issueId);
 			if (issue == null)
@@ -2359,8 +2393,12 @@ namespace Horde.Build.Notifications.Sinks
 				return null;
 			}
 
-			IStream? stream = await _streamService.GetStreamAsync(span.StreamId);
-			if (stream == null || !stream.Config.TryGetWorkflow(workflowId.Value, out WorkflowConfig? workflow))
+			StreamConfig? streamConfig;
+			if (!globalConfig.TryGetStream(span.StreamId, out streamConfig))
+			{
+				return null;
+			}
+			if (!streamConfig.TryGetWorkflow(workflowId.Value, out WorkflowConfig? workflow))
 			{
 				return null;
 			}
@@ -2369,7 +2407,7 @@ namespace Horde.Build.Notifications.Sinks
 				return null;
 			}
 
-			if (!IsIssueOpenForWorkflow(stream.Id, span.TemplateRefId, workflow.Id, spans))
+			if (!IsIssueOpenForWorkflow(streamConfig.Id, span.TemplateRefId, workflow.Id, spans))
 			{
 				return null;
 			}
@@ -2687,7 +2725,7 @@ namespace Horde.Build.Notifications.Sinks
 					if (recipient != null)
 					{
 						IIssueDetails details = await _issueService.GetIssueDetailsAsync(newIssue);
-						await SendIssueMessageAsync(recipient, newIssue, details, userId, DefaultAllowMentions);
+						await SendIssueMessageAsync(_globalConfig.CurrentValue, recipient, newIssue, details, userId, DefaultAllowMentions);
 					}
 				}
 			}
diff --git a/Engine/Source/Programs/Horde/Horde.Build/Notifications/SubscriptionsController.cs b/Engine/Source/Programs/Horde/Horde.Build/Notifications/SubscriptionsController.cs
--- a/Engine/Source/Programs/Horde/Horde.Build/Notifications/SubscriptionsController.cs
+++ b/Engine/Source/Programs/Horde/Horde.Build/Notifications/SubscriptionsController.cs
@@ -4,10 +4,12 @@ using System;
 using System.Collections.Generic;
 using System.Threading.Tasks;
 using Horde.Build.Acls;
+using Horde.Build.Server;
 using Horde.Build.Users;
 using Horde.Build.Utilities;
 using Microsoft.AspNetCore.Authorization;
 using Microsoft.AspNetCore.Mvc;
+using Microsoft.Extensions.Options;
 
 namespace Horde.Build.Notifications
 {
@@ -21,25 +23,16 @@ namespace Horde.Build.Notifications
 	[Route("[controller]")]
 	public class SubscriptionsController : ControllerBase
 	{
-		/// <summary>
-		/// The ACL service singleton
-		/// </summary>
-		readonly AclService _aclService;
-
-		/// <summary>
-		/// Collection of subscription documents
-		/// </summary>
 		readonly ISubscriptionCollection _subscriptionCollection;
+		readonly IOptionsSnapshot<GlobalConfig> _globalConfig;
 
 		/// <summary>
 		/// Constructor
 		/// </summary>
-		/// <param name="aclService">The acl service singleton</param>
-		/// <param name="subscriptionCollection">The collection of subscription documents</param>
-		public SubscriptionsController(AclService aclService, ISubscriptionCollection subscriptionCollection)
+		public SubscriptionsController(ISubscriptionCollection subscriptionCollection, IOptionsSnapshot<GlobalConfig> globalConfig)
 		{
-			_aclService = aclService;
 			_subscriptionCollection = subscriptionCollection;
+			_globalConfig = globalConfig;
 		}
 
 		/// <summary>
@@ -58,7 +51,7 @@ namespace Horde.Build.Notifications
 			{
 				return BadRequest("Invalid user id");
 			}
-			if (!await _aclService.AuthorizeAsUserAsync(User, userIdValue))
+			if (!_globalConfig.Value.AuthorizeAsUser(User, userIdValue))
 			{
 				return Forbid();
 			}
@@ -83,7 +76,7 @@ namespace Horde.Build.Notifications
 			{
 				return NotFound();
 			}
-			if (!await _aclService.AuthorizeAsUserAsync(User, subscription.UserId))
+			if (!_globalConfig.Value.AuthorizeAsUser(User, subscription.UserId))
 			{
 				return Forbid();
 			}
@@ -106,7 +99,7 @@ namespace Horde.Build.Notifications
 			{
 				return NotFound();
 			}
-			if (!await _aclService.AuthorizeAsUserAsync(User, subscription.UserId))
+			if (!_globalConfig.Value.AuthorizeAsUser(User, subscription.UserId))
 			{
 				return Forbid();
 			}
@@ -132,8 +125,6 @@ namespace Horde.Build.Notifications
 				authorizedUsers.Add(currentUserId.Value);
 			}
 
-			GlobalPermissionsCache cache = new GlobalPermissionsCache();
-
 			List<NewSubscription> newSubscriptions = new List<NewSubscription>();
 			foreach (CreateSubscriptionRequest subscription in subscriptions)
 			{
@@ -142,7 +133,7 @@ namespace Horde.Build.Notifications
 				{
 					return BadRequest($"Invalid user id: '{subscription.UserId}'.");
 				}
-				if (authorizedUsers.Add(newUserId) && !await _aclService.AuthorizeAsync(AclAction.Impersonate, User, cache))
+				if (authorizedUsers.Add(newUserId) && !_globalConfig.Value.Authorize(AclAction.Impersonate, User))
 				{
 					return Forbid();
 				}
diff --git a/Engine/Source/Programs/Horde/Horde.Build/Perforce/CommitService.cs b/Engine/Source/Programs/Horde/Horde.Build/Perforce/CommitService.cs
--- a/Engine/Source/Programs/Horde/Horde.Build/Perforce/CommitService.cs
+++ b/Engine/Source/Programs/Horde/Horde.Build/Perforce/CommitService.cs
@@ -1,24 +1,6 @@
 // Copyright Epic Games, Inc. All Rights Reserved.
 
-using System;
-using System.Collections.Generic;
-using System.Globalization;
-using System.Linq;
-using System.Text.RegularExpressions;
-using System.Threading;
-using System.Threading.Tasks;
-using EpicGames.Core;
-using EpicGames.Perforce;
-using EpicGames.Redis;
-using Horde.Build.Server;
 using Horde.Build.Streams;
-using Horde.Build.Users;
-using Horde.Build.Utilities;
-using HordeCommon;
-using Microsoft.Extensions.Hosting;
-using Microsoft.Extensions.Logging;
-using Microsoft.Extensions.Options;
-using StackExchange.Redis;
 
 namespace Horde.Build.Perforce
 {
@@ -38,6 +20,6 @@ namespace Horde.Build.Perforce
 		}
 
 		/// <inheritdoc/>
-		public ICommitCollection GetCollection(IStream stream) => _perforceService.GetCommits(stream);
+		public ICommitCollection GetCollection(StreamConfig streamConfig) => _perforceService.GetCommits(streamConfig);
 	}
 }
diff --git a/Engine/Source/Programs/Horde/Horde.Build/Perforce/ICommitService.cs b/Engine/Source/Programs/Horde/Horde.Build/Perforce/ICommitService.cs
--- a/Engine/Source/Programs/Horde/Horde.Build/Perforce/ICommitService.cs
+++ b/Engine/Source/Programs/Horde/Horde.Build/Perforce/ICommitService.cs
@@ -12,8 +12,8 @@ namespace Horde.Build.Perforce
 		/// <summary>
 		/// Gets a commit collection for the given stream
 		/// </summary>
-		/// <param name="stream">Stream to get commits for</param>
+		/// <param name="streamConfig">Stream to get commits for</param>
 		/// <returns>Collection object</returns>
-		ICommitCollection GetCollection(IStream stream);
+		ICommitCollection GetCollection(StreamConfig streamConfig);
 	}
 }
diff --git a/Engine/Source/Programs/Horde/Horde.Build/Perforce/IPerforceService.cs b/Engine/Source/Programs/Horde/Horde.Build/Perforce/IPerforceService.cs
--- a/Engine/Source/Programs/Horde/Horde.Build/Perforce/IPerforceService.cs
+++ b/Engine/Source/Programs/Horde/Horde.Build/Perforce/IPerforceService.cs
@@ -89,10 +89,10 @@ namespace Horde.Build.Perforce
 		/// <summary>
 		/// Gets a cached stream view for the given stream
 		/// </summary>
-		/// <param name="stream">Stream to find the view for</param>
+		/// <param name="streamConfig">Stream to find the view for</param>
 		/// <param name="cancellationToken"></param>
 		/// <returns></returns>
-		ValueTask<PerforceViewMap> GetCachedStreamViewAsync(IStream stream, CancellationToken cancellationToken);
+		ValueTask<PerforceViewMap> GetCachedStreamViewAsync(StreamConfig streamConfig, CancellationToken cancellationToken);
 	}
 
 	/// <summary>
@@ -104,26 +104,26 @@ namespace Horde.Build.Perforce
 		/// Map files from a <see cref="DescribeRecord"/> to relative paths within a stream
 		/// </summary>
 		/// <param name="perforce">Perforce connection</param>
-		/// <param name="stream">Stream to target</param>
+		/// <param name="streamConfig">Config for the stream to target</param>
 		/// <param name="describeRecord">Description of the change</param>
 		/// <param name="cancellationToken">Cancellation token for the operation</param>
 		/// <returns>List of files relative to the stream</returns>
-		public static Task<List<string>> GetStreamFilesAsync(this IPooledPerforceConnection perforce, IStream stream, DescribeRecord describeRecord, CancellationToken cancellationToken)
+		public static Task<List<string>> GetStreamFilesAsync(this IPooledPerforceConnection perforce, StreamConfig streamConfig, DescribeRecord describeRecord, CancellationToken cancellationToken)
 		{
-			return GetStreamFilesAsync(perforce, stream, describeRecord.Files.ConvertAll(x => x.DepotFile), cancellationToken);
+			return GetStreamFilesAsync(perforce, streamConfig, describeRecord.Files.ConvertAll(x => x.DepotFile), cancellationToken);
 		}
 
 		/// <summary>
 		/// Map files from a <see cref="DescribeRecord"/> to relative paths within a stream
 		/// </summary>
 		/// <param name="perforce">Perforce connection</param>
-		/// <param name="stream">Stream to target</param>
+		/// <param name="streamConfig">Stream to target</param>
 		/// <param name="depotPaths">Paths within the depot</param>
 		/// <param name="cancellationToken">Cancellation token for the operation</param>
 		/// <returns>List of files relative to the stream</returns>
-		public static async Task<List<string>> GetStreamFilesAsync(this IPooledPerforceConnection perforce, IStream stream, List<string> depotPaths, CancellationToken cancellationToken)
+		public static async Task<List<string>> GetStreamFilesAsync(this IPooledPerforceConnection perforce, StreamConfig streamConfig, List<string> depotPaths, CancellationToken cancellationToken)
 		{
-			PerforceViewMap viewMap = await perforce.GetCachedStreamViewAsync(stream, cancellationToken);
+			PerforceViewMap viewMap = await perforce.GetCachedStreamViewAsync(streamConfig, cancellationToken);
 			InfoRecord infoRecord = await perforce.GetInfoAsync(cancellationToken);
 			return viewMap.MapFiles(depotPaths, infoRecord.PathComparison).ToList();
 		}
@@ -155,11 +155,11 @@ namespace Horde.Build.Perforce
 		/// <summary>
 		/// Checks a shelf is valid for the given stream
 		/// </summary>
-		/// <param name="stream">The stream to query</param>
+		/// <param name="streamConfig">The stream to query</param>
 		/// <param name="changeNumber">Shelved changelist number</param>
 		/// <param name="cancellationToken">Cancellation token for the operation</param>
 		/// <returns>Whether the shelf is valid, plus the description for it</returns>
-		public Task<(CheckShelfResult, ShelfInfo?)> CheckShelfAsync(IStream stream, int changeNumber, CancellationToken cancellationToken = default);
+		public Task<(CheckShelfResult, ShelfInfo?)> CheckShelfAsync(StreamConfig streamConfig, int changeNumber, CancellationToken cancellationToken = default);
 
 		/// <summary>
 		/// Duplicates a shelved changelist
@@ -173,12 +173,12 @@ namespace Horde.Build.Perforce
 		/// <summary>
 		/// Submit a shelved changelist
 		/// </summary>
-		/// <param name="stream">Stream to submit to</param>
+		/// <param name="streamConfig">Stream to submit to</param>
 		/// <param name="shelvedChange">The shelved changelist number, created by <see cref="DuplicateShelvedChangeAsync(String,Int32,CancellationToken)"/></param>
 		/// <param name="originalChange">The original changelist number</param>
 		/// <param name="cancellationToken">Cancellation token for the operation</param>
 		/// <returns>Tuple consisting of the submitted changelist number and message</returns>
-		public Task<(int? Change, string Message)> SubmitShelvedChangeAsync(IStream stream, int shelvedChange, int originalChange, CancellationToken cancellationToken = default);
+		public Task<(int? Change, string Message)> SubmitShelvedChangeAsync(StreamConfig streamConfig, int shelvedChange, int originalChange, CancellationToken cancellationToken = default);
 
 		/// <summary>
 		/// Deletes a changelist containing shelved files.
@@ -202,8 +202,8 @@ namespace Horde.Build.Perforce
 		/// <summary>
 		/// Creates a commit source for the given stream
 		/// </summary>
-		/// <param name="stream">Stream to create a commit source for</param>
+		/// <param name="streamConfig">Stream to create a commit source for</param>
 		/// <returns>Commit source instance</returns>
-		public ICommitCollection GetCommits(IStream stream);
+		public ICommitCollection GetCommits(StreamConfig streamConfig);
 	}
 }
diff --git a/Engine/Source/Programs/Horde/Horde.Build/Perforce/PerforceController.cs b/Engine/Source/Programs/Horde/Horde.Build/Perforce/PerforceController.cs
--- a/Engine/Source/Programs/Horde/Horde.Build/Perforce/PerforceController.cs
+++ b/Engine/Source/Programs/Horde/Horde.Build/Perforce/PerforceController.cs
@@ -8,6 +8,7 @@ using Horde.Build.Server;
 using Microsoft.AspNetCore.Authorization;
 using Microsoft.AspNetCore.Mvc;
 using Microsoft.Extensions.Logging;
+using Microsoft.Extensions.Options;
 
 namespace Horde.Build.Perforce
 {
@@ -19,29 +20,16 @@ namespace Horde.Build.Perforce
 	[Route("[controller]")]
 	public class PerforceController : ControllerBase
 	{
-		/// <summary>
-		/// The globals service instance
-		/// </summary>
-		private readonly GlobalsService _globalsService;
-
-		/// <summary>
-		/// The ACL service instance
-		/// </summary>
-		private readonly AclService _aclService;
-
-		/// <summary>
-		/// Load balancer instance
-		/// </summary>
 		private readonly PerforceLoadBalancer _perforceLoadBalancer;
+		private readonly IOptionsSnapshot<GlobalConfig> _globalConfig;
 
 		/// <summary>
 		/// Constructor
 		/// </summary>
-		public PerforceController(GlobalsService globalsService, AclService aclService, PerforceLoadBalancer perforceLoadBalancer)
+		public PerforceController(PerforceLoadBalancer perforceLoadBalancer, IOptionsSnapshot<GlobalConfig> globalConfig)
 		{
-			_globalsService = globalsService;
-			_aclService = aclService;
 			_perforceLoadBalancer = perforceLoadBalancer;
+			_globalConfig = globalConfig;
 		}
 
 		/// <summary>
@@ -68,15 +56,16 @@ namespace Horde.Build.Perforce
 		/// <returns>List of Perforce clusters</returns>
 		[HttpGet]
 		[Route("/api/v1/perforce/settings")]
-		public async Task<ActionResult<List<PerforceCluster>>> GetPerforceSettingsAsync()
+		public ActionResult<List<PerforceCluster>> GetPerforceSettingsAsync()
 		{
-			if (!await _aclService.AuthorizeAsync(AclAction.AdminRead, User))
+			GlobalConfig globalConfig = _globalConfig.Value;
+
+			if (!globalConfig.Authorize(AclAction.AdminRead, User))
 			{
 				return Forbid();
 			}
 
-			IGlobals globals = await _globalsService.GetAsync();
-			return globals.Config.PerforceClusters;
+			return globalConfig.PerforceClusters;
 		}
 	}
 
@@ -87,9 +76,6 @@ namespace Horde.Build.Perforce
 	[Route("[controller]")]
 	public class PublicPerforceController : ControllerBase
 	{
-		/// <summary>
-		/// Logger instance
-		/// </summary>
 		private readonly ILogger<PerforceController> _logger;
 
 		/// <summary>
diff --git a/Engine/Source/Programs/Horde/Horde.Build/Perforce/PerforceLoadBalancer.cs b/Engine/Source/Programs/Horde/Horde.Build/Perforce/PerforceLoadBalancer.cs
--- a/Engine/Source/Programs/Horde/Horde.Build/Perforce/PerforceLoadBalancer.cs
+++ b/Engine/Source/Programs/Horde/Horde.Build/Perforce/PerforceLoadBalancer.cs
@@ -19,6 +19,7 @@ using HordeCommon;
 using HordeCommon.Rpc.Tasks;
 using Microsoft.Extensions.Hosting;
 using Microsoft.Extensions.Logging;
+using Microsoft.Extensions.Options;
 using MongoDB.Bson.Serialization.Attributes;
 using MongoDB.Driver;
 
@@ -141,18 +142,20 @@ namespace Horde.Build.Perforce
 		readonly ILeaseCollection _leaseCollection;
 		readonly SingletonDocument<PerforceServerList> _serverListSingleton;
 		readonly Random _random = new Random();
+		readonly IOptionsMonitor<GlobalConfig> _globalConfig;
 		readonly ILogger _logger;
 		readonly ITicker _ticker;
 
 		/// <summary>
 		/// Constructor
 		/// </summary>
-		public PerforceLoadBalancer(MongoService mongoService, GlobalsService globalsService, ILeaseCollection leaseCollection, IClock clock, ILogger<PerforceLoadBalancer> logger)
+		public PerforceLoadBalancer(MongoService mongoService, GlobalsService globalsService, ILeaseCollection leaseCollection, IClock clock, IOptionsMonitor<GlobalConfig> globalConfig, ILogger<PerforceLoadBalancer> logger)
 		{
 			_mongoService = mongoService;
 			_globalsService = globalsService;
 			_leaseCollection = leaseCollection;
 			_serverListSingleton = new SingletonDocument<PerforceServerList>(mongoService);
+			_globalConfig = globalConfig;
 			_logger = logger;
 			if (mongoService.ReadOnlyMode)
 			{
@@ -358,13 +361,13 @@ namespace Horde.Build.Perforce
 		/// <inheritdoc/>
 		async ValueTask TickInternalAsync(CancellationToken cancellationToken)
 		{
-			IGlobals globals = await _globalsService.GetAsync();
+			GlobalConfig globalConfig = _globalConfig.CurrentValue;
 
 			// Set of new server entries
 			List<PerforceServerEntry> newServers = new List<PerforceServerEntry>();
 
 			// Update the state of all the valid servers
-			foreach (PerforceCluster cluster in globals.Config.PerforceClusters)
+			foreach (PerforceCluster cluster in globalConfig.PerforceClusters)
 			{
 				foreach (PerforceServer server in cluster.Servers)
 				{
diff --git a/Engine/Source/Programs/Horde/Horde.Build/Perforce/PerforceReplicator.cs b/Engine/Source/Programs/Horde/Horde.Build/Perforce/PerforceReplicator.cs
--- a/Engine/Source/Programs/Horde/Horde.Build/Perforce/PerforceReplicator.cs
+++ b/Engine/Source/Programs/Horde/Horde.Build/Perforce/PerforceReplicator.cs
@@ -80,19 +80,19 @@ namespace Horde.Build.Perforce
 		}
 
 		// Partial mirror of FileSysType from P4 API (filesys.h)
-		const uint FST_TEXT = 0x0001;
+//		const uint FST_TEXT = 0x0001;
 		const uint FST_BINARY = 0x0002;
-		const uint FST_UNICODE = 0x000c;
+//		const uint FST_UNICODE = 0x000c;
 		const uint FST_UTF16 = 0x000e;
-		const uint FST_UTF8 = 0x000f;
+//		const uint FST_UTF8 = 0x000f;
 		const uint FST_MASK = 0x000f;
 
 		// Mirrors FilePerm from P4 API (filesys.h)
 		const uint FPM_RO = 0;     // leave file read-only
-		const uint FPM_RW = 1;     // leave file read-write
+//		const uint FPM_RW = 1;     // leave file read-write
 		const uint FPM_ROO = 2;    // leave file read-only (owner)
 		const uint FPM_RXO = 3;    // set file read-execute (owner) NO W
-		const uint FPM_RWO = 4;    // set file read-write (owner) NO X
+//		const uint FPM_RWO = 4;    // set file read-write (owner) NO X
 		const uint FPM_RWXO = 5;   // set file read-write-execute (owner)
 
 		[DebuggerDisplay("{_path}")]
@@ -249,11 +249,11 @@ namespace Horde.Build.Perforce
 		/// <summary>
 		/// Replicates a change to storage
 		/// </summary>
-		/// <param name="stream">Stream to replicate data from</param>
+		/// <param name="streamConfig">Stream to replicate data from</param>
 		/// <param name="change">Changelist to replicate</param>
 		/// <param name="options">Options for replication</param>
 		/// <param name="cancellationToken">Cancellation token for the operation</param>
-		public async Task WriteAsync(IStream stream, int change, PerforceReplicationOptions options, CancellationToken cancellationToken)
+		public async Task WriteAsync(StreamConfig streamConfig, int change, PerforceReplicationOptions options, CancellationToken cancellationToken)
 		{
 			IStorageClient store = await _storageService.GetClientAsync(Namespace.Perforce, cancellationToken);
 
@@ -263,13 +263,13 @@ namespace Horde.Build.Perforce
 			TreeReader reader = new TreeReader(store, _memoryCache, readerOptions, _logger);
 
 			// Find the parent node
-			RefName refName = GetRefName(stream.Id);
+			RefName refName = GetRefName(streamConfig.Id);
 			CommitNode? parent = await FindParentAsync(reader, refName, change, cancellationToken);
 			TreeNodeRef<CommitNode>? parentRef = (parent == null) ? null : new TreeNodeRef<CommitNode>(parent);
 			int parentChange = parent?.Number ?? 0;
 
 			// Read the current incremental state or create a new node to track the incremental state
-			RefName incRefName = GetIncrementalRefName(stream.Id);
+			RefName incRefName = GetIncrementalRefName(streamConfig.Id);
 			SyncNode? syncNode = await reader.TryReadNodeAsync<SyncNode>(incRefName, cancellationToken: cancellationToken);
 			if (syncNode == null || syncNode.Change != change || syncNode.ParentChange != parentChange)
 			{
@@ -278,7 +278,7 @@ namespace Horde.Build.Perforce
 			DirectoryNode root = await syncNode.Contents.ExpandAsync(reader, cancellationToken);
 
 			// Create a client to replicate from this stream
-			ReplicationClient clientInfo = await FindOrAddReplicationClientAsync(stream);
+			ReplicationClient clientInfo = await FindOrAddReplicationClientAsync(streamConfig);
 
 			// Connect to the server and flush the workspace
 			using IPerforceConnection perforce = await PerforceConnection.CreateAsync(clientInfo.Settings, _logger);
@@ -371,7 +371,7 @@ namespace Horde.Build.Perforce
 
 				syncedSize += size;
 				double syncPct = (totalSize == 0) ? 100.0 : (syncedSize * 100.0) / totalSize;
-				_logger.LogInformation("Syncing {StreamId} to {Change} [{SyncPct:n1}%] ({Size:n1}mb)", stream.Id, change, syncPct, size / (1024.0 * 1024.0));
+				_logger.LogInformation("Syncing {StreamId} to {Change} [{SyncPct:n1}%] ({Size:n1}mb)", streamConfig.Id, change, syncPct, size / (1024.0 * 1024.0));
 
 				// Copy them to a separate list and remove any redundant paths
 				List<string> syncPaths = new List<string>();
@@ -636,12 +636,12 @@ namespace Horde.Build.Perforce
 			return path.Substring(clientRoot.Length);
 		}
 
-		async Task<ReplicationClient?> FindReplicationClientAsync(IStream stream)
+		async Task<ReplicationClient?> FindReplicationClientAsync(StreamConfig streamConfig)
 		{
 			ReplicationClient? clientInfo;
-			if (_cachedPerforceClients.TryGetValue(stream.Id, out clientInfo))
+			if (_cachedPerforceClients.TryGetValue(streamConfig.Id, out clientInfo))
 			{
-				if (!String.Equals(clientInfo.ClusterName, stream.Config.ClusterName, StringComparison.Ordinal) && String.Equals(clientInfo.Client.Stream, stream.Name, StringComparison.Ordinal))
+				if (!String.Equals(clientInfo.ClusterName, streamConfig.ClusterName, StringComparison.Ordinal) && String.Equals(clientInfo.Client.Stream, streamConfig.Name, StringComparison.Ordinal))
 				{
 					PerforceSettings serverSettings = new PerforceSettings(clientInfo.Settings);
 					serverSettings.ClientName = null;
@@ -649,19 +649,19 @@ namespace Horde.Build.Perforce
 					using IPerforceConnection perforce = await PerforceConnection.CreateAsync(_logger);
 					await perforce.DeleteClientAsync(DeleteClientOptions.None, clientInfo.Client.Name);
 
-					_cachedPerforceClients.Remove(stream.Id);
+					_cachedPerforceClients.Remove(streamConfig.Id);
 					clientInfo = null;
 				}
 			}
 			return clientInfo;
 		}
 
-		async Task<ReplicationClient> FindOrAddReplicationClientAsync(IStream stream)
+		async Task<ReplicationClient> FindOrAddReplicationClientAsync(StreamConfig streamConfig)
 		{
-			ReplicationClient? clientInfo = await FindReplicationClientAsync(stream);
+			ReplicationClient? clientInfo = await FindReplicationClientAsync(streamConfig);
 			if (clientInfo == null)
 			{
-				using IPerforceConnection? perforce = await _perforceService.ConnectAsync(stream.Config.ClusterName);
+				using IPerforceConnection? perforce = await _perforceService.ConnectAsync(streamConfig.ClusterName);
 				if (perforce == null)
 				{
 					throw new PerforceException($"Unable to create connection to Perforce server");
@@ -669,21 +669,21 @@ namespace Horde.Build.Perforce
 
 				InfoRecord serverInfo = await perforce.GetInfoAsync(InfoOptions.ShortOutput);
 
-				ClientRecord newClient = new ClientRecord($"Horde.Build_Rep_{serverInfo.ClientHost}_{stream.Id}", perforce.Settings.UserName, "/p4/");
+				ClientRecord newClient = new ClientRecord($"Horde.Build_Rep_{serverInfo.ClientHost}_{streamConfig.Id}", perforce.Settings.UserName, "/p4/");
 				newClient.Description = "Created to mirror Perforce content to Horde Storage";
 				newClient.Owner = perforce.Settings.UserName;
 				newClient.Host = serverInfo.ClientHost;
-				newClient.Stream = stream.Config.ReplicationStream ?? stream.Name;
+				newClient.Stream = streamConfig.ReplicationStream ?? streamConfig.Name;
 				newClient.Type = "readonly";
 				await perforce.CreateClientAsync(newClient);
-				_logger.LogInformation("Created client {ClientName} for {StreamName}", newClient.Name, stream.Name);
+				_logger.LogInformation("Created client {ClientName} for {StreamName}", newClient.Name, streamConfig.Name);
 
 				PerforceSettings settings = new PerforceSettings(perforce.Settings);
 				settings.ClientName = newClient.Name;
 				settings.PreferNativeClient = true;
 
-				clientInfo = new ReplicationClient(settings, stream.Config.ClusterName, serverInfo, newClient, -1);
-				_cachedPerforceClients.Add(stream.Id, clientInfo);
+				clientInfo = new ReplicationClient(settings, streamConfig.ClusterName, serverInfo, newClient, -1);
+				_cachedPerforceClients.Add(streamConfig.Id, clientInfo);
 			}
 			return clientInfo;
 		}
diff --git a/Engine/Source/Programs/Horde/Horde.Build/Perforce/PerforceService.cs b/Engine/Source/Programs/Horde/Horde.Build/Perforce/PerforceService.cs
--- a/Engine/Source/Programs/Horde/Horde.Build/Perforce/PerforceService.cs
+++ b/Engine/Source/Programs/Horde/Horde.Build/Perforce/PerforceService.cs
@@ -111,15 +111,15 @@ namespace Horde.Build.Perforce
 			public ValueTask<InfoRecord> GetInfoAsync(CancellationToken cancellationToken) => _inner.GetInfoAsync(cancellationToken);
 
 			/// <inheritdoc/>
-			public async ValueTask<PerforceViewMap> GetCachedStreamViewAsync(IStream stream, CancellationToken cancellationToken)
+			public async ValueTask<PerforceViewMap> GetCachedStreamViewAsync(StreamConfig streamConfig, CancellationToken cancellationToken)
 			{
 				PerforceViewMap? viewMap;
-				if (!_owner._streamCache.TryGetValue(stream.Id, out viewMap))
+				if (!_owner._streamCache.TryGetValue(streamConfig.Id, out viewMap))
 				{
-					StreamRecord streamRecord = await _inner.Perforce.GetStreamAsync(stream.Name, true, cancellationToken);
+					StreamRecord streamRecord = await _inner.Perforce.GetStreamAsync(streamConfig.Name, true, cancellationToken);
 					viewMap = PerforceViewMap.Parse(streamRecord.View);
 
-					using (ICacheEntry entry = _owner._streamCache.CreateEntry(stream.Id))
+					using (ICacheEntry entry = _owner._streamCache.CreateEntry(streamConfig.Id))
 					{
 						entry.AbsoluteExpirationRelativeToNow = TimeSpan.FromMinutes(2.0);
 						entry.Value = viewMap;
@@ -147,9 +147,9 @@ namespace Horde.Build.Perforce
 		class Commit : ICommit
 		{
 			readonly PerforceService _owner;
-			readonly IStream _stream;
+			readonly StreamConfig _streamConfig;
 
-			public StreamId StreamId => _stream.Id;
+			public StreamId StreamId => _streamConfig.Id;
 			public int Number { get; }
 			public int OriginalChange { get; }
 			public UserId AuthorId { get; }
@@ -160,10 +160,10 @@ namespace Horde.Build.Perforce
 
 			IReadOnlyList<string>? _files;
 
-			public Commit(PerforceService owner, IStream stream, int number, int originalChange, UserId authorId, UserId ownerId, string description, string basePath, DateTime dateUtc)
+			public Commit(PerforceService owner, StreamConfig streamConfig, int number, int originalChange, UserId authorId, UserId ownerId, string description, string basePath, DateTime dateUtc)
 			{
 				_owner = owner;
-				_stream = stream;
+				_streamConfig = streamConfig;
 
 				Number = number;
 				OriginalChange = originalChange;
@@ -184,9 +184,9 @@ namespace Horde.Build.Perforce
 				List<CommitTag> commitTags = new List<CommitTag>();
 
 				IReadOnlyList<string> files = await GetFilesAsync(cancellationToken);
-				foreach (CommitTagConfig tagConfig in _stream.Config.GetAllCommitTags())
+				foreach (CommitTagConfig tagConfig in _streamConfig.GetAllCommitTags())
 				{
-					if (_stream.Config.TryGetCommitTagFilter(tagConfig.Name, out FileFilter? filter) && filter.ApplyTo(files).Any())
+					if (_streamConfig.TryGetCommitTagFilter(tagConfig.Name, out FileFilter? filter) && filter.ApplyTo(files).Any())
 					{
 						commitTags.Add(tagConfig.Name);
 					}
@@ -199,10 +199,10 @@ namespace Horde.Build.Perforce
 			{
 				if (_files == null)
 				{
-					using (IPooledPerforceConnection perforce = await _owner.ConnectAsync(_stream.Config.ClusterName, null, cancellationToken))
+					using (IPooledPerforceConnection perforce = await _owner.ConnectAsync(_streamConfig.ClusterName, null, cancellationToken))
 					{
 						DescribeRecord describeRecord = await perforce.DescribeAsync(Number, cancellationToken);
-						_files = await perforce.GetStreamFilesAsync(_stream, describeRecord, cancellationToken);
+						_files = await perforce.GetStreamFilesAsync(_streamConfig, describeRecord, cancellationToken);
 					}
 				}
 				return _files;
@@ -210,7 +210,7 @@ namespace Horde.Build.Perforce
 		}
 
 		readonly PerforceLoadBalancer _loadBalancer;
-		readonly LazyCachedValue<Task<IGlobals>> _cachedGlobals;
+		readonly IOptionsMonitor<GlobalConfig> _globalConfig;
 		readonly ILogger _logger;
 
 		// Useful overrides for local debugging with read-only data
@@ -228,12 +228,12 @@ namespace Horde.Build.Perforce
 		/// <summary>
 		/// Constructor
 		/// </summary>
-		public PerforceService(PerforceLoadBalancer loadBalancer, GlobalsService globalsService, IUserCollection userCollection, IOptions<ServerSettings> settings, ILogger<PerforceService> logger)
+		public PerforceService(PerforceLoadBalancer loadBalancer, IUserCollection userCollection, IOptions<ServerSettings> settings, IOptionsMonitor<GlobalConfig> globalConfig, ILogger<PerforceService> logger)
 		{
 			_loadBalancer = loadBalancer;
-			_cachedGlobals = new LazyCachedValue<Task<IGlobals>>(async () => await globalsService.GetAsync(), TimeSpan.FromSeconds(30.0));
 			_userCollection = userCollection;
 			_settings = settings.Value;
+			_globalConfig = globalConfig;
 			_logger = logger;
 
 			if(settings.Value.UseLocalPerforceEnv)
@@ -401,7 +401,7 @@ namespace Horde.Build.Perforce
 
 		public async Task<IPooledPerforceConnection> ConnectAsync(string clusterName, string? userName = null, CancellationToken cancellationToken = default)
 		{
-			PerforceCluster cluster = await GetClusterAsync(clusterName);
+			PerforceCluster cluster = _globalConfig.CurrentValue.GetPerforceCluster(clusterName);
 			return await ConnectAsync(cluster, userName, cancellationToken);
 		}
 
@@ -468,7 +468,7 @@ namespace Horde.Build.Perforce
 
 		async Task<PooledConnectionHandle> ConnectWithStreamClientAsync(string clusterName, string? userName, string streamName, CancellationToken cancellationToken)
 		{
-			PerforceCluster cluster = await GetClusterAsync(clusterName);
+			PerforceCluster cluster = _globalConfig.CurrentValue.GetPerforceCluster(clusterName);
 
 			PooledConnectionHandle? handle = GetPooledConnectionForStream(cluster, userName, streamName);
 			if (handle != null)
@@ -494,9 +494,9 @@ namespace Horde.Build.Perforce
 			}
 		}
 
-		Task<PooledConnectionHandle> ConnectWithStreamClientAsync(IStream stream, string? userName, CancellationToken cancellationToken)
+		Task<PooledConnectionHandle> ConnectWithStreamClientAsync(StreamConfig streamConfig, string? userName, CancellationToken cancellationToken)
 		{
-			return ConnectWithStreamClientAsync(stream.Config.ClusterName, userName, stream.Name, cancellationToken);
+			return ConnectWithStreamClientAsync(streamConfig.ClusterName, userName, streamConfig.Name, cancellationToken);
 		}
 
 		async Task<Credentials> GetTicketAsync(PerforceCluster cluster, string userName, CancellationToken cancellationToken)
@@ -553,7 +553,7 @@ namespace Horde.Build.Perforce
 
 		public async ValueTask<IUser> FindOrAddUserAsync(string clusterName, string userName, CancellationToken cancellationToken = default)
 		{
-			PerforceCluster cluster = await GetClusterAsync(clusterName);
+			PerforceCluster cluster = _globalConfig.CurrentValue.GetPerforceCluster(clusterName);
 			return await FindOrAddUserAsync(cluster, userName, cancellationToken);
 		}
 
@@ -595,16 +595,13 @@ namespace Horde.Build.Perforce
 			return user!;
 		}
 
-		async Task<PerforceCluster> GetClusterAsync(string? clusterName, string? serverAndPort = null)
+		PerforceCluster GetCluster(string? clusterName, string? serverAndPort = null)
 		{
-			IGlobals globals = await _cachedGlobals.GetCached();
-
-			PerforceCluster? cluster = globals.Config.FindPerforceCluster(clusterName, serverAndPort);
+			PerforceCluster? cluster = _globalConfig.CurrentValue.FindPerforceCluster(clusterName, serverAndPort);
 			if (cluster == null)
 			{
 				throw new Exception($"Unknown Perforce cluster '{clusterName}'");
 			}
-
 			return cluster;
 		}
 
@@ -642,7 +639,7 @@ namespace Horde.Build.Perforce
 			return result;
 		}
 
-		async ValueTask<Commit> CreateCommitInternalAsync(IStream stream, int number, string author, string description, string? basePath, DateTime dateUtc, CancellationToken cancellationToken)
+		async ValueTask<Commit> CreateCommitInternalAsync(StreamConfig streamConfig, int number, string author, string description, string? basePath, DateTime dateUtc, CancellationToken cancellationToken)
 		{
 			// Robomerge submits changes as its own user, then modifies the commit with the real author afterwards. Parse it out of the changelist description to handle these cases.
 			if (String.Equals(author, "robomerge", StringComparison.OrdinalIgnoreCase))
@@ -656,7 +653,7 @@ namespace Horde.Build.Perforce
 				}
 			}
 
-			IUser authorUser = await FindOrAddUserAsync(stream.Config.ClusterName, author, cancellationToken);
+			IUser authorUser = await FindOrAddUserAsync(streamConfig.ClusterName, author, cancellationToken);
 			IUser ownerUser = authorUser;
 
 			int originalChange = ParseRobomergeSource(description) ?? number;
@@ -664,33 +661,33 @@ namespace Horde.Build.Perforce
 			string? owner = ParseRobomergeOwner(description);
 			if (owner != null)
 			{
-				ownerUser = await FindOrAddUserAsync(stream.Config.ClusterName, owner, cancellationToken);
+				ownerUser = await FindOrAddUserAsync(streamConfig.ClusterName, owner, cancellationToken);
 			}
 
-			return new Commit(this, stream, number, originalChange, authorUser.Id, ownerUser.Id, description, basePath ?? String.Empty, dateUtc);
+			return new Commit(this, streamConfig, number, originalChange, authorUser.Id, ownerUser.Id, description, basePath ?? String.Empty, dateUtc);
 		}
 
-		protected async ValueTask<ICommit> CreateCommitAsync(IPooledPerforceConnection perforce, IStream stream, DescribeRecord describeRecord, InfoRecord serverInfo, CancellationToken cancellationToken)
+		protected async ValueTask<ICommit> CreateCommitAsync(IPooledPerforceConnection perforce, StreamConfig streamConfig, DescribeRecord describeRecord, InfoRecord serverInfo, CancellationToken cancellationToken)
 		{
 			DateTime timeUtc = new DateTime(describeRecord.Time.Ticks - serverInfo.TimeZoneOffsetSecs * TimeSpan.TicksPerSecond, DateTimeKind.Utc);
 
-			Commit commit = await CreateCommitInternalAsync(stream, describeRecord.Number, describeRecord.User, describeRecord.Description, describeRecord.Path, timeUtc, cancellationToken);
+			Commit commit = await CreateCommitInternalAsync(streamConfig, describeRecord.Number, describeRecord.User, describeRecord.Description, describeRecord.Path, timeUtc, cancellationToken);
 
-			List<string> files = await perforce.GetStreamFilesAsync(stream, describeRecord, cancellationToken);
+			List<string> files = await perforce.GetStreamFilesAsync(streamConfig, describeRecord, cancellationToken);
 			if (files.Count == 0)
 			{
-				throw new PerforceException($"Changelist {commit.Number} does not contain any files in {stream.Id}");
+				throw new PerforceException($"Changelist {commit.Number} does not contain any files in {streamConfig.Id}");
 			}
 			commit.SetFiles(files);
 
 			return commit;
 		}
 
-		async ValueTask<Commit> CreateCommitAsync(IStream stream, ChangesRecord changesRecord, InfoRecord serverInfo, CancellationToken cancellationToken)
+		async ValueTask<Commit> CreateCommitAsync(StreamConfig streamConfig, ChangesRecord changesRecord, InfoRecord serverInfo, CancellationToken cancellationToken)
 		{
 			DateTime timeUtc = new DateTime(changesRecord.Time.Ticks - serverInfo.TimeZoneOffsetSecs * TimeSpan.TicksPerSecond, DateTimeKind.Utc);
 
-			Commit commit = await CreateCommitInternalAsync(stream, changesRecord.Number, changesRecord.User, changesRecord.Description, changesRecord.Path, timeUtc, cancellationToken);
+			Commit commit = await CreateCommitInternalAsync(streamConfig, changesRecord.Number, changesRecord.User, changesRecord.Description, changesRecord.Path, timeUtc, cancellationToken);
 
 			return commit;
 		}
@@ -734,14 +731,14 @@ namespace Horde.Build.Perforce
 		}
 
 		/// <inheritdoc/>
-		public async Task<ICommit> GetChangeDetailsAsync(IStream stream, int changeNumber, CancellationToken cancellationToken)
+		public async Task<ICommit> GetChangeDetailsAsync(StreamConfig streamConfig, int changeNumber, CancellationToken cancellationToken)
 		{
 			using IScope scope = GlobalTracer.Instance.BuildSpan("PerforceService.GetChangeDetailsAsync").StartActive();
-			scope.Span.SetTag("ClusterName", stream.Config.ClusterName);
-			scope.Span.SetTag("StreamName", stream.Name);
+			scope.Span.SetTag("ClusterName", streamConfig.ClusterName);
+			scope.Span.SetTag("StreamName", streamConfig.Name);
 			scope.Span.SetTag("ChangeNumber", changeNumber);
 
-			PerforceCluster cluster = await GetClusterAsync(stream.Config.ClusterName);
+			PerforceCluster cluster = _globalConfig.CurrentValue.GetPerforceCluster(streamConfig.ClusterName);
 
 			using (PooledConnectionHandle perforce = await ConnectAsync(cluster, null, cancellationToken))
 			{
@@ -751,20 +748,20 @@ namespace Horde.Build.Perforce
 
 				DateTime timeUtc = new DateTime(record.Time.Ticks - info.TimeZoneOffsetSecs * TimeSpan.TicksPerSecond, DateTimeKind.Utc);
 
-				ICommit commit = await CreateCommitAsync(perforce, stream, record, info, cancellationToken);
+				ICommit commit = await CreateCommitAsync(perforce, streamConfig, record, info, cancellationToken);
 				return commit;
 			}
 		}
 
 		/// <inheritdoc/>
-		public async Task<(CheckShelfResult, ShelfInfo?)> CheckShelfAsync(IStream stream, int changeNumber, CancellationToken cancellationToken)
+		public async Task<(CheckShelfResult, ShelfInfo?)> CheckShelfAsync(StreamConfig streamConfig, int changeNumber, CancellationToken cancellationToken)
 		{
 			using IScope scope = GlobalTracer.Instance.BuildSpan("PerforceService.CheckPreflightAsync").StartActive();
-			scope.Span.SetTag("ClusterName", stream.Config.ClusterName);
-			scope.Span.SetTag("StreamName", stream.Name);
+			scope.Span.SetTag("ClusterName", streamConfig.ClusterName);
+			scope.Span.SetTag("StreamName", streamConfig.Name);
 			scope.Span.SetTag("ChangeNumber", changeNumber);
-
-			PerforceCluster cluster = await GetClusterAsync(stream.Config.ClusterName);
+			
+			PerforceCluster cluster = _globalConfig.CurrentValue.GetPerforceCluster(streamConfig.ClusterName);
 
 			using (IPooledPerforceConnection perforce = await ConnectAsync(cluster, null, cancellationToken))
 			{
@@ -787,7 +784,7 @@ namespace Horde.Build.Perforce
 					return (CheckShelfResult.NoShelvedFiles, null);
 				}
 
-				PerforceViewMap viewMap = await perforce.GetCachedStreamViewAsync(stream, cancellationToken);
+				PerforceViewMap viewMap = await perforce.GetCachedStreamViewAsync(streamConfig, cancellationToken);
 				List<string> mappedFiles = new List<string>();
 
 				bool hasUnmappedFile = false;
@@ -810,10 +807,10 @@ namespace Horde.Build.Perforce
 				}
 
 				List<CommitTag> tags = new List<CommitTag>();
-				foreach (CommitTagConfig tagConfig in stream.Config.GetAllCommitTags())
+				foreach (CommitTagConfig tagConfig in streamConfig.GetAllCommitTags())
 				{
 					FileFilter? filter;
-					if (stream.Config.TryGetCommitTagFilter(tagConfig.Name, out filter) && mappedFiles.Any(x => filter.Matches(x)))
+					if (streamConfig.TryGetCommitTagFilter(tagConfig.Name, out filter) && mappedFiles.Any(x => filter.Matches(x)))
 					{
 						tags.Add(tagConfig.Name);
 					}
@@ -824,22 +821,22 @@ namespace Horde.Build.Perforce
 		}
 
 		/// <inheritdoc/>
-		public async Task<List<ICommit>> GetChangeDetailsAsync(IStream stream, IReadOnlyList<int> changeNumbers, CancellationToken cancellationToken)
+		public async Task<List<ICommit>> GetChangeDetailsAsync(StreamConfig streamConfig, IReadOnlyList<int> changeNumbers, CancellationToken cancellationToken)
 		{
 			using IScope scope = GlobalTracer.Instance.BuildSpan("PerforceService.GetChangeDetailsAsync").StartActive();
-			scope.Span.SetTag("ClusterName", stream.Config.ClusterName);
-			scope.Span.SetTag("StreamName", stream.Name);
+			scope.Span.SetTag("ClusterName", streamConfig.ClusterName);
+			scope.Span.SetTag("StreamName", streamConfig.Name);
 			scope.Span.SetTag("ChangeNumbers.Count", changeNumbers.Count);
 
 			List<ICommit> results = new List<ICommit>();
-			using (PooledConnectionHandle perforce = await ConnectWithStreamClientAsync(stream, null, cancellationToken))
+			using (PooledConnectionHandle perforce = await ConnectWithStreamClientAsync(streamConfig, null, cancellationToken))
 			{
 				InfoRecord info = await perforce.GetInfoAsync(cancellationToken);
 
 				List<DescribeRecord> records = await perforce.DescribeAsync(DescribeOptions.Shelved, -1, changeNumbers.ToArray(), cancellationToken);
 				foreach (DescribeRecord record in records)
 				{
-					ICommit commit = await CreateCommitAsync(perforce, stream, record, info, cancellationToken);
+					ICommit commit = await CreateCommitAsync(perforce, streamConfig, record, info, cancellationToken);
 					results.Add(commit);
 				}
 			}
@@ -859,7 +856,7 @@ namespace Horde.Build.Perforce
 			scope.Span.SetTag("ClusterName", clusterName);
 			scope.Span.SetTag("ShelvedChange", shelvedChange);
 			
-			PerforceCluster cluster = await GetClusterAsync(clusterName);
+			PerforceCluster cluster = _globalConfig.CurrentValue.GetPerforceCluster(clusterName);
 
 			using (IPerforceConnection perforce = await ConnectAsChangeOwnerAsync(cluster, shelvedChange, cancellationToken))
 			{
@@ -877,7 +874,7 @@ namespace Horde.Build.Perforce
 
 			try
 			{
-				PerforceCluster cluster = await GetClusterAsync(clusterName);
+				PerforceCluster cluster = _globalConfig.CurrentValue.GetPerforceCluster(clusterName);
 				using (IPerforceConnection perforce = await ConnectAsChangeOwnerAsync(cluster, change, cancellationToken))
 				{
 					ChangeRecord record = await perforce.GetChangeAsync(GetChangeOptions.None, change, cancellationToken);
@@ -1044,14 +1041,14 @@ namespace Horde.Build.Perforce
 		}
 
 		/// <inheritdoc/>
-		public async Task<(int? Change, string Message)> SubmitShelvedChangeAsync(IStream stream, int change, int originalChange, CancellationToken cancellationToken)
+		public async Task<(int? Change, string Message)> SubmitShelvedChangeAsync(StreamConfig streamConfig, int change, int originalChange, CancellationToken cancellationToken)
 		{
 			using IScope scope = GlobalTracer.Instance.BuildSpan("PerforceService.SubmitShelvedChangeAsync").StartActive();
-			scope.Span.SetTag("Stream", stream.Id);
+			scope.Span.SetTag("Stream", streamConfig.Id);
 			scope.Span.SetTag("Change", change);
 			scope.Span.SetTag("OriginalChange", originalChange);
 
-			PerforceCluster cluster = await GetClusterAsync(stream.Config.ClusterName);
+			PerforceCluster cluster = _globalConfig.CurrentValue.GetPerforceCluster(streamConfig.ClusterName);
 
 			List<DescribeRecord> records;
 			using (IPerforceConnection perforce = await ConnectAsync(cluster, null, cancellationToken))
@@ -1160,12 +1157,12 @@ namespace Horde.Build.Perforce
 		protected class CommitSource : ICommitCollection
 		{
 			protected readonly PerforceService _perforceService;
-			protected readonly IStream _stream;
+			protected readonly StreamConfig _streamConfig;
 
-			public CommitSource(PerforceService perforceService, IStream stream)
+			public CommitSource(PerforceService perforceService, StreamConfig streamConfig)
 			{
 				_perforceService = perforceService;
-				_stream = stream;
+				_streamConfig = streamConfig;
 			}
 
 			public async Task<int> CreateNewAsync(string path, string description, CancellationToken cancellationToken = default)
@@ -1173,23 +1170,23 @@ namespace Horde.Build.Perforce
 				Match match = Regex.Match(path, @"^(//[^/]+/[^/]+)/(.+)$");
 				if (match.Success)
 				{
-					return await _perforceService.CreateNewChangeAsync(_stream.Config.ClusterName, match.Groups[1].Value, match.Groups[2].Value, description, cancellationToken);
+					return await _perforceService.CreateNewChangeAsync(_streamConfig.ClusterName, match.Groups[1].Value, match.Groups[2].Value, description, cancellationToken);
 				}
 				else
 				{
-					return await _perforceService.CreateNewChangeAsync(_stream.Config.ClusterName, _stream.Name, path, description, cancellationToken);
+					return await _perforceService.CreateNewChangeAsync(_streamConfig.ClusterName, _streamConfig.Name, path, description, cancellationToken);
 				}
 			}
 
 			public virtual async IAsyncEnumerable<ICommit> FindAsync(int? minChange, int? maxChange, int? maxResults, IReadOnlyList<CommitTag>? tags, [EnumeratorCancellation] CancellationToken cancellationToken = default)
 			{
 				using IScope scope = GlobalTracer.Instance.BuildSpan("PerforceService.CommitSource.FindAsync").StartActive();
-				scope.Span.SetTag("ClusterName", _stream.Config.ClusterName);
+				scope.Span.SetTag("ClusterName", _streamConfig.ClusterName);
 				scope.Span.SetTag("MinChange", minChange ?? -2);
 				scope.Span.SetTag("MaxChange", maxChange ?? -2);
 				scope.Span.SetTag("MaxResults", maxResults ?? -1);
 
-				using (PooledConnectionHandle perforce = await _perforceService.ConnectWithStreamClientAsync(_stream, null, cancellationToken))
+				using (PooledConnectionHandle perforce = await _perforceService.ConnectWithStreamClientAsync(_streamConfig, null, cancellationToken))
 				{
 					InfoRecord info = await perforce.GetInfoAsync(cancellationToken);
 
@@ -1200,7 +1197,7 @@ namespace Horde.Build.Perforce
 						List<ChangesRecord> changes = await perforce.GetChangesAsync(ChangesOptions.IncludeTimes | ChangesOptions.LongOutput, maxResults ?? -1, ChangeStatus.Submitted, filter, cancellationToken);
 						foreach (ChangesRecord change in changes)
 						{
-							ICommit commit = await _perforceService.CreateCommitAsync(_stream, change, info, cancellationToken);
+							ICommit commit = await _perforceService.CreateCommitAsync(_streamConfig, change, info, cancellationToken);
 							yield return commit;
 						}
 					}
@@ -1213,14 +1210,14 @@ namespace Horde.Build.Perforce
 						foreach (CommitTag tag in tags)
 						{
 							FileFilter? filter;
-							if (_stream.Config.TryGetCommitTagFilter(tag, out filter))
+							if (_streamConfig.TryGetCommitTagFilter(tag, out filter))
 							{
 								filters.Add(filter);
 							}
 						}
 
 						// Get the stream view
-						PerforceViewMap viewMap = await perforce.GetCachedStreamViewAsync(_stream, cancellationToken);
+						PerforceViewMap viewMap = await perforce.GetCachedStreamViewAsync(_streamConfig, cancellationToken);
 
 						// Do the query in batches
 						int numResults = maxResults ?? Int32.MaxValue;
@@ -1240,7 +1237,7 @@ namespace Horde.Build.Perforce
 								List<string> files = viewMap.MapFiles(describeRecord.Files.Select(x => x.DepotFile), info.PathComparison).ToList();
 								if (files.Any(x => filters.Any(y => y.Matches(x))))
 								{
-									ICommit commit = await _perforceService.CreateCommitAsync(perforce, _stream, describeRecord, info, cancellationToken);
+									ICommit commit = await _perforceService.CreateCommitAsync(perforce, _streamConfig, describeRecord, info, cancellationToken);
 									yield return commit;
 
 									if (--numResults == 0)
@@ -1259,7 +1256,7 @@ namespace Horde.Build.Perforce
 			/// <inheritdoc/>
 			public virtual Task<ICommit> GetAsync(int changeNumber, CancellationToken cancellationToken = default)
 			{
-				return _perforceService.GetChangeDetailsAsync(_stream, changeNumber, cancellationToken);
+				return _perforceService.GetChangeDetailsAsync(_streamConfig, changeNumber, cancellationToken);
 			}
 
 			/// <inheritdoc/>
@@ -1277,9 +1274,9 @@ namespace Horde.Build.Perforce
 		}
 
 		/// <inheritdoc/>
-		public virtual ICommitCollection GetCommits(IStream stream)
+		public virtual ICommitCollection GetCommits(StreamConfig streamConfig)
 		{
-			return new CommitSource(this, stream);
+			return new CommitSource(this, streamConfig);
 		}
 	}
 }
diff --git a/Engine/Source/Programs/Horde/Horde.Build/Perforce/PerforceServiceCache.cs b/Engine/Source/Programs/Horde/Horde.Build/Perforce/PerforceServiceCache.cs
--- a/Engine/Source/Programs/Horde/Horde.Build/Perforce/PerforceServiceCache.cs
+++ b/Engine/Source/Programs/Horde/Horde.Build/Perforce/PerforceServiceCache.cs
@@ -1,10 +1,8 @@
 // Copyright Epic Games, Inc. All Rights Reserved.
 
 using System;
-using System.Collections.Concurrent;
 using System.Collections.Generic;
 using System.Diagnostics;
-using System.IO;
 using System.Linq;
 using System.Runtime.CompilerServices;
 using System.Threading;
@@ -12,6 +10,7 @@ using System.Threading.Tasks;
 using EpicGames.Core;
 using EpicGames.Perforce;
 using EpicGames.Redis;
+using Horde.Build.Projects;
 using Horde.Build.Server;
 using Horde.Build.Streams;
 using Horde.Build.Users;
@@ -65,18 +64,18 @@ namespace Horde.Build.Perforce
 
 		class StreamInfo
 		{
-			public IStream Stream { get; set; }
+			public StreamConfig StreamConfig { get; set; }
 			public PerforceViewMap View { get; }
 			public List<CommitTagInfo> CommitTags { get; set; } = new List<CommitTagInfo>();
 
-			public StreamInfo(IStream stream, PerforceViewMap view)
+			public StreamInfo(StreamConfig streamConfig, PerforceViewMap view)
 			{
-				Stream = stream;
+				StreamConfig = streamConfig;
 				View = view;
 
-				foreach (CommitTagConfig commitTagConfig in stream.Config.GetAllCommitTags())
+				foreach (CommitTagConfig commitTagConfig in streamConfig.GetAllCommitTags())
 				{
-					if (stream.Config.TryGetCommitTagFilter(commitTagConfig.Name, out FileFilter? filter))
+					if (streamConfig.TryGetCommitTagFilter(commitTagConfig.Name, out FileFilter? filter))
 					{
 						CommitTags.Add(new CommitTagInfo(commitTagConfig.Name, filter));
 					}
@@ -103,7 +102,7 @@ namespace Horde.Build.Perforce
 			PerforceServiceCache _owner = null!;
 
 			[BsonIgnore]
-			IStream _stream = null!;
+			StreamConfig _streamConfig = null!;
 
 			[BsonIgnoreIfDefault] // Allow upserts
 			public ObjectId Id { get; set; }
@@ -148,10 +147,10 @@ namespace Horde.Build.Perforce
 				return new CachedCommitDoc(commit, commitTags);
 			}
 
-			public void PostLoad(PerforceServiceCache owner, IStream stream)
+			public void PostLoad(PerforceServiceCache owner, StreamConfig streamConfig)
 			{
 				_owner = owner;
-				_stream = stream;
+				_streamConfig = streamConfig;
 			}
 
 			public ValueTask<IReadOnlyList<CommitTag>> GetTagsAsync(CancellationToken cancellationToken)
@@ -161,7 +160,7 @@ namespace Horde.Build.Perforce
 
 			public async ValueTask<IReadOnlyList<string>> GetFilesAsync(CancellationToken cancellationToken)
 			{
-				ICommit? other = await _owner!.GetChangeDetailsAsync(_stream, Number, cancellationToken);
+				ICommit? other = await _owner!.GetChangeDetailsAsync(_streamConfig, Number, cancellationToken);
 				if (other == null)
 				{
 					return Array.Empty<string>();
@@ -174,7 +173,7 @@ namespace Horde.Build.Perforce
 		readonly RedisService _redisService;
 		readonly IDowntimeService _downtimeService;
 		readonly IMongoCollection<CachedCommitDoc> _commits;
-		readonly IStreamCollection _streamCollection;
+		readonly IOptionsMonitor<GlobalConfig> _globalConfig;
 		readonly ILogger _logger;
 
 		static readonly RedisChannel<StreamId> s_commitUpdateChannel = new RedisChannel<StreamId>("commit-update");
@@ -184,8 +183,8 @@ namespace Horde.Build.Perforce
 		/// <summary>
 		/// Constructor
 		/// </summary>
-		public PerforceServiceCache(PerforceLoadBalancer loadBalancer, MongoService mongoService, RedisService redisService, IDowntimeService downtimeService, GlobalsService globalsService, IUserCollection userCollection, IStreamCollection streamCollection, IClock clock, IOptions<ServerSettings> settings, ILogger<PerforceService> logger)
-			: base(loadBalancer, globalsService, userCollection, settings, logger)
+		public PerforceServiceCache(PerforceLoadBalancer loadBalancer, MongoService mongoService, RedisService redisService, IDowntimeService downtimeService, IUserCollection userCollection, IClock clock, IOptions<ServerSettings> settings, IOptionsMonitor<GlobalConfig> globalConfig, ILogger<PerforceService> logger)
+			: base(loadBalancer, userCollection, settings, globalConfig, logger)
 		{
 			_mongoService = mongoService;
 			_redisService = redisService;
@@ -196,7 +195,7 @@ namespace Horde.Build.Perforce
 			indexes.Add(MongoIndex.Create<CachedCommitDoc>(keys => keys.Ascending(x => x.StreamId).Ascending(x => x.CommitTags).Descending(x => x.Number), true));
 			_commits = mongoService.GetCollection<CachedCommitDoc>("CommitsV2", indexes);
 
-			_streamCollection = streamCollection;
+			_globalConfig = globalConfig;
 			_logger = logger;
 
 			_updateCommitsTicker = clock.AddSharedTicker<PerforceServiceCache>(TimeSpan.FromSeconds(10.0), UpdateCommitsAsync, logger);
@@ -340,10 +339,10 @@ namespace Horde.Build.Perforce
 
 		async Task<Dictionary<string, List<StreamInfo>>> CreateStreamInfoAsync(CancellationToken cancellationToken)
 		{
-			List<IStream> streams = await _streamCollection.FindAllAsync();
+			IReadOnlyList<StreamConfig> streams = _globalConfig.CurrentValue.Streams;
 
 			Dictionary<string, List<StreamInfo>> clusters = new Dictionary<string, List<StreamInfo>>(StringComparer.OrdinalIgnoreCase);
-			foreach (IGrouping<string, IStream> group in streams.GroupBy(x => x.Config.ClusterName, StringComparer.OrdinalIgnoreCase))
+			foreach (IGrouping<string, StreamConfig> group in streams.GroupBy(x => x.ClusterName, StringComparer.OrdinalIgnoreCase))
 			{
 				List<StreamInfo> streamInfoList = await CreateStreamInfoForClusterAsync(group.Key, group, cancellationToken);
 				clusters[group.Key] = streamInfoList;
@@ -352,12 +351,12 @@ namespace Horde.Build.Perforce
 			return clusters;
 		}
 
-		async Task<List<StreamInfo>> CreateStreamInfoForClusterAsync(string clusterName, IEnumerable<IStream> streams, CancellationToken cancellationToken)
+		async Task<List<StreamInfo>> CreateStreamInfoForClusterAsync(string clusterName, IEnumerable<StreamConfig> streams, CancellationToken cancellationToken)
 		{
 			using (IPooledPerforceConnection perforce = await ConnectAsync(clusterName, null, cancellationToken))
 			{
 				List<StreamInfo> streamInfoList = new List<StreamInfo>();
-				foreach (IStream stream in streams)
+				foreach (StreamConfig stream in streams)
 				{
 					StreamRecord record = await perforce.GetStreamAsync(stream.Name, true, cancellationToken);
 					PerforceViewMap view = PerforceViewMap.Parse(record.View);
@@ -433,8 +432,8 @@ namespace Horde.Build.Perforce
 
 						if (files.Count > 0)
 						{
-							ICommit commit = await CreateCommitAsync(perforce, streamInfo.Stream, describeRecord, info, cancellationToken);
-							_logger.LogInformation("Replicating {StreamId} commit {Change}", streamInfo.Stream.Id, describeRecord.Number);
+							ICommit commit = await CreateCommitAsync(perforce, streamInfo.StreamConfig, describeRecord, info, cancellationToken);
+							_logger.LogInformation("Replicating {StreamId} commit {Change}", streamInfo.StreamConfig.Id, describeRecord.Number);
 
 							List<CommitTag> commitTags = new List<CommitTag>();
 							foreach (CommitTagInfo commitTagInfo in streamInfo.CommitTags)
@@ -448,7 +447,7 @@ namespace Horde.Build.Perforce
 							CachedCommitDoc commitDoc = new CachedCommitDoc(commit, commitTags);
 							await AddCachedCommitAsync(commitDoc, cancellationToken);
 
-							await _redisService.PublishAsync(s_commitUpdateChannel, streamInfo.Stream.Id, CommandFlags.FireAndForget);
+							await _redisService.PublishAsync(s_commitUpdateChannel, streamInfo.StreamConfig.Id, CommandFlags.FireAndForget);
 						}
 					}
 				}
@@ -458,11 +457,11 @@ namespace Horde.Build.Perforce
 				foreach (StreamInfo streamInfo in streamInfos)
 				{
 					int minChange;
-					if (reset || !state.MinChanges.TryGetValue(streamInfo.Stream.Id, out minChange))
+					if (reset || !state.MinChanges.TryGetValue(streamInfo.StreamConfig.Id, out minChange))
 					{
 						minChange = changes[^1].Number;
 					}
-					minChanges.Add(streamInfo.Stream.Id, minChange);
+					minChanges.Add(streamInfo.StreamConfig.Id, minChange);
 				}
 				state.MinChanges = minChanges;
 				state.MaxChange = changes[0].Number;
@@ -487,7 +486,7 @@ namespace Horde.Build.Perforce
 		{
 			readonly PerforceServiceCache _owner;
 
-			public CachedCommitSource(PerforceServiceCache owner, IStream stream)
+			public CachedCommitSource(PerforceServiceCache owner, StreamConfig stream)
 				: base(owner, stream)
 			{
 				_owner = owner;
@@ -501,15 +500,15 @@ namespace Horde.Build.Perforce
 				}
 
 				int numResults = 0;
-				_owner._logger.LogDebug("Querying Perforce cache for {StreamId} commits from {MinChange} to {MaxChange} (max: {MaxResults}, tags: {Tags})", _stream.Id, minChange ?? -2, maxChange ?? -2, maxResults ?? -1, (tags == null || tags.Count == 0) ? "none" : String.Join("/", tags.Select(x => x.ToString())));
+				_owner._logger.LogDebug("Querying Perforce cache for {StreamId} commits from {MinChange} to {MaxChange} (max: {MaxResults}, tags: {Tags})", _streamConfig.Id, minChange ?? -2, maxChange ?? -2, maxResults ?? -1, (tags == null || tags.Count == 0) ? "none" : String.Join("/", tags.Select(x => x.ToString())));
 
 				CacheState state = await _owner._mongoService.GetSingletonAsync<CacheState>();
-				if (state.Clusters.TryGetValue(_stream.Config.ClusterName, out ClusterState? clusterState))
+				if (state.Clusters.TryGetValue(_streamConfig.ClusterName, out ClusterState? clusterState))
 				{
 					int minReplicatedChange;
-					if (clusterState.MinChanges.TryGetValue(_stream.Id, out minReplicatedChange) && (maxChange == null || maxChange > minReplicatedChange))
+					if (clusterState.MinChanges.TryGetValue(_streamConfig.Id, out minReplicatedChange) && (maxChange == null || maxChange > minReplicatedChange))
 					{
-						FilterDefinition<CachedCommitDoc> filter = Builders<CachedCommitDoc>.Filter.Eq(x => x.StreamId, _stream.Id);
+						FilterDefinition<CachedCommitDoc> filter = Builders<CachedCommitDoc>.Filter.Eq(x => x.StreamId, _streamConfig.Id);
 
 						if (tags != null && tags.Count > 0)
 						{
@@ -547,7 +546,7 @@ namespace Horde.Build.Perforce
 							{
 								foreach (CachedCommitDoc commit in cursor.Current)
 								{
-									commit.PostLoad(_owner, _stream);
+									commit.PostLoad(_owner, _streamConfig);
 									yield return commit;
 									numResults++;
 								}
@@ -575,8 +574,8 @@ namespace Horde.Build.Perforce
 							{
 								CachedCommitDoc cachedCommit = await CachedCommitDoc.FromCommitAsync(commit, cancellationToken);
 								await _owner.AddCachedCommitAsync(cachedCommit, cancellationToken);
-								await _owner._mongoService.UpdateSingletonAsync<CacheState>(x => TryUpdateRange(x, _stream, commit.Number, maxChange));
-								_owner._logger.LogDebug("Adding new cached commit for {StreamId} at change {Change}", _stream.Id, commit.Number);
+								await _owner._mongoService.UpdateSingletonAsync<CacheState>(x => TryUpdateRange(x, _streamConfig, commit.Number, maxChange));
+								_owner._logger.LogDebug("Adding new cached commit for {StreamId} at change {Change}", _streamConfig.Id, commit.Number);
 
 								if (tags == null || tags.Any(x => cachedCommit.CommitTags.Contains(x)))
 								{
@@ -589,8 +588,8 @@ namespace Horde.Build.Perforce
 							if (maxResults == null || maxResults.Value > 0)
 							{
 								int newMinChange = minChange ?? 0;
-								_owner._logger.LogDebug("Extending range for {StreamId} cache to {Change}..", _stream.Id, newMinChange);
-								await _owner._mongoService.UpdateSingletonAsync<CacheState>(x => TryUpdateRange(x, _stream, newMinChange, maxChange));
+								_owner._logger.LogDebug("Extending range for {StreamId} cache to {Change}..", _streamConfig.Id, newMinChange);
+								await _owner._mongoService.UpdateSingletonAsync<CacheState>(x => TryUpdateRange(x, _streamConfig, newMinChange, maxChange));
 							}
 						}
 					}
@@ -598,7 +597,7 @@ namespace Horde.Build.Perforce
 
 				if (maxResults == null || maxResults.Value > 0)
 				{
-					_owner._logger.LogDebug("Querying Perforce server for {StreamId} commits from {MinChange} to {MaxChange} (max: {MaxResults}, tags: {Tags})", _stream.Id, minChange ?? -2, maxChange ?? -2, maxResults ?? -1, (tags == null || tags.Count == 0) ? "none" : String.Join("/", tags.Select(x => x.ToString())));
+					_owner._logger.LogDebug("Querying Perforce server for {StreamId} commits from {MinChange} to {MaxChange} (max: {MaxResults}, tags: {Tags})", _streamConfig.Id, minChange ?? -2, maxChange ?? -2, maxResults ?? -1, (tags == null || tags.Count == 0) ? "none" : String.Join("/", tags.Select(x => x.ToString())));
 
 					await foreach (ICommit commit in base.FindAsync(minChange, maxChange, maxResults, tags, cancellationToken))
 					{
@@ -607,17 +606,17 @@ namespace Horde.Build.Perforce
 				}
 			}
 
-			static bool TryUpdateRange(CacheState state, IStream stream, int newMinChange, int? maxChange)
+			static bool TryUpdateRange(CacheState state, StreamConfig streamConfig, int newMinChange, int? maxChange)
 			{
 				ClusterState? clusterState;
-				if (state.Clusters.TryGetValue(stream.Config.ClusterName, out clusterState))
+				if (state.Clusters.TryGetValue(streamConfig.ClusterName, out clusterState))
 				{
 					int streamMinChange;
-					if (clusterState.MinChanges.TryGetValue(stream.Id, out streamMinChange))
+					if (clusterState.MinChanges.TryGetValue(streamConfig.Id, out streamMinChange))
 					{
 						if (newMinChange < streamMinChange && (maxChange == null || maxChange.Value >= streamMinChange - 1))
 						{
-							clusterState.MinChanges[stream.Id] = newMinChange;
+							clusterState.MinChanges[streamConfig.Id] = newMinChange;
 							return true;
 						}
 					}
@@ -627,10 +626,10 @@ namespace Horde.Build.Perforce
 
 			public override async Task<ICommit> GetAsync(int changeNumber, CancellationToken cancellationToken = default)
 			{
-				CachedCommitDoc? commit = await _owner._commits.Find(x => x.StreamId == _stream.Id && x.Number == changeNumber).FirstOrDefaultAsync(cancellationToken);
+				CachedCommitDoc? commit = await _owner._commits.Find(x => x.StreamId == _streamConfig.Id && x.Number == changeNumber).FirstOrDefaultAsync(cancellationToken);
 				if(commit != null)
 				{
-					commit.PostLoad(_owner, _stream);
+					commit.PostLoad(_owner, _streamConfig);
 					return commit;
 				}
 				return await base.GetAsync(changeNumber, cancellationToken);
@@ -643,7 +642,7 @@ namespace Horde.Build.Perforce
 
 				void OnUpdate(StreamId streamId)
 				{
-					if (streamId == _stream.Id)
+					if (streamId == _streamConfig.Id)
 					{
 						updateEvent.Set();
 					}
@@ -692,9 +691,9 @@ namespace Horde.Build.Perforce
 		}
 
 		/// <inheritdoc/>
-		public override ICommitCollection GetCommits(IStream stream)
+		public override ICommitCollection GetCommits(StreamConfig streamConfig)
 		{
-			return new CachedCommitSource(this, stream);
+			return new CachedCommitSource(this, streamConfig);
 		}
 
 		#endregion
