Horde: Refactor codebase to use new config system.

diff --git a/Engine/Source/Programs/Horde/Horde.Build/Perforce/ReplicationService.cs b/Engine/Source/Programs/Horde/Horde.Build/Perforce/ReplicationService.cs
--- a/Engine/Source/Programs/Horde/Horde.Build/Perforce/ReplicationService.cs
+++ b/Engine/Source/Programs/Horde/Horde.Build/Perforce/ReplicationService.cs
@@ -9,6 +9,7 @@ using System.Threading.Tasks;
 using EpicGames.Core;
 using EpicGames.Horde.Storage;
 using EpicGames.Horde.Storage.Nodes;
+using Horde.Build.Server;
 using Horde.Build.Storage;
 using Horde.Build.Streams;
 using Horde.Build.Utilities;
@@ -16,6 +17,7 @@ using HordeCommon;
 using Microsoft.Extensions.Caching.Memory;
 using Microsoft.Extensions.Hosting;
 using Microsoft.Extensions.Logging;
+using Microsoft.Extensions.Options;
 
 namespace Horde.Build.Perforce
 {
@@ -42,12 +44,13 @@ namespace Horde.Build.Perforce
 		readonly StorageService _storageService;
 		readonly IMemoryCache _memoryCache;
 		readonly ITicker _ticker;
+		readonly IOptionsMonitor<GlobalConfig> _globalConfig;
 		readonly ILogger _logger;
 
 		/// <summary>
 		/// Constructor
 		/// </summary>
-		public ReplicationService(IStreamCollection streamCollection, IPerforceService perforceService, PerforceReplicator replicator, StorageService storageService, IMemoryCache memoryCache, IClock clock, ILogger<ReplicationService> logger)
+		public ReplicationService(IStreamCollection streamCollection, IPerforceService perforceService, PerforceReplicator replicator, StorageService storageService, IMemoryCache memoryCache, IClock clock, IOptionsMonitor<GlobalConfig> globalConfig, ILogger<ReplicationService> logger)
 		{
 			_streamCollection = streamCollection;
 			_perforceService = perforceService;
@@ -55,6 +58,7 @@ namespace Horde.Build.Perforce
 			_storageService = storageService;
 			_memoryCache = memoryCache;
 			_ticker = clock.AddSharedTicker<ReplicationService>(TimeSpan.FromSeconds(20.0), TickSharedAsync, logger);
+			_globalConfig = globalConfig;
 			_logger = logger;
 		}
 
@@ -77,18 +81,18 @@ namespace Horde.Build.Perforce
 			{
 				for (; ; )
 				{
-					List<IStream> streams = await _streamCollection.FindAllAsync();
+					IReadOnlyList<StreamConfig> streamConfigs = _globalConfig.CurrentValue.Streams;
 
 					HashSet<StreamId> removeStreams = new HashSet<StreamId>(streamIdToTask.Keys);
-					foreach (IStream stream in streams)
+					foreach (StreamConfig streamConfig in streamConfigs)
 					{
-						if (stream.Config.ReplicationMode == ContentReplicationMode.Full)
+						if (streamConfig.ReplicationMode == ContentReplicationMode.Full)
 						{
-							removeStreams.Remove(stream.Id);
-							if (!streamIdToTask.ContainsKey(stream.Id))
+							removeStreams.Remove(streamConfig.Id);
+							if (!streamIdToTask.ContainsKey(streamConfig.Id))
 							{
-								streamIdToTask.Add(stream.Id, BackgroundTask.StartNew(ctx => RunReplicationGuardedAsync(stream, ctx)));
-								_logger.LogInformation("Started replication of {StreamId}", stream.Id);
+								streamIdToTask.Add(streamConfig.Id, BackgroundTask.StartNew(ctx => RunReplicationGuardedAsync(streamConfig, ctx)));
+								_logger.LogInformation("Started replication of {StreamId}", streamConfig.Id);
 							}
 						}
 					}
@@ -111,13 +115,13 @@ namespace Horde.Build.Perforce
 			}
 		}
 
-		async Task RunReplicationGuardedAsync(IStream stream, CancellationToken cancellationToken)
+		async Task RunReplicationGuardedAsync(StreamConfig streamConfig, CancellationToken cancellationToken)
 		{
 			for (; ; )
 			{
 				try
 				{
-					await RunReplicationAsync(stream, cancellationToken);
+					await RunReplicationAsync(streamConfig, cancellationToken);
 				}
 				catch (OperationCanceledException) when (cancellationToken.IsCancellationRequested)
 				{
@@ -125,22 +129,22 @@ namespace Horde.Build.Perforce
 				}
 				catch (Exception ex)
 				{
-					_logger.LogError(ex, "Exception running replication for {StreamId}: {Message}", stream.Id, ex.Message);
+					_logger.LogError(ex, "Exception running replication for {StreamId}: {Message}", streamConfig.Id, ex.Message);
 				}
 
 				await Task.Delay(TimeSpan.FromSeconds(20.0), cancellationToken);
 			}
 		}
 
-		async Task RunReplicationAsync(IStream stream, CancellationToken cancellationToken)
+		async Task RunReplicationAsync(StreamConfig streamConfig, CancellationToken cancellationToken)
 		{
-			RefName refName = new RefName(stream.Id.ToString());
+			RefName refName = new RefName(streamConfig.Id.ToString());
 
 			IStorageClientImpl store = await _storageService.GetClientAsync(Namespace.Perforce, cancellationToken);
 			TreeReader reader = new TreeReader(store, _memoryCache, _logger);
 
 			CommitNode? lastCommitNode = await reader.TryReadNodeAsync<CommitNode>(refName, cancellationToken: cancellationToken);
-			ICommitCollection commits = _perforceService.GetCommits(stream);
+			ICommitCollection commits = _perforceService.GetCommits(streamConfig);
 
 			PerforceReplicationOptions options = new PerforceReplicationOptions();
 
@@ -156,8 +160,8 @@ namespace Horde.Build.Perforce
 
 			for (; ; )
 			{
-				_logger.LogInformation("Replicating {StreamId} change {Change}", stream.Id, commit.Number);
-				await _replicator.WriteAsync(stream, commit.Number, options, cancellationToken);
+				_logger.LogInformation("Replicating {StreamId} change {Change}", streamConfig.Id, commit.Number);
+				await _replicator.WriteAsync(streamConfig, commit.Number, options, cancellationToken);
 				commit = await commits.SubscribeAsync(commit.Number, cancellationToken: cancellationToken).FirstAsync(cancellationToken);
 			}
 		}
diff --git a/Engine/Source/Programs/Horde/Horde.Build/Projects/IProject.cs b/Engine/Source/Programs/Horde/Horde.Build/Projects/IProject.cs
--- a/Engine/Source/Programs/Horde/Horde.Build/Projects/IProject.cs
+++ b/Engine/Source/Programs/Horde/Horde.Build/Projects/IProject.cs
@@ -1,12 +1,6 @@
 // Copyright Epic Games, Inc. All Rights Reserved.
 
-using System.Collections.Generic;
-using System.Linq;
-using System.Text.RegularExpressions;
-using Horde.Build.Acls;
-using Horde.Build.Streams;
 using Horde.Build.Utilities;
-using MongoDB.Driver;
 
 namespace Horde.Build.Projects
 {
@@ -22,30 +16,10 @@ namespace Horde.Build.Projects
 		/// </summary>
 		ProjectId Id { get; }
 
-		/// <summary>
-		/// Name of the project
-		/// </summary>
-		string Name { get; }
-
-		/// <summary>
-		/// Revision of the config file used to configure this project
-		/// </summary>
-		string? ConfigRevision { get; }
-
-		/// <summary>
-		/// Order to display on the dashboard
-		/// </summary>
-		public int Order { get; }
-
 		/// <summary>
 		/// Configuration settings for the stream
 		/// </summary>
 		ProjectConfig Config { get; }
-
-		/// <summary>
-		/// The ACL for this object
-		/// </summary>
-		Acl? Acl { get; }
 	}
 	
 	/// <summary>
@@ -78,90 +52,4 @@ namespace Horde.Build.Projects
 		/// </summary>
 		public byte[] Data { get; }
 	}
-
-	/// <summary>
-	/// Extension methods for projects
-	/// </summary>
-	public static class ProjectExtensions
-	{
-		/// <summary>
-		/// Converts this object to a public response
-		/// </summary>
-		/// <param name="project">The project instance</param>
-		/// <param name="includeStreams">Whether to include streams in the response</param>
-		/// <param name="includeCategories">Whether to include categories in the response</param>
-		/// <param name="streams">The list of streams</param>
-		/// <param name="includeAcl">Whether to include the ACL in the response</param>
-		/// <returns>Response instance</returns>
-		public static GetProjectResponse ToResponse(this IProject project, bool includeStreams, bool includeCategories, List<IStream>? streams, bool includeAcl)
-		{
-			List<GetProjectStreamResponse>? streamResponses = null;
-			if(includeStreams)
-			{
-				streamResponses = streams!.ConvertAll(x => new GetProjectStreamResponse(x.Id.ToString(), x.Name));
-			}
-
-			List<GetProjectCategoryResponse>? categoryResponses = null;
-			if (includeCategories)
-			{
-				categoryResponses = project.Config.Categories.ConvertAll(x => new GetProjectCategoryResponse(x));
-				if (streams != null)
-				{
-					foreach (IStream stream in streams)
-					{
-						GetProjectCategoryResponse? categoryResponse = categoryResponses.FirstOrDefault(x => MatchCategory(stream.Name, x));
-						if(categoryResponse == null)
-						{
-							int row = (categoryResponses.Count > 0) ? categoryResponses.Max(x => x.Row) : 0;
-							if (categoryResponses.Count(x => x.Row == row) >= 3)
-							{
-								row++;
-							}
-
-							ProjectCategoryConfig otherCategory = new ProjectCategoryConfig();
-							otherCategory.Name = "Other";
-							otherCategory.Row = row;
-							otherCategory.IncludePatterns.Add(".*");
-
-							categoryResponse = new GetProjectCategoryResponse(otherCategory);
-							categoryResponses.Add(categoryResponse);
-						}
-						categoryResponse.Streams!.Add(stream.Id.ToString());
-					}
-				}
-			}
-
-			GetAclResponse? aclResponse = (includeAcl && project.Acl != null) ? new GetAclResponse(project.Acl) : null;
-			return new GetProjectResponse(project.Id.ToString(), project.Name, project.Order, streamResponses, categoryResponses, aclResponse);
-		}
-
-		/// <summary>
-		/// Tests if a category response matches a given stream name
-		/// </summary>
-		/// <param name="name">The stream name</param>
-		/// <param name="category">The category response</param>
-		/// <returns>True if the category matches</returns>
-		static bool MatchCategory(string name, GetProjectCategoryResponse category)
-		{
-			if (category.IncludePatterns.Any(x => Regex.IsMatch(name, x)))
-			{
-				if (!category.ExcludePatterns.Any(x => Regex.IsMatch(name, x)))
-				{
-					return true;
-				}
-			}
-			return false;
-		}
-	}
-
-	/// <summary>
-	/// Projection of a project definition to just include permissions info
-	/// </summary>
-	public interface IProjectPermissions
-	{
-		/// <summary>
-		/// ACL for the project
-		/// </summary>
-		public Acl? Acl { get; }
-	}
 }
diff --git a/Engine/Source/Programs/Horde/Horde.Build/Projects/IProjectCollection.cs b/Engine/Source/Programs/Horde/Horde.Build/Projects/IProjectCollection.cs
--- a/Engine/Source/Programs/Horde/Horde.Build/Projects/IProjectCollection.cs
+++ b/Engine/Source/Programs/Horde/Horde.Build/Projects/IProjectCollection.cs
@@ -1,6 +1,5 @@
 // Copyright Epic Games, Inc. All Rights Reserved.
 
-using System.Collections.Generic;
 using System.Threading.Tasks;
 using Horde.Build.Utilities;
 
@@ -13,28 +12,6 @@ namespace Horde.Build.Projects
 	/// </summary>
 	public interface IProjectCollection
 	{
-		/// <summary>
-		/// Updates the project configuration
-		/// </summary>
-		/// <param name="id">The project id</param>
-		/// <param name="configRevision">The config file revision. Can be referenced withing the config collection.</param>
-		/// <param name="order">Order of the project</param>
-		/// <returns>New project instance</returns>
-		Task<IProject?> AddOrUpdateAsync(ProjectId id, string configRevision, int order);
-
-		/// <summary>
-		/// Gets all the available projects
-		/// </summary>
-		/// <returns>List of project documents</returns>
-		Task<List<IProject>> FindAllAsync();
-
-		/// <summary>
-		/// Gets a project by ID
-		/// </summary>
-		/// <param name="projectId">Unique id of the project</param>
-		/// <returns>The project document</returns>
-		Task<IProject?> GetAsync(ProjectId projectId);
-
 		/// <summary>
 		/// Gets the logo for a project
 		/// </summary>
@@ -52,19 +29,5 @@ namespace Horde.Build.Projects
 		/// <param name="data"></param>
 		/// <returns></returns>
 		Task SetLogoAsync(ProjectId projectId, string path, string revision, string mimeType, byte[] data);
-
-		/// <summary>
-		/// Gets a project's permissions info by ID
-		/// </summary>
-		/// <param name="projectId">Unique id of the project</param>
-		/// <returns>The project document</returns>
-		Task<IProjectPermissions?> GetPermissionsAsync(ProjectId projectId);
-
-		/// <summary>
-		/// Deletes a project by id
-		/// </summary>
-		/// <param name="projectId">Unique id of the project</param>
-		/// <returns>True if the project was deleted</returns>
-		Task DeleteAsync(ProjectId projectId);
 	}
 }
diff --git a/Engine/Source/Programs/Horde/Horde.Build/Projects/ProjectCollection.cs b/Engine/Source/Programs/Horde/Horde.Build/Projects/ProjectCollection.cs
--- a/Engine/Source/Programs/Horde/Horde.Build/Projects/ProjectCollection.cs
+++ b/Engine/Source/Programs/Horde/Horde.Build/Projects/ProjectCollection.cs
@@ -2,15 +2,12 @@
 
 using System;
 using System.Collections.Generic;
-using System.Diagnostics.CodeAnalysis;
 using System.Linq;
 using System.Threading.Tasks;
-using Horde.Build.Acls;
-using Horde.Build.Configuration;
 using Horde.Build.Server;
 using Horde.Build.Utilities;
 using Microsoft.Extensions.Logging;
-using MongoDB.Bson.Serialization.Attributes;
+using Microsoft.Extensions.Options;
 using MongoDB.Driver;
 
 namespace Horde.Build.Projects
@@ -22,56 +19,6 @@ namespace Horde.Build.Projects
 	/// </summary>
 	public class ProjectCollection : IProjectCollection
 	{
-		/// <summary>
-		/// Represents a project
-		/// </summary>
-		class ProjectDocument : IProject
-		{
-			public const int DefaultOrder = 128;
-
-			[BsonRequired, BsonId]
-			public ProjectId Id { get; set; }
-
-			[BsonRequired]
-			public string ConfigRevision { get; set; } = null!;
-
-			public int Order { get; set; } = DefaultOrder;
-			public Acl? Acl { get; set; }
-
-			[BsonIgnoreIfDefault, BsonDefaultValue(false)]
-			public bool Deleted { get; set; }
-
-			#region IProject interface
-
-			[BsonIgnore]
-			public string Name => Config?.Name!;
-
-			[BsonIgnore]
-			public ProjectConfig Config { get; private set; } = null!;
-
-			public async Task PostLoadAsync(ConfigCollection configCollection, ILogger logger)
-			{
-				try
-				{
-					if (Deleted)
-					{
-						Config = new ProjectConfig();
-					}
-					else
-					{
-						Config = await configCollection.GetConfigAsync<ProjectConfig>(ConfigRevision);
-					}
-				}
-				catch (Exception)
-				{
-					logger.LogError("Unable to get stream config for {StreamId} at {Revision}", Id, ConfigRevision);
-					throw;
-				}
-			}
-
-			#endregion
-		}
-
 		/// <summary>
 		/// Logo for a project
 		/// </summary>
@@ -86,75 +33,14 @@ namespace Horde.Build.Projects
 			public byte[] Data { get; set; } = Array.Empty<byte>();
 		}
 
-		/// <summary>
-		/// Projection of a project definition to just include permissions info
-		/// </summary>
-		[SuppressMessage("Design", "CA1812: Class is never instantiated")]
-		class ProjectPermissions : IProjectPermissions
-		{
-			/// <summary>
-			/// ACL for the project
-			/// </summary>
-			public Acl? Acl { get; set; }
-
-			/// <summary>
-			/// Projection to extract the permissions info from the project
-			/// </summary>
-			public static readonly ProjectionDefinition<ProjectDocument> Projection = Builders<ProjectDocument>.Projection.Include(x => x.Acl);
-		}
-
-		readonly IMongoCollection<ProjectDocument> _projects;
 		readonly IMongoCollection<ProjectLogoDocument> _projectLogos;
-		readonly ConfigCollection _configCollection;
-		readonly ILogger _logger;
 
 		/// <summary>
 		/// Constructor
 		/// </summary>
-		public ProjectCollection(MongoService mongoService, ConfigCollection configCollection, ILogger<ProjectCollection> logger)
+		public ProjectCollection(MongoService mongoService)
 		{
-			_projects = mongoService.GetCollection<ProjectDocument>("Projects");
 			_projectLogos = mongoService.GetCollection<ProjectLogoDocument>("ProjectLogos");
-			_configCollection = configCollection;
-			_logger = logger;
-		}
-
-		/// <inheritdoc/>
-		public async Task<IProject?> AddOrUpdateAsync(ProjectId id, string configRevision, int order)
-		{
-			ProjectConfig config = await _configCollection.GetConfigAsync<ProjectConfig>(configRevision);
-
-			ProjectDocument newProject = new ProjectDocument();
-			newProject.Id = id;
-			newProject.ConfigRevision = configRevision;
-			newProject.Order = order;
-			newProject.Acl = Acl.Merge(new Acl(), config.Acl);
-
-			ProjectDocument result = await _projects.FindOneAndReplaceAsync<ProjectDocument>(x => x.Id == id, newProject, new FindOneAndReplaceOptions<ProjectDocument> { IsUpsert = true, ReturnDocument = ReturnDocument.After });
-			await result.PostLoadAsync(_configCollection, _logger);
-			return result;
-		}
-
-		/// <inheritdoc/>
-		public async Task<List<IProject>> FindAllAsync()
-		{
-			List<ProjectDocument> results = await _projects.Find(x => !x.Deleted).ToListAsync();
-			foreach (ProjectDocument result in results)
-			{
-				await result.PostLoadAsync(_configCollection, _logger);
-			}
-			return results.OrderBy(x => x.Order).ThenBy(x => x.Name).Select<ProjectDocument, IProject>(x => x).ToList();
-		}
-
-		/// <inheritdoc/>
-		public async Task<IProject?> GetAsync(ProjectId projectId)
-		{
-			ProjectDocument? project = await _projects.Find<ProjectDocument>(x => x.Id == projectId).FirstOrDefaultAsync();
-			if (project != null)
-			{
-				await project.PostLoadAsync(_configCollection, _logger);
-			}
-			return project;
 		}
 
 		/// <inheritdoc/>
@@ -174,17 +60,5 @@ namespace Horde.Build.Projects
 			logo.Data = data;
 			await _projectLogos.ReplaceOneAsync(x => x.Id == projectId, logo, new ReplaceOptions { IsUpsert = true });
 		}
-
-		/// <inheritdoc/>
-		public async Task<IProjectPermissions?> GetPermissionsAsync(ProjectId projectId)
-		{
-			return await _projects.Find<ProjectDocument>(x => x.Id == projectId).Project<ProjectPermissions>(ProjectPermissions.Projection).FirstOrDefaultAsync();
-		}
-
-		/// <inheritdoc/>
-		public async Task DeleteAsync(ProjectId projectId)
-		{
-			await _projects.UpdateOneAsync(x => x.Id == projectId, Builders<ProjectDocument>.Update.Set(x => x.Deleted, true));
-		}
 	}
 }
diff --git a/Engine/Source/Programs/Horde/Horde.Build/Projects/ProjectConfig.cs b/Engine/Source/Programs/Horde/Horde.Build/Projects/ProjectConfig.cs
--- a/Engine/Source/Programs/Horde/Horde.Build/Projects/ProjectConfig.cs
+++ b/Engine/Source/Programs/Horde/Horde.Build/Projects/ProjectConfig.cs
@@ -4,31 +4,61 @@ using System;
 using System.Collections.Generic;
 using System.ComponentModel.DataAnnotations;
 using System.Diagnostics;
+using System.Security.Claims;
 using System.Text.Json.Serialization;
 using EpicGames.Core;
 using Horde.Build.Acls;
 using Horde.Build.Agents.Pools;
 using Horde.Build.Configuration;
+using Horde.Build.Server;
 using Horde.Build.Streams;
 using Horde.Build.Utilities;
 
 namespace Horde.Build.Projects
 {
-	using StreamId = StringId<IStream>;
+	using ProjectId = StringId<IProject>;
 
 	/// <summary>
 	/// Stores configuration for a project
 	/// </summary>
 	[JsonSchema("https://unrealengine.com/horde/project")]
 	[JsonSchemaCatalog("Horde Project", "Horde project configuration file", new[] { "*.project.json", "Projects/*.json" })]
+	[ConfigIncludeRoot]
+	[DebuggerDisplay("{Id}")]
 	public class ProjectConfig
 	{
+		/// <summary>
+		/// Accessor for the global config owning this project
+		/// </summary>
+		[JsonIgnore]
+		public GlobalConfig GlobalConfig { get; private set; } = null!;
+
+		/// <summary>
+		/// The project id
+		/// </summary>
+		public ProjectId Id { get; set; }
+
 		/// <summary>
 		/// Name for the new project
 		/// </summary>
-		[Required]
 		public string Name { get; set; } = null!;
 
+		/// <summary>
+		/// Direct include path for the project config. For backwards compatibility with old config files when including from a GlobalConfig object.
+		/// </summary>
+		[ConfigInclude, ConfigRelativePath]
+		public string? Path { get; set; }
+
+		/// <summary>
+		/// Includes for other configuration files
+		/// </summary>
+		public List<ConfigInclude> Include { get; set; } = new List<ConfigInclude>();
+
+		/// <summary>
+		/// Order of this project on the dashboard
+		/// </summary>
+		public int Order { get; set; } = 128;
+
 		/// <summary>
 		/// Path to the project logo
 		/// </summary>
@@ -47,37 +77,38 @@ namespace Horde.Build.Projects
 		/// <summary>
 		/// List of streams
 		/// </summary>
-		public List<StreamConfigRef> Streams { get; set; } = new List<StreamConfigRef>();
+		public List<StreamConfig> Streams { get; set; } = new List<StreamConfig>();
 
 		/// <summary>
 		/// Acl entries
 		/// </summary>
 		public AclConfig? Acl { get; set; }
-	}
-
-	/// <summary>
-	/// Reference to configuration for a stream
-	/// </summary>
-	[DebuggerDisplay("{Id}")]
-	public class StreamConfigRef
-	{
-		/// <summary>
-		/// Unique id for the stream
-		/// </summary>
-		[Required]
-		public StreamId Id { get; set; } = StreamId.Empty;
 
 		/// <summary>
-		/// Path to the configuration file
+		/// Callback after this configuration has been read
 		/// </summary>
-		[Required, ConfigRelativePath]
-		public string Path { get; set; } = null!;
+		/// <param name="id">Id of this project</param>
+		/// <param name="globalConfig">The owning global config object</param>
+		public void PostLoad(ProjectId id, GlobalConfig globalConfig)
+		{
+			Id = id;
+			GlobalConfig = globalConfig;
+
+			foreach (StreamConfig stream in Streams)
+			{
+				stream.PostLoad(stream.Id, this);
+			}
+		}
 
 		/// <summary>
-		/// Target configuration object
+		/// Authorizes a user to perform a given action
 		/// </summary>
-		[JsonIgnore]
-		public StreamConfig Target { get; set; } = null!;
+		/// <param name="action">The action being performed</param>
+		/// <param name="user">The principal to validate</param>
+		public bool Authorize(AclAction action, ClaimsPrincipal user)
+		{
+			return Acl?.Authorize(action, user) ?? GlobalConfig.Authorize(action, user);
+		}
 	}
 
 	/// <summary>
diff --git a/Engine/Source/Programs/Horde/Horde.Build/Projects/ProjectMessages.cs b/Engine/Source/Programs/Horde/Horde.Build/Projects/ProjectMessages.cs
--- a/Engine/Source/Programs/Horde/Horde.Build/Projects/ProjectMessages.cs
+++ b/Engine/Source/Programs/Horde/Horde.Build/Projects/ProjectMessages.cs
@@ -2,7 +2,10 @@
 
 using System.Collections.Generic;
 using System.ComponentModel.DataAnnotations;
+using System.Linq;
+using System.Text.RegularExpressions;
 using Horde.Build.Acls;
+using Horde.Build.Streams;
 
 namespace Horde.Build.Projects
 {
@@ -112,11 +115,6 @@ namespace Horde.Build.Projects
 		/// </summary>
 		public List<GetProjectCategoryResponse>? Categories { get; set; }
 
-		/// <summary>
-		/// Custom permissions for this object
-		/// </summary>
-		public GetAclResponse? Acl { get; set; }
-
 		/// <summary>
 		/// Constructor
 		/// </summary>
@@ -125,15 +123,80 @@ namespace Horde.Build.Projects
 		/// <param name="order">Order to show this project on the dashboard</param>
 		/// <param name="streams">List of streams to display</param>
 		/// <param name="categories">List of stream categories to display</param>
-		/// <param name="acl">Custom permissions for this object</param>
-		public GetProjectResponse(string id, string name, int order, List<GetProjectStreamResponse>? streams, List<GetProjectCategoryResponse>? categories, GetAclResponse? acl)
+		public GetProjectResponse(string id, string name, int order, List<GetProjectStreamResponse>? streams, List<GetProjectCategoryResponse>? categories)
 		{
 			Id = id;
 			Name = name;
 			Order = order;
 			Streams = streams;
 			Categories = categories;
-			Acl = acl;
+		}
+
+		/// <summary>
+		/// Converts this object to a public response
+		/// </summary>
+		/// <param name="projectConfig">The project instance</param>
+		/// <param name="includeStreams">Whether to include streams in the response</param>
+		/// <param name="includeCategories">Whether to include categories in the response</param>
+		/// <param name="streams">The list of streams</param>
+		/// <returns>Response instance</returns>
+		public static GetProjectResponse FromConfig(ProjectConfig projectConfig, bool includeStreams, bool includeCategories, List<StreamConfig>? streams)
+		{
+			List<GetProjectStreamResponse>? streamResponses = null;
+			if (includeStreams)
+			{
+				streamResponses = streams!.ConvertAll(x => new GetProjectStreamResponse(x.Id.ToString(), x.Name));
+			}
+
+			List<GetProjectCategoryResponse>? categoryResponses = null;
+			if (includeCategories)
+			{
+				categoryResponses = projectConfig.Categories.ConvertAll(x => new GetProjectCategoryResponse(x));
+				if (streams != null)
+				{
+					foreach (IStream stream in streams)
+					{
+						GetProjectCategoryResponse? categoryResponse = categoryResponses.FirstOrDefault(x => MatchCategory(stream.Config.Name, x));
+						if (categoryResponse == null)
+						{
+							int row = (categoryResponses.Count > 0) ? categoryResponses.Max(x => x.Row) : 0;
+							if (categoryResponses.Count(x => x.Row == row) >= 3)
+							{
+								row++;
+							}
+
+							ProjectCategoryConfig otherCategory = new ProjectCategoryConfig();
+							otherCategory.Name = "Other";
+							otherCategory.Row = row;
+							otherCategory.IncludePatterns.Add(".*");
+
+							categoryResponse = new GetProjectCategoryResponse(otherCategory);
+							categoryResponses.Add(categoryResponse);
+						}
+						categoryResponse.Streams!.Add(stream.Id.ToString());
+					}
+				}
+			}
+
+			return new GetProjectResponse(projectConfig.Id.ToString(), projectConfig.Name, projectConfig.Order, streamResponses, categoryResponses);
+		}
+
+		/// <summary>
+		/// Tests if a category response matches a given stream name
+		/// </summary>
+		/// <param name="name">The stream name</param>
+		/// <param name="category">The category response</param>
+		/// <returns>True if the category matches</returns>
+		static bool MatchCategory(string name, GetProjectCategoryResponse category)
+		{
+			if (category.IncludePatterns.Any(x => Regex.IsMatch(name, x)))
+			{
+				if (!category.ExcludePatterns.Any(x => Regex.IsMatch(name, x)))
+				{
+					return true;
+				}
+			}
+			return false;
 		}
 	}
 }
diff --git a/Engine/Source/Programs/Horde/Horde.Build/Projects/ProjectsController.cs b/Engine/Source/Programs/Horde/Horde.Build/Projects/ProjectsController.cs
--- a/Engine/Source/Programs/Horde/Horde.Build/Projects/ProjectsController.cs
+++ b/Engine/Source/Programs/Horde/Horde.Build/Projects/ProjectsController.cs
@@ -1,12 +1,15 @@
 // Copyright Epic Games, Inc. All Rights Reserved.
 
 using System.Collections.Generic;
+using System.Linq;
 using System.Threading.Tasks;
 using Horde.Build.Acls;
+using Horde.Build.Server;
 using Horde.Build.Streams;
 using Horde.Build.Utilities;
 using Microsoft.AspNetCore.Authorization;
 using Microsoft.AspNetCore.Mvc;
+using Microsoft.Extensions.Options;
 
 namespace Horde.Build.Projects
 {
@@ -20,25 +23,16 @@ namespace Horde.Build.Projects
 	[Route("[controller]")]
 	public class ProjectsController : HordeControllerBase
 	{
-		/// <summary>
-		/// Singleton instance of the project service
-		/// </summary>
-		private readonly ProjectService _projectService;
-
-		/// <summary>
-		/// Singleton instance of the stream service
-		/// </summary>
-		private readonly StreamService _streamService;
+		private readonly IProjectCollection _projectCollection;
+		private readonly IOptionsSnapshot<GlobalConfig> _globalConfig;
 
 		/// <summary>
 		/// Constructor
 		/// </summary>
-		/// <param name="projectService">The project service</param>
-		/// <param name="streamService">The stream service</param>
-		public ProjectsController(ProjectService projectService, StreamService streamService)
+		public ProjectsController(IProjectCollection projectCollection, IOptionsSnapshot<GlobalConfig> globalConfig)
 		{
-			_projectService = projectService;
-			_streamService = streamService;
+			_projectCollection = projectCollection;
+			_globalConfig = globalConfig;
 		}
 
 		/// <summary>
@@ -51,24 +45,21 @@ namespace Horde.Build.Projects
 		[HttpGet]
 		[Route("/api/v1/projects")]
 		[ProducesResponseType(typeof(List<GetProjectResponse>), 200)]
-		public async Task<ActionResult<List<object>>> GetProjectsAsync([FromQuery(Name = "Streams")] bool includeStreams = false, [FromQuery(Name = "Categories")] bool includeCategories = false, [FromQuery] PropertyFilter? filter = null)
+		public ActionResult<List<object>> GetProjects([FromQuery(Name = "Streams")] bool includeStreams = false, [FromQuery(Name = "Categories")] bool includeCategories = false, [FromQuery] PropertyFilter? filter = null)
 		{
-			List<IProject> projects = await _projectService.GetProjectsAsync();
-			ProjectPermissionsCache permissionsCache = new ProjectPermissionsCache();
-
-			List<IStream>? streams = null;
-			if (includeStreams || includeCategories)
-			{
-				streams = await _streamService.GetStreamsAsync();
-			}
+			GlobalConfig globalConfig = _globalConfig.Value;
 
 			List<object> responses = new List<object>();
-			foreach (IProject project in projects)
+			foreach (ProjectConfig projectConfig in globalConfig.Projects)
 			{
-				if (await _projectService.AuthorizeAsync(project, AclAction.ViewProject, User, permissionsCache))
+				if (projectConfig.Authorize(AclAction.ViewProject, User))
 				{
-					bool includeAcl = await _projectService.AuthorizeAsync(project, AclAction.ViewPermissions, User, permissionsCache);
-					responses.Add(project.ToResponse(includeStreams, includeCategories, streams, includeAcl).ApplyFilter(filter));
+					List<StreamConfig>? visibleStreams = null;
+					if (includeStreams || includeCategories)
+					{
+						visibleStreams = projectConfig.Streams.Where(x => x.Authorize(AclAction.ViewStream, User)).ToList();
+					}
+					responses.Add(GetProjectResponse.FromConfig(projectConfig, includeStreams, includeCategories, visibleStreams).ApplyFilter(filter));
 				}
 			}
 			return responses;
@@ -83,16 +74,14 @@ namespace Horde.Build.Projects
 		[HttpGet]
 		[Route("/api/v1/projects/{projectId}")]
 		[ProducesResponseType(typeof(List<GetProjectResponse>), 200)]
-		public async Task<ActionResult<object>> GetProjectAsync(ProjectId projectId, [FromQuery] PropertyFilter? filter = null)
+		public ActionResult<object> GetProject(ProjectId projectId, [FromQuery] PropertyFilter? filter = null)
 		{
-			IProject? project = await _projectService.GetProjectAsync(projectId);
-			if (project == null)
+			ProjectConfig? projectConfig;
+			if(!_globalConfig.Value.TryGetProject(projectId, out projectConfig))
 			{
 				return NotFound(projectId);
 			}
-
-			ProjectPermissionsCache cache = new ProjectPermissionsCache();
-			if (!await _projectService.AuthorizeAsync(project, AclAction.ViewProject, User, cache))
+			if (!projectConfig.Authorize(AclAction.ViewProject, User))
 			{
 				return Forbid(AclAction.ViewProject, projectId);
 			}
@@ -100,23 +89,13 @@ namespace Horde.Build.Projects
 			bool includeStreams = PropertyFilter.Includes(filter, nameof(GetProjectResponse.Streams));
 			bool includeCategories = PropertyFilter.Includes(filter, nameof(GetProjectResponse.Categories));
 
-			List<IStream>? visibleStreams = null;
+			List<StreamConfig>? visibleStreams = null;
 			if (includeStreams || includeCategories)
 			{
-				visibleStreams = new List<IStream>();
-
-				List<IStream> streams = await _streamService.GetStreamsAsync(project.Id);
-				foreach (IStream stream in streams)
-				{
-					if (await _streamService.AuthorizeAsync(stream, AclAction.ViewStream, User, cache))
-					{
-						visibleStreams.Add(stream);
-					}
-				}
+				visibleStreams = projectConfig.Streams.Where(x => x.Authorize(AclAction.ViewStream, User)).ToList();
 			}
 
-			bool includeAcl = await _projectService.AuthorizeAsync(project, AclAction.ViewPermissions, User, cache);
-			return project.ToResponse(includeStreams, includeCategories, visibleStreams, includeAcl).ApplyFilter(filter);
+			return GetProjectResponse.FromConfig(projectConfig, includeStreams, includeCategories, visibleStreams).ApplyFilter(filter);
 		}
 
 		/// <summary>
@@ -128,17 +107,17 @@ namespace Horde.Build.Projects
 		[Route("/api/v1/projects/{projectId}/logo")]
 		public async Task<ActionResult<object>> GetProjectLogoAsync(ProjectId projectId)
 		{
-			IProject? project = await _projectService.GetProjectAsync(projectId);
-			if (project == null)
+			ProjectConfig? projectConfig;
+			if (!_globalConfig.Value.TryGetProject(projectId, out projectConfig))
 			{
 				return NotFound(projectId);
 			}
-			if (!await _projectService.AuthorizeAsync(project, AclAction.ViewProject, User, null))
+			if (!projectConfig.Authorize(AclAction.ViewProject, User))
 			{
 				return Forbid(AclAction.ViewProject, projectId);
 			}
 
-			IProjectLogo? projectLogo = await _projectService.Collection.GetLogoAsync(projectId);
+			IProjectLogo? projectLogo = await _projectCollection.GetLogoAsync(projectId);
 			if (projectLogo == null)
 			{
 				return NotFound();
diff --git a/Engine/Source/Programs/Horde/Horde.Build/Secrets/Credential.cs b/Engine/Source/Programs/Horde/Horde.Build/Secrets/Credential.cs
--- a/Engine/Source/Programs/Horde/Horde.Build/Secrets/Credential.cs
+++ b/Engine/Source/Programs/Horde/Horde.Build/Secrets/Credential.cs
@@ -35,11 +35,6 @@ namespace Horde.Build.Secrets
 		/// </summary>
 		public Dictionary<string, string> Properties { get; set; } = new Dictionary<string, string>();
 
-		/// <summary>
-		/// The ACL for this credential
-		/// </summary>
-		public Acl? Acl { get; set; }
-
 		/// <summary>
 		/// Private constructor for serialization
 		/// </summary>
diff --git a/Engine/Source/Programs/Horde/Horde.Build/Secrets/CredentialMessages.cs b/Engine/Source/Programs/Horde/Horde.Build/Secrets/CredentialMessages.cs
--- a/Engine/Source/Programs/Horde/Horde.Build/Secrets/CredentialMessages.cs
+++ b/Engine/Source/Programs/Horde/Horde.Build/Secrets/CredentialMessages.cs
@@ -85,22 +85,15 @@ namespace Horde.Build.Credentials
 		/// </summary>
 		public Dictionary<string, string> Properties { get; set; }
 
-		/// <summary>
-		/// Custom permissions for this object
-		/// </summary>
-		public GetAclResponse? Acl { get; set; }
-
 		/// <summary>
 		/// Constructor
 		/// </summary>
 		/// <param name="credential">The credential to construct from</param>
-		/// <param name="includeAcl">Whether to include the ACL in the response</param>
-		public GetCredentialResponse(Credential credential, bool includeAcl)
+		public GetCredentialResponse(Credential credential)
 		{
 			Id = credential.Id.ToString();
 			Name = credential.Name;
 			Properties = credential.Properties;
-			Acl = (includeAcl && credential.Acl != null)? new GetAclResponse(credential.Acl) : null;
 		}
 	}
 }
diff --git a/Engine/Source/Programs/Horde/Horde.Build/Secrets/CredentialService.cs b/Engine/Source/Programs/Horde/Horde.Build/Secrets/CredentialService.cs
--- a/Engine/Source/Programs/Horde/Horde.Build/Secrets/CredentialService.cs
+++ b/Engine/Source/Programs/Horde/Horde.Build/Secrets/CredentialService.cs
@@ -67,9 +67,8 @@ namespace Horde.Build.Secrets
 		/// <param name="id">Unique id of the credential</param>
 		/// <param name="newName">The new name for the credential</param>
 		/// <param name="newProperties">Properties on the credential to update. Any properties with a value of null will be removed.</param>
-		/// <param name="newAcl">The new ACL for the credential</param>
 		/// <returns>Async task object</returns>
-		public async Task UpdateCredentialAsync(ObjectId id, string? newName, Dictionary<string, string?>? newProperties, Acl? newAcl)
+		public async Task UpdateCredentialAsync(ObjectId id, string? newName, Dictionary<string, string?>? newProperties)
 		{
 			UpdateDefinitionBuilder<Credential> updateBuilder = Builders<Credential>.Update;
 
@@ -93,10 +92,6 @@ namespace Horde.Build.Secrets
 					}
 				}
 			}
-			if (newAcl != null)
-			{
-				updates.Add(Acl.CreateUpdate<Credential>(x => x.Acl!, newAcl));
-			}
 
 			if (updates.Count > 0)
 			{
@@ -152,26 +147,5 @@ namespace Horde.Build.Secrets
 			DeleteResult result = await _credentials.DeleteOneAsync<Credential>(x => x.Id == id);
 			return result.DeletedCount > 0;
 		}
-
-		/// <summary>
-		/// Determines if the user is authorized to perform an action on a particular credential
-		/// </summary>
-		/// <param name="credential">The credential to check</param>
-		/// <param name="action">The action being performed</param>
-		/// <param name="user">The principal to authorize</param>
-		/// <param name="cache">Permissions cache</param>
-		/// <returns>True if the action is authorized</returns>
-		public Task<bool> AuthorizeAsync(Credential credential, AclAction action, ClaimsPrincipal user, GlobalPermissionsCache? cache)
-		{
-			bool? result = credential.Acl?.Authorize(action, user);
-			if (result == null)
-			{
-				return _aclService.AuthorizeAsync(action, user, cache);
-			}
-			else
-			{
-				return Task.FromResult(result.Value);
-			}
-		}
 	}
 }
diff --git a/Engine/Source/Programs/Horde/Horde.Build/Secrets/CredentialsController.cs b/Engine/Source/Programs/Horde/Horde.Build/Secrets/CredentialsController.cs
--- a/Engine/Source/Programs/Horde/Horde.Build/Secrets/CredentialsController.cs
+++ b/Engine/Source/Programs/Horde/Horde.Build/Secrets/CredentialsController.cs
@@ -4,9 +4,11 @@ using System.Collections.Generic;
 using System.Threading.Tasks;
 using Horde.Build.Acls;
 using Horde.Build.Credentials;
+using Horde.Build.Server;
 using Horde.Build.Utilities;
 using Microsoft.AspNetCore.Authorization;
 using Microsoft.AspNetCore.Mvc;
+using Microsoft.Extensions.Options;
 using MongoDB.Bson;
 
 namespace Horde.Build.Secrets
@@ -19,25 +21,16 @@ namespace Horde.Build.Secrets
 	[Route("[controller]")]
 	public class CredentialsController : ControllerBase
 	{
-		/// <summary>
-		/// Singleton instance of the ACL service
-		/// </summary>
-		private readonly AclService _aclService;
-
-		/// <summary>
-		/// Singleton instance of the credential service
-		/// </summary>
 		private readonly CredentialService _credentialService;
+		private readonly IOptionsSnapshot<GlobalConfig> _globalConfig;
 
 		/// <summary>
 		/// Constructor
 		/// </summary>
-		/// <param name="aclService">The ACL service</param>
-		/// <param name="credentialService">The credential service</param>
-		public CredentialsController(AclService aclService, CredentialService credentialService)
+		public CredentialsController(CredentialService credentialService, IOptionsSnapshot<GlobalConfig> globalConfig)
 		{
-			_aclService = aclService;
 			_credentialService = credentialService;
+			_globalConfig = globalConfig;
 		}
 
 		/// <summary>
@@ -49,7 +42,7 @@ namespace Horde.Build.Secrets
 		[Route("/api/v1/credentials")]
 		public async Task<ActionResult<CreateCredentialResponse>> CreateCredentialAsync([FromBody] CreateCredentialRequest create)
 		{
-			if(!await _aclService.AuthorizeAsync(AclAction.CreateCredential, User))
+			if(!_globalConfig.Value.Authorize(AclAction.CreateCredential, User))
 			{
 				return Forbid();
 			}
@@ -69,21 +62,19 @@ namespace Horde.Build.Secrets
 		[ProducesResponseType(typeof(List<GetCredentialResponse>), 200)]
 		public async Task<ActionResult<object>> FindCredentialAsync([FromQuery] string? name = null, [FromQuery] PropertyFilter? filter = null)
 		{
-			if (!await _aclService.AuthorizeAsync(AclAction.ListCredentials, User))
+			if (!_globalConfig.Value.Authorize(AclAction.ListCredentials, User))
 			{
 				return Forbid();
 			}
 
 			List<Credential> credentials = await _credentialService.FindCredentialsAsync(name);
-			GlobalPermissionsCache cache = new GlobalPermissionsCache();
 
 			List<object> responses = new List<object>();
 			foreach (Credential credential in credentials)
 			{
-				if (await _credentialService.AuthorizeAsync(credential, AclAction.ViewCredential, User, cache))
+				if (_globalConfig.Value.Authorize(AclAction.ViewCredential, User))
 				{
-					bool includeAcl = await _credentialService.AuthorizeAsync(credential, AclAction.ViewPermissions, User, cache);
-					responses.Add(new GetCredentialResponse(credential, includeAcl).ApplyFilter(filter));
+					responses.Add(new GetCredentialResponse(credential).ApplyFilter(filter));
 				}
 			}
 			return responses;
@@ -108,14 +99,12 @@ namespace Horde.Build.Secrets
 				return NotFound();
 			}
 
-			GlobalPermissionsCache cache = new GlobalPermissionsCache();
-			if (!await _credentialService.AuthorizeAsync(credential, AclAction.ViewCredential, User, cache))
+			if (!_globalConfig.Value.Authorize(AclAction.ViewCredential, User))
 			{
 				return Forbid();
 			}
 
-			bool includeAcl = await _credentialService.AuthorizeAsync(credential, AclAction.ViewPermissions, User, cache);
-			return new GetCredentialResponse(credential, includeAcl).ApplyFilter(filter);
+			return new GetCredentialResponse(credential).ApplyFilter(filter);
 		}
 
 		/// <summary>
@@ -136,17 +125,16 @@ namespace Horde.Build.Secrets
 				return NotFound();
 			}
 
-			GlobalPermissionsCache cache = new GlobalPermissionsCache();
-			if (!await _credentialService.AuthorizeAsync(credential, AclAction.UpdateCredential, User, cache))
+			if (!_globalConfig.Value.Authorize(AclAction.UpdateCredential, User))
 			{
 				return Forbid();
 			}
-			if (update.Acl != null && !await _credentialService.AuthorizeAsync(credential, AclAction.ChangePermissions, User, cache))
+			if (update.Acl != null && !_globalConfig.Value.Authorize(AclAction.ChangePermissions, User))
 			{
 				return Forbid();
 			}
 
-			await _credentialService.UpdateCredentialAsync(credentialIdValue, update.Name, update.Properties, Acl.Merge(credential.Acl, update.Acl));
+			await _credentialService.UpdateCredentialAsync(credentialIdValue, update.Name, update.Properties);
 			return new OkResult();
 		}
 
@@ -166,7 +154,7 @@ namespace Horde.Build.Secrets
 			{
 				return NotFound();
 			}
-			if (!await _credentialService.AuthorizeAsync(credential, AclAction.DeleteCredential, User, null))
+			if (!_globalConfig.Value.Authorize(AclAction.DeleteCredential, User))
 			{
 				return Forbid();
 			}
diff --git a/Engine/Source/Programs/Horde/Horde.Build/Server/AccountController.cs b/Engine/Source/Programs/Horde/Horde.Build/Server/AccountController.cs
--- a/Engine/Source/Programs/Horde/Horde.Build/Server/AccountController.cs
+++ b/Engine/Source/Programs/Horde/Horde.Build/Server/AccountController.cs
@@ -32,25 +32,16 @@ namespace Horde.Build.Server
 			"table { margin:10px 20px; } " +
 			"td { margin:5px; font-size:13px; }";
 
-		/// <summary>
-		/// The ACL service singleton
-		/// </summary>
-		readonly AclService _aclService;
-
-		/// <summary>
-		/// Authentication scheme in use
-		/// </summary>
 		readonly string _authenticationScheme;
+		readonly IOptionsSnapshot<GlobalConfig> _globalConfig;
 
 		/// <summary>
 		/// Constructor
 		/// </summary>
-		/// <param name="aclService">ACL service instance</param>
-		/// <param name="serverSettings">Server settings</param>
-		public AccountController(AclService aclService, IOptionsMonitor<ServerSettings> serverSettings)
+		public AccountController(IOptionsMonitor<ServerSettings> serverSettings, IOptionsSnapshot<GlobalConfig> globalConfig)
 		{
-			_aclService = aclService;
 			_authenticationScheme = GetAuthScheme(serverSettings.CurrentValue.AuthMethod);
+			_globalConfig = globalConfig;
 		}
 
 		/// <summary>
@@ -75,14 +66,14 @@ namespace Horde.Build.Server
 		/// <returns>The current login state</returns>
 		[HttpGet]
 		[Route("/account")]
-		public async Task<ActionResult> State()
+		public ActionResult State()
 		{
 			StringBuilder content = new StringBuilder();
 			content.Append($"<html><style>{StyleSheet}</style><h1>Horde Server</h1>");
 			if (User.Identity?.IsAuthenticated ?? false)
 			{
 				content.Append(CultureInfo.InvariantCulture, $"<p>User <b>{User.Identity?.Name}</b> is logged in. <a href=\"/account/logout\">Log out</a></p>");
-				if (await _aclService.AuthorizeAsync(AclAction.AdminWrite, User))
+				if (_globalConfig.Value.Authorize(AclAction.AdminWrite, User))
 				{
 					content.Append("<p>");
 					content.Append("<a href=\"/api/v1/admin/token\">Get bearer token</a><br/>");
diff --git a/Engine/Source/Programs/Horde/Horde.Build/Server/AdminController.cs b/Engine/Source/Programs/Horde/Horde.Build/Server/AdminController.cs
--- a/Engine/Source/Programs/Horde/Horde.Build/Server/AdminController.cs
+++ b/Engine/Source/Programs/Horde/Horde.Build/Server/AdminController.cs
@@ -54,17 +54,17 @@ namespace Horde.Build.Server
 	public class AdminController : HordeControllerBase
 	{
 		readonly AclService _aclService;
-		readonly IOptionsMonitor<ServerSettings> _settings;
+		readonly IOptionsSnapshot<GlobalConfig> _globalConfig;
+		readonly IOptions<ServerSettings> _serverSettings;
 
 		/// <summary>
 		/// Constructor
 		/// </summary>
-		/// <param name="aclService">The ACL service singleton</param>
-		/// <param name="settings">Server settings</param>
-		public AdminController(AclService aclService, IOptionsMonitor<ServerSettings> settings)
+		public AdminController(AclService aclService, IOptionsSnapshot<GlobalConfig> globalConfig, IOptions<ServerSettings> serverSettings)
 		{
 			_aclService = aclService;
-			_settings = settings;
+			_globalConfig = globalConfig;
+			_serverSettings = serverSettings;
 		}
 
 		/// <summary>
@@ -75,7 +75,7 @@ namespace Horde.Build.Server
 		[Route("/api/v1/admin/token")]
 		public async Task<ActionResult<string>> GetTokenAsync()
 		{
-			if (!await _aclService.AuthorizeAsync(AclAction.IssueBearerToken, User))
+			if (!_globalConfig.Value.Authorize(AclAction.IssueBearerToken, User))
 			{
 				return Forbid(AclAction.IssueBearerToken);
 			}
@@ -92,7 +92,7 @@ namespace Horde.Build.Server
 		[Route("/api/v1/admin/roletoken")]
 		public async Task<ActionResult<string>> GetRoleTokenAsync([FromQuery] string roles)
 		{
-			if (!await _aclService.AuthorizeAsync(AclAction.AdminWrite, User))
+			if (!_globalConfig.Value.Authorize(AclAction.AdminWrite, User))
 			{
 				return Forbid(AclAction.AdminWrite);
 			}
@@ -111,14 +111,14 @@ namespace Horde.Build.Server
 		[Route("/api/v1/admin/registrationtoken")]
 		public async Task<ActionResult<string>> GetRegistrationTokenAsync()
 		{
-			if (!await _aclService.AuthorizeAsync(AclAction.AdminWrite, User))
+			if (!_globalConfig.Value.Authorize(AclAction.AdminWrite, User))
 			{
 				return Forbid(AclAction.AdminWrite);
 			}
 
-			List<AclClaim> claims = new List<AclClaim>();
-			claims.Add(new AclClaim(ClaimTypes.Name, User.Identity?.Name ?? "Unknown"));
-			claims.Add(AclService.AgentRegistrationClaim);
+			List<AclClaimConfig> claims = new List<AclClaimConfig>();
+			claims.Add(new AclClaimConfig(ClaimTypes.Name, User.Identity?.Name ?? "Unknown"));
+			claims.Add(HordeClaims.AgentRegistrationClaim);
 
 			return await _aclService.IssueBearerTokenAsync(claims, null);
 		}
@@ -131,14 +131,14 @@ namespace Horde.Build.Server
 		[Route("/api/v1/admin/softwaretoken")]
 		public async Task<ActionResult<string>> GetSoftwareTokenAsync()
 		{
-			if (!await _aclService.AuthorizeAsync(AclAction.AdminWrite, User))
+			if (!_globalConfig.Value.Authorize(AclAction.AdminWrite, User))
 			{
 				return Forbid(AclAction.AdminWrite);
 			}
 
-			List<AclClaim> claims = new List<AclClaim>();
-			claims.Add(new AclClaim(ClaimTypes.Name, User.Identity?.Name ?? "Unknown"));
-			claims.Add(AclService.UploadSoftwareClaim);
+			List<AclClaimConfig> claims = new List<AclClaimConfig>();
+			claims.Add(new AclClaimConfig(ClaimTypes.Name, User.Identity?.Name ?? "Unknown"));
+			claims.Add(HordeClaims.UploadSoftwareClaim);
 
 			return await _aclService.IssueBearerTokenAsync(claims, null);
 		}
@@ -151,14 +151,14 @@ namespace Horde.Build.Server
 		[Route("/api/v1/admin/softwaredownloadtoken")]
 		public async Task<ActionResult<string>> GetSoftwareDownloadTokenAsync()
 		{
-			if (!await _aclService.AuthorizeAsync(AclAction.AdminRead, User))
+			if (!_globalConfig.Value.Authorize(AclAction.AdminRead, User))
 			{
 				return Forbid(AclAction.AdminRead);
 			}
 
-			List<AclClaim> claims = new List<AclClaim>();
-			claims.Add(new AclClaim(ClaimTypes.Name, User.Identity?.Name ?? "Unknown"));
-			claims.Add(AclService.DownloadSoftwareClaim);
+			List<AclClaimConfig> claims = new List<AclClaimConfig>();
+			claims.Add(new AclClaimConfig(ClaimTypes.Name, User.Identity?.Name ?? "Unknown"));
+			claims.Add(HordeClaims.DownloadSoftwareClaim);
 
 			return await _aclService.IssueBearerTokenAsync(claims, null);
 		}
@@ -171,14 +171,14 @@ namespace Horde.Build.Server
 		[Route("/api/v1/admin/configtoken")]
 		public async Task<ActionResult<string>> GetConfigToken()
 		{
-			if (!await _aclService.AuthorizeAsync(AclAction.AdminRead, User))
+			if (!_globalConfig.Value.Authorize(AclAction.AdminRead, User))
 			{
 				return Forbid(AclAction.AdminRead);
 			}
 
-			List<AclClaim> claims = new List<AclClaim>();
-			claims.Add(new AclClaim(ClaimTypes.Name, User.Identity?.Name ?? "Unknown"));
-			claims.Add(AclService.ConfigureProjectsClaim);
+			List<AclClaimConfig> claims = new List<AclClaimConfig>();
+			claims.Add(new AclClaimConfig(ClaimTypes.Name, User.Identity?.Name ?? "Unknown"));
+			claims.Add(HordeClaims.ConfigureProjectsClaim);
 
 			return await _aclService.IssueBearerTokenAsync(claims, null);
 		}
@@ -191,14 +191,14 @@ namespace Horde.Build.Server
 		[Route("/api/v1/admin/chainedjobtoken")]
 		public async Task<ActionResult<string>> GetChainedJobToken()
 		{
-			if (!await _aclService.AuthorizeAsync(AclAction.AdminRead, User))
+			if (!_globalConfig.Value.Authorize(AclAction.AdminRead, User))
 			{
 				return Forbid(AclAction.AdminRead);
 			}
 
-			List<AclClaim> claims = new List<AclClaim>();
+			List<AclClaimConfig> claims = new List<AclClaimConfig>();
 			//Claims.Add(new AclClaim(ClaimTypes.Name, User.Identity.Name ?? "Unknown"));
-			claims.Add(AclService.StartChainedJobClaim);
+			claims.Add(HordeClaims.StartChainedJobClaim);
 
 			return await _aclService.IssueBearerTokenAsync(claims, null);
 		}
@@ -209,10 +209,12 @@ namespace Horde.Build.Server
 		/// <returns></returns>
 		private TimeSpan? GetDefaultExpiryTime()
 		{
+			ServerSettings serverSettings = _serverSettings.Value;
+
 			TimeSpan? expiryTime = null;
-			if (_settings.CurrentValue.JwtExpiryTimeHours != -1)
+			if (serverSettings.JwtExpiryTimeHours != -1)
 			{
-				expiryTime = TimeSpan.FromHours(_settings.CurrentValue.JwtExpiryTimeHours);
+				expiryTime = TimeSpan.FromHours(serverSettings.JwtExpiryTimeHours);
 			}
 			return expiryTime;
 		}
diff --git a/Engine/Source/Programs/Horde/Horde.Build/Server/DebugController.cs b/Engine/Source/Programs/Horde/Horde.Build/Server/DebugController.cs
--- a/Engine/Source/Programs/Horde/Horde.Build/Server/DebugController.cs
+++ b/Engine/Source/Programs/Horde/Horde.Build/Server/DebugController.cs
@@ -245,60 +245,22 @@ namespace Horde.Build.Server
 	{
 		private static readonly Random s_random = new ();
 		
-		/// <summary>
-		/// The ACL service singleton
-		/// </summary>
 		private readonly AclService _aclService;
-
-		/// <summary>
-		/// The database service instance
-		/// </summary>
 		private readonly MongoService _mongoService;
-
-		/// <summary>
-		/// The job task source singleton
-		/// </summary>
 		private readonly JobTaskSource _jobTaskSource;
-
-		/// <summary>
-		/// Collection of template documents
-		/// </summary>
 		private readonly ITemplateCollection _templateCollection;
-		
-		/// <summary>
-		/// The graph collection singleton
-		/// </summary>
 		private readonly IGraphCollection _graphCollection;
-
-		/// <summary>
-		/// The log file collection singleton
-		/// </summary>
 		private readonly ILogFileCollection _logFileCollection;
-
-		/// <summary>
-		/// Settings
-		/// </summary>
-		private readonly IOptionsMonitor<ServerSettings> _settings;
-		
-		/// <summary>
-		/// Logger
-		/// </summary>
+		private readonly IOptions<ServerSettings> _settings;
+		private readonly IOptionsSnapshot<GlobalConfig> _globalConfig;
 		private readonly ILogger<SecureDebugController> _logger;
 
 		/// <summary>
 		/// Constructor
 		/// </summary>
-		/// <param name="aclService">The ACL service singleton</param>
-		/// <param name="mongoService">The database service instance</param>
-		/// <param name="jobTaskSource">The dispatch service singleton</param>
-		/// <param name="templateCollection">Collection of template documents</param>
-		/// <param name="graphCollection">The graph collection</param>
-		/// <param name="logFileCollection">The log file collection</param>
-		/// <param name="settings">Settings</param>
-		/// <param name="logger">Logger</param>
 		public SecureDebugController(AclService aclService, MongoService mongoService, JobTaskSource jobTaskSource,
 			ITemplateCollection templateCollection, IGraphCollection graphCollection,
-			ILogFileCollection logFileCollection, IOptionsMonitor<ServerSettings> settings, ILogger<SecureDebugController> logger)
+			ILogFileCollection logFileCollection, IOptions<ServerSettings> settings, IOptionsSnapshot<GlobalConfig> globalConfig, ILogger<SecureDebugController> logger)
 		{
 			_aclService = aclService;
 			_mongoService = mongoService;
@@ -307,6 +269,7 @@ namespace Horde.Build.Server
 			_graphCollection = graphCollection;
 			_logFileCollection = logFileCollection;
 			_settings = settings;
+			_globalConfig = globalConfig;
 			_logger = logger;
 		}
 
@@ -316,9 +279,9 @@ namespace Horde.Build.Server
 		/// <returns>Http result</returns>
 		[HttpGet]
 		[Route("/api/v1/debug/environment")]
-		public async Task<ActionResult> GetServerEnvVars()
+		public ActionResult GetServerEnvVars()
 		{
-			if (!await _aclService.AuthorizeAsync(AclAction.AdminRead, User))
+			if (!_globalConfig.Value.Authorize(AclAction.AdminRead, User))
 			{
 				return Forbid();
 			}
@@ -342,9 +305,9 @@ namespace Horde.Build.Server
 		/// <returns>Information about the queue</returns>
 		[HttpGet]
 		[Route("/api/v1/debug/queue")]
-		public async Task<ActionResult<object>> GetQueueStatusAsync()
+		public ActionResult<object> GetQueueStatus()
 		{
-			if (!await _aclService.AuthorizeAsync(AclAction.AdminRead, User))
+			if (!_globalConfig.Value.Authorize(AclAction.AdminRead, User))
 			{
 				return Forbid();
 			}
@@ -358,9 +321,9 @@ namespace Horde.Build.Server
 		/// <returns>Information about the config</returns>
 		[HttpGet]
 		[Route("/api/v1/debug/config")]
-		public async Task<ActionResult> GetConfig()
+		public ActionResult GetConfig()
 		{
-			if (!await _aclService.AuthorizeAsync(AclAction.AdminRead, User))
+			if (!_globalConfig.Value.Authorize(AclAction.AdminRead, User))
 			{
 				return Forbid();
 			}
@@ -370,7 +333,7 @@ namespace Horde.Build.Server
 				WriteIndented = true
 			};
 
-			return Ok(JsonSerializer.Serialize(_settings.CurrentValue, options));
+			return Ok(JsonSerializer.Serialize(_settings.Value, options));
 		}
 		
 		/// <summary>
@@ -379,14 +342,14 @@ namespace Horde.Build.Server
 		/// <returns>Information about the log message generated</returns>
 		[HttpGet]
 		[Route("/api/v1/debug/generate-log-msg")]
-		public async Task<ActionResult> GenerateLogMessage(
+		public ActionResult GenerateLogMessage(
 			[FromQuery] string? logLevel = null,
 			[FromQuery] int messageLen = 0,
 			[FromQuery] int exceptionMessageLen = 0,
 			[FromQuery] int argCount = 0,
 			[FromQuery] int argLen = 10)
 		{
-			if (!await _aclService.AuthorizeAsync(AclAction.AdminRead, User))
+			if (!_globalConfig.Value.Authorize(AclAction.AdminRead, User))
 			{
 				return Forbid();
 			}
@@ -439,7 +402,7 @@ namespace Horde.Build.Server
 		[ProducesResponseType(200, Type = typeof(GetGraphResponse))]
 		public async Task<ActionResult<List<object>>> GetGraphsAsync([FromQuery] int? index = null, [FromQuery] int? count = null, [FromQuery] PropertyFilter? filter = null)
 		{
-			if (!await _aclService.AuthorizeAsync(AclAction.AdminRead, User))
+			if (!_globalConfig.Value.Authorize(AclAction.AdminRead, User))
 			{
 				return Forbid();
 			}
@@ -457,7 +420,7 @@ namespace Horde.Build.Server
 		[ProducesResponseType(200, Type = typeof(GetGraphResponse))]
 		public async Task<ActionResult<object>> GetGraphAsync(string graphId, [FromQuery] PropertyFilter? filter = null)
 		{
-			if (!await _aclService.AuthorizeAsync(AclAction.AdminRead, User))
+			if (!_globalConfig.Value.Authorize(AclAction.AdminRead, User))
 			{
 				return Forbid();
 			}
@@ -466,50 +429,6 @@ namespace Horde.Build.Server
 			return new GetGraphResponse(graph).ApplyFilter(filter);
 		}
 
-		/// <summary>
-		/// Query all the job templates.
-		/// </summary>
-		/// <param name="filter">Filter for properties to return</param>
-		/// <returns>Information about all the job templates</returns>
-		[HttpGet]
-		[Route("/api/v1/debug/templates")]
-		[ProducesResponseType(typeof(List<GetTemplateResponse>), 200)]
-		public async Task<ActionResult<object>> GetTemplatesAsync([FromQuery] PropertyFilter? filter = null)
-		{
-			if (!await _aclService.AuthorizeAsync(AclAction.AdminRead, User))
-			{
-				return Forbid();
-			}
-
-			List<ITemplate> templates = await _templateCollection.FindAllAsync();
-			return templates.ConvertAll(x => new GetTemplateResponse(x).ApplyFilter(filter));
-		}
-
-		/// <summary>
-		/// Retrieve information about a specific job template.
-		/// </summary>
-		/// <param name="templateHash">Id of the template to get information about</param>
-		/// <param name="filter">List of properties to return</param>
-		/// <returns>Information about the requested template</returns>
-		[HttpGet]
-		[Route("/api/v1/debug/templates/{TemplateHash}")]
-		[ProducesResponseType(typeof(GetTemplateResponse), 200)]
-		public async Task<ActionResult<object>> GetTemplateAsync(string templateHash, [FromQuery] PropertyFilter? filter = null)
-		{
-			ContentHash templateHashValue = ContentHash.Parse(templateHash);
-			if (!await _aclService.AuthorizeAsync(AclAction.AdminRead, User))
-			{
-				return Forbid();
-			}
-
-			ITemplate? template = await _templateCollection.GetAsync(templateHashValue);
-			if (template == null)
-			{
-				return NotFound();
-			}
-			return template.ApplyFilter(filter);
-		}
-
 		/// <summary>
 		/// Retrieve metadata about a specific log file
 		/// </summary>
@@ -520,7 +439,7 @@ namespace Horde.Build.Server
 		[Route("/api/v1/debug/logs/{LogFileId}")]
 		public async Task<ActionResult<object>> GetLogAsync(ObjectId<ILogFile> logFileId, [FromQuery] PropertyFilter? filter = null)
 		{
-			if (!await _aclService.AuthorizeAsync(AclAction.AdminRead, User))
+			if (!_globalConfig.Value.Authorize(AclAction.AdminRead, User))
 			{
 				return Forbid();
 			}
@@ -542,7 +461,7 @@ namespace Horde.Build.Server
 		[Route("/api/v1/debug/collections/{Name}")]
 		public async Task<ActionResult<object>> GetDocumentsAsync(string name, [FromQuery] string? filter = null, [FromQuery] int index = 0, [FromQuery] int count = 10)
 		{
-			if (!await _aclService.AuthorizeAsync(AclAction.AdminRead, User))
+			if (!_globalConfig.Value.Authorize(AclAction.AdminRead, User))
 			{
 				return Forbid();
 			}
diff --git a/Engine/Source/Programs/Horde/Horde.Build/Server/DowntimeService.cs b/Engine/Source/Programs/Horde/Horde.Build/Server/DowntimeService.cs
--- a/Engine/Source/Programs/Horde/Horde.Build/Server/DowntimeService.cs
+++ b/Engine/Source/Programs/Horde/Horde.Build/Server/DowntimeService.cs
@@ -39,21 +39,19 @@ namespace Horde.Build.Server
 			private set;
 		}
 
-		readonly GlobalsService _globalsService;
 		readonly IClock _clock;
-		readonly ITicker _ticker;
+		readonly IOptionsMonitor<GlobalConfig> _globalConfig;
 		readonly ILogger _logger;
 
+		readonly ITicker _ticker;
+
 		/// <summary>
 		/// Constructor
 		/// </summary>
-		/// <param name="globalsService">The database service instance</param>
-		/// <param name="clock"></param>
-		/// <param name="logger">Logger instance</param>
-		public DowntimeService(GlobalsService globalsService, IClock clock, ILogger<DowntimeService> logger)
+		public DowntimeService(IClock clock, IOptionsMonitor<GlobalConfig> globalConfig, ILogger<DowntimeService> logger)
 		{
-			_globalsService = globalsService;
 			_clock = clock;
+			_globalConfig = globalConfig;
 			_logger = logger;
 
 			// Ensure the initial value to be correct
@@ -76,15 +74,15 @@ namespace Horde.Build.Server
 		/// </summary>
 		/// <param name="stoppingToken">Token indicating that the service should stop</param>
 		/// <returns>Async task</returns>
-		async ValueTask TickAsync(CancellationToken stoppingToken)
+		ValueTask TickAsync(CancellationToken stoppingToken)
 		{
-			IGlobals globals = await _globalsService.GetAsync();
+			GlobalConfig globalConfig = _globalConfig.CurrentValue;
 
 			DateTimeOffset now = TimeZoneInfo.ConvertTime(new DateTimeOffset(_clock.UtcNow), _clock.TimeZone);
-			bool isActive = globals.Config.Downtime.Any(x => x.IsActive(now));
+			bool isActive = globalConfig.Downtime.Any(x => x.IsActive(now));
 
 			DateTimeOffset? next = null;
-			foreach (ScheduledDowntime schedule in globals.Config.Downtime)
+			foreach (ScheduledDowntime schedule in globalConfig.Downtime)
 			{
 				DateTimeOffset start = schedule.GetNext(now).StartTime;
 				if (next == null || start < next)
@@ -110,6 +108,8 @@ namespace Horde.Build.Server
 					_logger.LogInformation("Leaving downtime");
 				}
 			}
+
+			return new ValueTask();
 		}
 	}
 }
diff --git a/Engine/Source/Programs/Horde/Horde.Build/Server/GlobalConfig.cs b/Engine/Source/Programs/Horde/Horde.Build/Server/GlobalConfig.cs
--- a/Engine/Source/Programs/Horde/Horde.Build/Server/GlobalConfig.cs
+++ b/Engine/Source/Programs/Horde/Horde.Build/Server/GlobalConfig.cs
@@ -4,7 +4,9 @@ using System;
 using System.Collections.Generic;
 using System.ComponentModel.DataAnnotations;
 using System.Diagnostics;
+using System.Diagnostics.CodeAnalysis;
 using System.Linq;
+using System.Security.Claims;
 using System.Text.Json.Serialization;
 using EpicGames.Core;
 using EpicGames.Horde.Common;
@@ -16,13 +18,19 @@ using Horde.Build.Agents.Software;
 using Horde.Build.Configuration;
 using Horde.Build.Projects;
 using Horde.Build.Storage;
+using Horde.Build.Streams;
 using Horde.Build.Tools;
+using Horde.Build.Users;
 using Horde.Build.Utilities;
+using Microsoft.Extensions.Options;
 
 namespace Horde.Build.Server
 {
+	using UserId = ObjectId<IUser>;
 	using ProjectId = StringId<IProject>;
+	using StreamId = StringId<IStream>;
 	using AgentSoftwareChannelName = StringId<AgentSoftwareChannels>;
+	using ToolId = StringId<ITool>;
 
 	/// <summary>
 	/// Directive to merge config data from another source
@@ -44,6 +52,18 @@ namespace Horde.Build.Server
 	[ConfigIncludeRoot]
 	public class GlobalConfig
 	{
+		/// <summary>
+		/// Global server settings object
+		/// </summary>
+		[JsonIgnore]
+		public ServerSettings ServerSettings { get; private set; } = null!;
+
+		/// <summary>
+		/// Unique identifier for this config revision. Useful to detect changes.
+		/// </summary>
+		[JsonIgnore]
+		public string Revision { get; set; } = String.Empty;
+
 		/// <summary>
 		/// Other paths to include
 		/// </summary>
@@ -52,7 +72,7 @@ namespace Horde.Build.Server
 		/// <summary>
 		/// List of projects
 		/// </summary>
-		public List<ProjectConfigRef> Projects { get; set; } = new List<ProjectConfigRef>();
+		public List<ProjectConfig> Projects { get; set; } = new List<ProjectConfig>();
 
 		/// <summary>
 		/// List of pools
@@ -109,6 +129,117 @@ namespace Horde.Build.Server
 		/// </summary>
 		public AclConfig Acl { get; set; } = new AclConfig();
 
+		/// <summary>
+		/// Enumerates all the streams
+		/// </summary>
+		[JsonIgnore]
+		public IReadOnlyList<StreamConfig> Streams { get; private set; } = null!;
+
+		private readonly Dictionary<ProjectId, ProjectConfig> _projectLookup = new Dictionary<ProjectId, ProjectConfig>();
+		private readonly Dictionary<StreamId, StreamConfig> _streamLookup = new Dictionary<StreamId, StreamConfig>();
+		private readonly Dictionary<ToolId, ToolConfig> _toolLookup = new Dictionary<ToolId, ToolConfig>();
+
+		/// <summary>
+		/// Called after the config file has been read
+		/// </summary>
+		public void PostLoad(ServerSettings serverSettings)
+		{
+			ServerSettings = serverSettings;
+
+			Streams = Projects.SelectMany(x => x.Streams).ToList();
+
+			_projectLookup.Clear();
+			_streamLookup.Clear();
+			foreach (ProjectConfig project in Projects)
+			{
+				_projectLookup.Add(project.Id, project);
+				project.PostLoad(project.Id, this);
+
+				foreach (StreamConfig stream in project.Streams)
+				{
+					_streamLookup.Add(stream.Id, stream);
+				}
+			}
+
+			_toolLookup.Clear();
+			foreach (ToolConfig tool in Tools)
+			{
+				_toolLookup.Add(tool.Id, tool);
+				tool.PostLoad(this);
+			}
+
+			Storage.PostLoad(this);
+		}
+
+		/// <summary>
+		/// Attempts to get configuration for a project from this object
+		/// </summary>
+		/// <param name="projectId">The stream identifier</param>
+		/// <param name="config">Configuration for the stream</param>
+		/// <returns>True if the stream configuration was found</returns>
+		public bool TryGetProject(ProjectId projectId, [NotNullWhen(true)] out ProjectConfig? config) => _projectLookup.TryGetValue(projectId, out config);
+
+		/// <summary>
+		/// Attempts to get configuration for a stream from this object
+		/// </summary>
+		/// <param name="streamId">The stream identifier</param>
+		/// <param name="config">Configuration for the stream</param>
+		/// <returns>True if the stream configuration was found</returns>
+		public bool TryGetStream(StreamId streamId, [NotNullWhen(true)] out StreamConfig? config) => _streamLookup.TryGetValue(streamId, out config);
+
+		/// <summary>
+		/// Attempts to get configuration for a tool from this object
+		/// </summary>
+		/// <param name="toolId">The tool identifier</param>
+		/// <param name="config">Configuration for the stream</param>
+		/// <returns>True if the stream configuration was found</returns>
+		public bool TryGetTool(ToolId toolId, [NotNullWhen(true)] out ToolConfig? config) => _toolLookup.TryGetValue(toolId, out config);
+
+		/// <summary>
+		/// Authorizes a user to perform a given action
+		/// </summary>
+		/// <param name="action">The action being performed</param>
+		/// <param name="user">The principal to validate</param>
+		public bool Authorize(AclAction action, ClaimsPrincipal user)
+		{
+			return Acl.Authorize(action, user) ?? ServerSettings.Authorize(action, user);
+		}
+
+		/// <summary>
+		/// Determines whether the given user can masquerade as a given user
+		/// </summary>
+		/// <param name="user"></param>
+		/// <param name="userId"></param>
+		/// <returns></returns>
+		public bool AuthorizeAsUser(ClaimsPrincipal user, UserId userId)
+		{
+			UserId? currentUserId = user.GetUserId();
+			if (currentUserId != null && currentUserId.Value == userId)
+			{
+				return true;
+			}
+			else
+			{
+				return Authorize(AclAction.Impersonate, user);
+			}
+		}
+
+		/// <summary>
+		/// Finds a perforce cluster with the given name or that contains the provided server
+		/// </summary>
+		/// <param name="name">Name of the cluster</param>
+		/// <param name="serverAndPort">Find cluster which contains server</param>
+		/// <returns></returns>
+		public PerforceCluster GetPerforceCluster(string? name, string? serverAndPort = null)
+		{
+			PerforceCluster? cluster = FindPerforceCluster(name, serverAndPort);
+			if (cluster == null)
+			{
+				throw new Exception($"Unknown Perforce cluster '{name}'");
+			}
+			return cluster;
+		}
+
 		/// <summary>
 		/// Finds a perforce cluster with the given name or that contains the provided server
 		/// </summary>
@@ -187,31 +318,6 @@ namespace Horde.Build.Server
 		public AclConfig? Acl { get; set; }
 	}
 
-	/// <summary>
-	/// References a project configuration
-	/// </summary>
-	[DebuggerDisplay("{Id}")]
-	public class ProjectConfigRef
-	{
-		/// <summary>
-		/// Unique id for the project
-		/// </summary>
-		[Required]
-		public ProjectId Id { get; set; } = ProjectId.Empty;
-
-		/// <summary>
-		/// Config path for the project
-		/// </summary>
-		[Required, ConfigRelativePath]
-		public string Path { get; set; } = null!;
-
-		/// <summary>
-		/// Configuration for this project
-		/// </summary>
-		[JsonIgnore]
-		public ProjectConfig Target { get; set; } = null!;
-	}
-
 	/// <summary>
 	/// Configuration for a device platform 
 	/// </summary>
diff --git a/Engine/Source/Programs/Horde/Horde.Build/Server/GlobalsService.cs b/Engine/Source/Programs/Horde/Horde.Build/Server/GlobalsService.cs
--- a/Engine/Source/Programs/Horde/Horde.Build/Server/GlobalsService.cs
+++ b/Engine/Source/Programs/Horde/Horde.Build/Server/GlobalsService.cs
@@ -39,9 +39,6 @@ namespace Horde.Build.Server
 			public byte[]? JwtSigningKey { get; set; }
 			public int? SchemaVersion { get; set; }
 
-			[BsonIgnore]
-			public GlobalConfig Config { get; set; } = null!;
-
 			[BsonIgnore]
 			string IGlobals.JwtIssuer => _owner._jwtIssuer;
 
@@ -65,7 +62,6 @@ namespace Horde.Build.Server
 		}
 
 		readonly MongoService _mongoService;
-		readonly ConfigCollection _configCollection;
 		readonly string _jwtIssuer;
 		readonly byte[]? _fixedJwtSecret;
 
@@ -73,12 +69,10 @@ namespace Horde.Build.Server
 		/// Constructor
 		/// </summary>
 		/// <param name="mongoService"></param>
-		/// <param name="configCollection"></param>
 		/// <param name="settings">Global settings instance</param>
-		public GlobalsService(MongoService mongoService, ConfigCollection configCollection, IOptions<ServerSettings> settings)
+		public GlobalsService(MongoService mongoService, IOptions<ServerSettings> settings)
 		{
 			_mongoService = mongoService;
-			_configCollection = configCollection;
 
 			if (String.IsNullOrEmpty(settings.Value.JwtIssuer))
 			{
@@ -96,20 +90,6 @@ namespace Horde.Build.Server
 			}
 		}
 
-		async Task PostLoadAsync(Globals globals)
-		{
-			globals._owner = this;
-
-			if (String.IsNullOrEmpty(globals.ConfigRevision))
-			{
-				globals.Config = new GlobalConfig();
-			}
-			else
-			{
-				globals.Config = await _configCollection.GetConfigAsync<GlobalConfig>(globals.ConfigRevision);
-			}
-		}
-
 		/// <summary>
 		/// Gets the current globals instance
 		/// </summary>
@@ -117,7 +97,7 @@ namespace Horde.Build.Server
 		public async ValueTask<IGlobals> GetAsync()
 		{
 			Globals globals = await _mongoService.GetSingletonAsync<Globals>(() => CreateGlobals());
-			await PostLoadAsync(globals);
+			globals._owner = this;
 			return globals;
 		}
 
diff --git a/Engine/Source/Programs/Horde/Horde.Build/Server/IGlobals.cs b/Engine/Source/Programs/Horde/Horde.Build/Server/IGlobals.cs
--- a/Engine/Source/Programs/Horde/Horde.Build/Server/IGlobals.cs
+++ b/Engine/Source/Programs/Horde/Horde.Build/Server/IGlobals.cs
@@ -27,16 +27,6 @@ namespace Horde.Build.Server
 		/// </summary>
 		ObjectId InstanceId { get; }
 
-		/// <summary>
-		/// Revision string for the current config
-		/// </summary>
-		string ConfigRevision { get; }
-
-		/// <summary>
-		/// The global config object
-		/// </summary>
-		GlobalConfig Config { get; }
-
 		/// <summary>
 		/// Issuer for JWT keys. Can be overridden by config file.
 		/// </summary>
diff --git a/Engine/Source/Programs/Horde/Horde.Build/Server/Notices/NoticesController.cs b/Engine/Source/Programs/Horde/Horde.Build/Server/Notices/NoticesController.cs
--- a/Engine/Source/Programs/Horde/Horde.Build/Server/Notices/NoticesController.cs
+++ b/Engine/Source/Programs/Horde/Horde.Build/Server/Notices/NoticesController.cs
@@ -22,45 +22,18 @@ namespace Horde.Build.Server.Notices
 	[Route("[controller]")]
 	public class NoticesController : ControllerBase
 	{
-		/// <summary>
-		/// The acl service singleton
-		/// </summary>
-		readonly AclService _aclService;
-
-		/// <summary>
-		/// Reference to the user collection
-		/// </summary>
+		private readonly AclService _aclService;
 		private readonly IUserCollection _userCollection;
-
-		/// <summary>
-		/// Notice service
-		/// </summary>
 		private readonly NoticeService _noticeService;
-
-		/// <summary>
-		/// cached globals
-		/// </summary>
-		readonly LazyCachedValue<Task<IGlobals>> _cachedGlobals;
-
-		/// <summary>
-		/// cached notices
-		/// </summary>
-		readonly LazyCachedValue<Task<List<INotice>>> _cachedNotices;
-
-		/// <summary>
-		/// Time information
-		/// </summary>
-		readonly IClock _clock;
+		private readonly LazyCachedValue<Task<IGlobals>> _cachedGlobals;
+		private readonly LazyCachedValue<Task<List<INotice>>> _cachedNotices;
+		private readonly IOptionsSnapshot<GlobalConfig> _globalConfig;
+		private readonly IClock _clock;
 
 		/// <summary>
 		/// Constructor
 		/// </summary>
-		/// <param name="globalsService">The globals service singleton</param>
-		/// <param name="noticeService">The notice service singleton</param>
-		/// <param name="aclService">The acl service singleton</param>
-		/// <param name="userCollection">The user collection singleton</param>
-		/// <param name="clock"></param>
-		public NoticesController(GlobalsService globalsService, NoticeService noticeService, AclService aclService, IUserCollection userCollection, IClock clock)
+		public NoticesController(GlobalsService globalsService, NoticeService noticeService, AclService aclService, IUserCollection userCollection, IClock clock, IOptionsSnapshot<GlobalConfig> globalConfig)
 		{			
 			_aclService = aclService;
 			_userCollection = userCollection;
@@ -68,6 +41,7 @@ namespace Horde.Build.Server.Notices
 			_cachedGlobals = new LazyCachedValue<Task<IGlobals>>(async () => await globalsService.GetAsync(), TimeSpan.FromSeconds(30.0));
 			_cachedNotices = new LazyCachedValue<Task<List<INotice>>>(() => noticeService.GetNoticesAsync(), TimeSpan.FromMinutes(1));
 			_clock = clock;
+			_globalConfig = globalConfig;
 		}
 
 		/// <summary>
@@ -78,7 +52,7 @@ namespace Horde.Build.Server.Notices
 		[HttpPost("/api/v1/notices")]
 		public async Task<ActionResult> AddNoticeAsync(CreateNoticeRequest request)
 		{
-			if (!await _aclService.AuthorizeAsync(AclAction.AdminWrite, User))
+			if (!_globalConfig.Value.Authorize(AclAction.AdminWrite, User))
 			{
 				return Forbid();
 			}
@@ -97,7 +71,7 @@ namespace Horde.Build.Server.Notices
 		[HttpPut("/api/v1/notices")]
 		public async Task<ActionResult> UpdateNoticeAsync(UpdateNoticeRequest request)
 		{
-			if (!await _aclService.AuthorizeAsync(AclAction.AdminWrite, User))
+			if (!_globalConfig.Value.Authorize(AclAction.AdminWrite, User))
 			{
 				return Forbid();
 			}
@@ -112,7 +86,7 @@ namespace Horde.Build.Server.Notices
 		[HttpDelete("/api/v1/notices/{id}")]
 		public async Task<ActionResult> DeleteNoticeAsync(string id)
 		{
-			if (!await _aclService.AuthorizeAsync(AclAction.AdminWrite, User))
+			if (!_globalConfig.Value.Authorize(AclAction.AdminWrite, User))
 			{
 				return Forbid();
 			}
@@ -131,11 +105,9 @@ namespace Horde.Build.Server.Notices
 		{
 			List<GetNoticeResponse> messages = new List<GetNoticeResponse>();
 
-			IGlobals globals = await _cachedGlobals.GetCached();
-
 			DateTimeOffset now = TimeZoneInfo.ConvertTime(new DateTimeOffset(_clock.UtcNow), _clock.TimeZone);
 						
-			foreach (ScheduledDowntime schedule in globals.Config.Downtime)
+			foreach (ScheduledDowntime schedule in _globalConfig.Value.Downtime)
 			{
 				DateTimeOffset start = TimeZoneInfo.ConvertTimeBySystemTimeZoneId(schedule.GetNext(now).StartTime, "UTC");
 				DateTimeOffset finish = TimeZoneInfo.ConvertTimeBySystemTimeZoneId(schedule.GetNext(now).FinishTime, "UTC");
diff --git a/Engine/Source/Programs/Horde/Horde.Build/Server/RpcService.cs b/Engine/Source/Programs/Horde/Horde.Build/Server/RpcService.cs
--- a/Engine/Source/Programs/Horde/Horde.Build/Server/RpcService.cs
+++ b/Engine/Source/Programs/Horde/Horde.Build/Server/RpcService.cs
@@ -39,6 +39,7 @@ using Microsoft.AspNetCore.Authorization;
 using Microsoft.Extensions.Logging;
 using MongoDB.Bson;
 using MongoDB.Bson.Serialization;
+using Microsoft.Extensions.Options;
 
 namespace Horde.Build.Server
 {
@@ -68,7 +69,7 @@ namespace Horde.Build.Server
 		internal TimeSpan _longPollTimeout = TimeSpan.FromMinutes(9);
 		readonly AclService _aclService;
 		readonly AgentService _agentService;
-		readonly StreamService _streamService;
+		readonly IStreamCollection _streamCollection;
 		readonly JobService _jobService;
 		readonly AgentSoftwareService _agentSoftwareService;
 		readonly IArtifactCollection _artifactCollection;
@@ -82,16 +83,17 @@ namespace Horde.Build.Server
 		readonly ConformTaskSource _conformTaskSource;
 		readonly HttpClient _httpClient;
 		readonly IClock _clock;
+		readonly IOptionsSnapshot<GlobalConfig> _globalConfig;
 		readonly ILogger<RpcService> _logger;
 
 		/// <summary>
 		/// Constructor
 		/// </summary>
-		public RpcService(AclService aclService, AgentService agentService, StreamService streamService, JobService jobService, AgentSoftwareService agentSoftwareService, IArtifactCollection artifactCollection, ILogFileService logFileService, PoolService poolService, LifetimeService lifetimeService, IGraphCollection graphs, ITestDataCollection testData, IJobStepRefCollection jobStepRefCollection, ITemplateCollection templateCollection, ConformTaskSource conformTaskSource, HttpClient httpClient, IClock clock, ILogger<RpcService> logger)
+		public RpcService(AclService aclService, AgentService agentService, IStreamCollection streamCollection, JobService jobService, AgentSoftwareService agentSoftwareService, IArtifactCollection artifactCollection, ILogFileService logFileService, PoolService poolService, LifetimeService lifetimeService, IGraphCollection graphs, ITestDataCollection testData, IJobStepRefCollection jobStepRefCollection, ITemplateCollection templateCollection, ConformTaskSource conformTaskSource, HttpClient httpClient, IClock clock, IOptionsSnapshot<GlobalConfig> globalConfig, ILogger<RpcService> logger)
 		{
 			_aclService = aclService;
 			_agentService = agentService;
-			_streamService = streamService;
+			_streamCollection = streamCollection;
 			_jobService = jobService;
 			_agentSoftwareService = agentSoftwareService;
 			_artifactCollection = artifactCollection;
@@ -105,6 +107,7 @@ namespace Horde.Build.Server
 			_conformTaskSource = conformTaskSource;
 			_httpClient = httpClient;
 			_clock = clock;
+			_globalConfig = globalConfig;
 			_logger = logger;
 		}
 
@@ -218,7 +221,7 @@ namespace Horde.Build.Server
 				List<AgentWorkspace> newWorkspaces = request.Workspaces.Select(x => new AgentWorkspace(x)).ToList();
 
 				// Get the set of workspaces that are currently required
-				HashSet<AgentWorkspace> conformWorkspaces = await _poolService.GetWorkspacesAsync(agent, DateTime.UtcNow);
+				HashSet<AgentWorkspace> conformWorkspaces = await _poolService.GetWorkspacesAsync(agent, DateTime.UtcNow, _globalConfig.Value);
 				bool pendingConform = !conformWorkspaces.SetEquals(newWorkspaces) || (agent.RequestFullConform && !request.RemoveUntrackedFiles);
 
 				// Update the workspaces
@@ -293,11 +296,13 @@ namespace Horde.Build.Server
 			AgentId agentId = new AgentId(request.Name);
 			using IDisposable scope = _logger.BeginScope("CreateSession({AgentId})", agentId.ToString());
 
+			GlobalConfig globalConfig = _globalConfig.Value;
+
 			// Find the agent
 			IAgent? agent = await _agentService.GetAgentAsync(agentId);
 			if (agent == null)
 			{
-				if (!await _aclService.AuthorizeAsync(AclAction.CreateAgent, context.GetHttpContext().User))
+				if (!globalConfig.Authorize(AclAction.CreateAgent, context.GetHttpContext().User))
 				{
 					throw new StructuredRpcException(StatusCode.PermissionDenied, "User is not authenticated to create new agents");
 				}
@@ -306,7 +311,7 @@ namespace Horde.Build.Server
 			}
 
 			// Make sure we're allowed to create sessions on this agent
-			if (!await _agentService.AuthorizeAsync(agent, AclAction.CreateSession, context.GetHttpContext().User, null))
+			if (!globalConfig.Authorize(AclAction.CreateSession, context.GetHttpContext().User))
 			{
 				throw new StructuredRpcException(StatusCode.PermissionDenied, "User is not authenticated to create session for {AgentId}", request.Name);
 			}
@@ -419,21 +424,21 @@ namespace Horde.Build.Server
 		/// <param name="request">Request arguments</param>
 		/// <param name="context">Context for the RPC call</param>
 		/// <returns>Information about the new agent</returns>
-		public override async Task<RpcGetStreamResponse> GetStream(GetStreamRequest request, ServerCallContext context)
+		public override Task<RpcGetStreamResponse> GetStream(GetStreamRequest request, ServerCallContext context)
 		{
 			StreamId streamIdValue = new StreamId(request.StreamId);
 
-			IStream? stream = await _streamService.GetStreamAsync(streamIdValue);
-			if (stream == null)
+			StreamConfig? streamConfig;
+			if (!_globalConfig.Value.TryGetStream(streamIdValue, out streamConfig))
 			{
 				throw new StructuredRpcException(StatusCode.NotFound, "Stream {StreamId} does not exist", request.StreamId);
 			}
-			if (!await _streamService.AuthorizeAsync(stream, AclAction.ViewStream, context.GetHttpContext().User, null))
+			if (!streamConfig.Authorize(AclAction.ViewStream, context.GetHttpContext().User))
 			{
 				throw new StructuredRpcException(StatusCode.PermissionDenied, "Not authenticated to access stream {StreamId}", request.StreamId);
 			}
 
-			return stream.ToRpcResponse();
+			return Task.FromResult(streamConfig.ToRpcResponse());
 		}
 
 		/// <summary>
@@ -498,9 +503,13 @@ namespace Horde.Build.Server
 			{
 				throw new StructuredRpcException(StatusCode.NotFound, "Job {JobId} not found", request.JobId);
 			}
+			if(!_globalConfig.Value.TryGetStream(job.StreamId, out StreamConfig? streamConfig))
+			{
+				throw new StructuredRpcException(StatusCode.NotFound, "Stream {StreamId} not found", job.StreamId);
+			}
 
 			IJobStepBatch batch = AuthorizeBatch(job, request.BatchId.ToSubResourceId(), context);
-			job = await _jobService.UpdateBatchAsync(job, batchId, null, JobStepBatchState.Starting);
+			job = await _jobService.UpdateBatchAsync(job, batchId, streamConfig, newState: JobStepBatchState.Starting);
 
 			if (job == null)
 			{
@@ -528,9 +537,13 @@ namespace Horde.Build.Server
 			{
 				throw new StructuredRpcException(StatusCode.NotFound, "Job {JobId} not found", request.JobId);
 			}
+			if (!_globalConfig.Value.TryGetStream(job.StreamId, out StreamConfig? streamConfig))
+			{
+				throw new StructuredRpcException(StatusCode.NotFound, "Stream {StreamId} not found", job.StreamId);
+			}
 
 			IJobStepBatch batch = AuthorizeBatch(job, request.BatchId.ToSubResourceId(), context);
-			await _jobService.UpdateBatchAsync(job, batch.Id, null, JobStepBatchState.Complete);
+			await _jobService.UpdateBatchAsync(job, batch.Id, streamConfig, newState: JobStepBatchState.Complete);
 			return new Empty();
 		}
 
@@ -561,6 +574,10 @@ namespace Horde.Build.Server
 			{
 				throw new StructuredRpcException(StatusCode.NotFound, "Job {JobId} not found", request.JobId);
 			}
+			if (!_globalConfig.Value.TryGetStream(job.StreamId, out StreamConfig? streamConfig))
+			{
+				throw new StructuredRpcException(StatusCode.NotFound, "Stream {StreamId} not found", job.StreamId);
+			}
 
 			// Find the batch being executed
 			IJobStepBatch batch = AuthorizeBatch(job, request.BatchId.ToSubResourceId(), context);
@@ -617,7 +634,7 @@ namespace Horde.Build.Server
 			//				}
 
 			// Update the step state
-			IJob? newJob = await _jobService.TryUpdateStepAsync(job, batch.Id, step.Id, JobStepState.Running, JobStepOutcome.Unspecified, newLogId: log.Value.Id);
+			IJob? newJob = await _jobService.TryUpdateStepAsync(job, batch.Id, step.Id, streamConfig, JobStepState.Running, JobStepOutcome.Unspecified, newLogId: log.Value.Id);
 			if (newJob != null)
 			{
 				BeginStepResponse response = new BeginStepResponse();
@@ -689,28 +706,24 @@ namespace Horde.Build.Server
 					}
 				}
 
-				IStream? stream = await _streamService.GetStreamAsync(job.StreamId);
-				if (stream != null)
+				List<TokenConfig> tokenConfigs = new List<TokenConfig>(streamConfig.Tokens);
+
+				INodeGroup group = graph.Groups[batch.GroupIdx];
+				if (streamConfig.AgentTypes.TryGetValue(group.AgentType, out AgentConfig? agentConfig) && agentConfig.Tokens != null)
 				{
-					List<TokenConfig> tokenConfigs = new List<TokenConfig>(stream.Config.Tokens);
+					tokenConfigs.AddRange(agentConfig.Tokens);
+				}
 
-					INodeGroup group = graph.Groups[batch.GroupIdx];
-					if (stream.Config.AgentTypes.TryGetValue(group.AgentType, out AgentConfig? agentConfig) && agentConfig.Tokens != null)
+				foreach (TokenConfig tokenConfig in tokenConfigs)
+				{
+					string? value = await AllocateTokenAsync(tokenConfig, context.CancellationToken);
+					if (value == null)
 					{
-						tokenConfigs.AddRange(agentConfig.Tokens);
+						_logger.LogWarning("Unable to allocate token for job {JobId}:{BatchId}:{StepId} from {Url}", job.Id, batch.Id, step.Id, tokenConfig.Url);
 					}
-
-					foreach (TokenConfig tokenConfig in tokenConfigs)
+					else
 					{
-						string? value = await AllocateTokenAsync(tokenConfig, context.CancellationToken);
-						if (value == null)
-						{
-							_logger.LogWarning("Unable to allocate token for job {JobId}:{BatchId}:{StepId} from {Url}", job.Id, batch.Id, step.Id, tokenConfig.Url);
-						}
-						else
-						{
-							response.Credentials.Add(tokenConfig.EnvVar, value);
-						}
+						response.Credentials.Add(tokenConfig.EnvVar, value);
 					}
 				}
 
@@ -798,6 +811,10 @@ namespace Horde.Build.Server
 			{
 				throw new StructuredRpcException(StatusCode.NotFound, "Job {JobId} not found", request.JobId);
 			}
+			if (!_globalConfig.Value.TryGetStream(job.StreamId, out StreamConfig? streamConfig))
+			{
+				throw new StructuredRpcException(StatusCode.NotFound, "Stream {StreamId} not found", job.StreamId);
+			}
 
 			IJobStepBatch batch = AuthorizeBatch(job, request.BatchId.ToSubResourceId(), context);
 
@@ -813,7 +830,7 @@ namespace Horde.Build.Server
 				error = JobStepError.TimedOut;
 			}
 
-			await _jobService.UpdateStepAsync(job, batch.Id, request.StepId.ToSubResourceId(), request.State, request.Outcome, error, null, null, null, null, null);
+			await _jobService.UpdateStepAsync(job, batch.Id, request.StepId.ToSubResourceId(), streamConfig, request.State, request.Outcome, error, null, null, null, null, null);
 			return new Empty();
 		}
 
@@ -910,7 +927,7 @@ namespace Horde.Build.Server
 		/// <returns>Information about the new agent</returns>
 		public override async Task<Empty> CreateEvents(CreateEventsRequest request, ServerCallContext context)
 		{
-			if (!await _aclService.AuthorizeAsync(AclAction.CreateEvent, context.GetHttpContext().User))
+			if (!_globalConfig.Value.Authorize(AclAction.CreateEvent, context.GetHttpContext().User))
 			{
 				throw new StructuredRpcException(StatusCode.NotFound, "Access denied");
 			}
@@ -959,7 +976,7 @@ namespace Horde.Build.Server
 		/// <returns>Information about the new agent</returns>
 		public override async Task<UploadSoftwareResponse> UploadSoftware(UploadSoftwareRequest request, ServerCallContext context)
 		{
-			if (!await _aclService.AuthorizeAsync(AclAction.UploadSoftware, context.GetHttpContext().User))
+			if (!_globalConfig.Value.Authorize(AclAction.UploadSoftware, context.GetHttpContext().User))
 			{
 				throw new StructuredRpcException(StatusCode.PermissionDenied, "Access to software is forbidden");
 			}
@@ -984,7 +1001,7 @@ namespace Horde.Build.Server
 		/// <returns>Information about the new agent</returns>
 		public override async Task DownloadSoftware(DownloadSoftwareRequest request, IServerStreamWriter<DownloadSoftwareResponse> responseStream, ServerCallContext context)
 		{
-			if (!await _aclService.AuthorizeAsync(AclAction.DownloadSoftware, context.GetHttpContext().User))
+			if (!_globalConfig.Value.Authorize(AclAction.DownloadSoftware, context.GetHttpContext().User))
 			{
 				throw new StructuredRpcException(StatusCode.NotFound, "Access to software is forbidden");
 			}
@@ -1122,6 +1139,11 @@ namespace Horde.Build.Server
 		public override async Task<CreateReportResponse> CreateReport(CreateReportRequest request, ServerCallContext context)
 		{
 			IJob job = await GetJobAsync(new JobId(request.JobId));
+			if (!_globalConfig.Value.TryGetStream(job.StreamId, out StreamConfig? streamConfig))
+			{
+				throw new StructuredRpcException(StatusCode.NotFound, "Stream {StreamId} not found", job.StreamId);
+			}
+
 			IJobStepBatch batch = AuthorizeBatch(job, request.BatchId.ToSubResourceId(), context);
 
 			Report newReport = new Report { Name = request.Name, Placement = request.Placement, ArtifactId = request.ArtifactId.ToObjectId() };
@@ -1133,7 +1155,7 @@ namespace Horde.Build.Server
 			else
 			{
 				_logger.LogDebug("Adding report to step {JobId}:{BatchId}:{StepId}: {Name} -> {ArtifactId}", job.Id, batch.Id, request.StepId, request.Name, request.ArtifactId);
-				await _jobService.UpdateStepAsync(job, batch.Id, request.StepId.ToSubResourceId(), newReports: new List<Report> { newReport });
+				await _jobService.UpdateStepAsync(job, batch.Id, request.StepId.ToSubResourceId(), streamConfig, newReports: new List<Report> { newReport });
 			}
 
 			return new CreateReportResponse();
diff --git a/Engine/Source/Programs/Horde/Horde.Build/ServerSettings.cs b/Engine/Source/Programs/Horde/Horde.Build/ServerSettings.cs
--- a/Engine/Source/Programs/Horde/Horde.Build/ServerSettings.cs
+++ b/Engine/Source/Programs/Horde/Horde.Build/ServerSettings.cs
@@ -3,15 +3,17 @@
 using System;
 using System.Collections.Generic;
 using System.Linq;
+using System.Security.Claims;
+using System.Text.Json.Serialization;
 using EpicGames.Horde.Storage;
 using EpicGames.Horde.Storage.Nodes;
 using EpicGames.Perforce;
+using Horde.Build.Acls;
 using Horde.Build.Agents.Fleet;
 using Horde.Build.Server;
 using Horde.Build.Storage.Backends;
 using Horde.Build.Telemetry;
 using Horde.Build.Utilities;
-using MongoDB.Driver.Core.Configuration;
 using Serilog.Events;
 
 namespace Horde.Build
@@ -655,6 +657,56 @@ namespace Horde.Build
 		/// </summary>
 		public TelemetryConfig Telemetry { get; set; } = new TelemetryConfig();
 
+		/// <summary>
+		/// Default pre-baked ACL for authentication of well-known roles
+		/// </summary>
+		[JsonIgnore]
+		public AclConfig DefaultAcl
+		{
+			get
+			{
+				_defaultAcl ??= GetDefaultAcl();
+				return _defaultAcl;
+			}
+		}
+		
+		[JsonIgnore]
+		AclConfig? _defaultAcl;
+
+		/// <summary>
+		/// Authorizes a user to perform a given action
+		/// </summary>
+		/// <param name="action">The action being performed</param>
+		/// <param name="user">The principal to validate</param>
+		public bool Authorize(AclAction action, ClaimsPrincipal user)
+		{
+			return DefaultAcl.Authorize(action, user) ?? false;
+		}
+
+		/// <summary>
+		/// Create the default ACL for the server, including all predefined roles.
+		/// </summary>
+		/// <returns></returns>
+		AclConfig GetDefaultAcl()
+		{
+			AclConfig defaultAcl = new AclConfig();
+			defaultAcl.Entries.Add(new AclEntryConfig(new AclClaimConfig(ClaimTypes.Role, "internal:AgentRegistration"), new[] { AclAction.CreateAgent, AclAction.CreateSession }));
+			defaultAcl.Entries.Add(new AclEntryConfig(HordeClaims.AgentRegistrationClaim, new[] { AclAction.CreateAgent, AclAction.CreateSession, AclAction.UpdateAgent, AclAction.DownloadSoftware, AclAction.CreatePool, AclAction.UpdatePool, AclAction.ViewPool, AclAction.DeletePool, AclAction.ListPools, AclAction.ViewStream, AclAction.ViewProject, AclAction.ViewJob, AclAction.ViewCosts }));
+			defaultAcl.Entries.Add(new AclEntryConfig(HordeClaims.AgentRoleClaim, new[] { AclAction.ViewProject, AclAction.ViewStream, AclAction.CreateEvent, AclAction.DownloadSoftware }));
+			defaultAcl.Entries.Add(new AclEntryConfig(HordeClaims.DownloadSoftwareClaim, new[] { AclAction.DownloadSoftware }));
+			defaultAcl.Entries.Add(new AclEntryConfig(HordeClaims.UploadSoftwareClaim, new[] { AclAction.UploadSoftware }));
+			defaultAcl.Entries.Add(new AclEntryConfig(HordeClaims.ConfigureProjectsClaim, new[] { AclAction.CreateProject, AclAction.UpdateProject, AclAction.ViewProject, AclAction.CreateStream, AclAction.UpdateStream, AclAction.ViewStream, AclAction.ChangePermissions }));
+			defaultAcl.Entries.Add(new AclEntryConfig(HordeClaims.StartChainedJobClaim, new[] { AclAction.CreateJob, AclAction.ExecuteJob, AclAction.UpdateJob, AclAction.ViewJob, AclAction.ViewTemplate, AclAction.ViewStream }));
+
+			if (AdminClaimType != null && AdminClaimValue != null)
+			{
+				AclAction[] actions = Enum.GetValues(typeof(AclAction)).OfType<AclAction>().ToArray();
+				defaultAcl.Entries.Add(new AclEntryConfig(new AclClaimConfig(AdminClaimType, AdminClaimValue), actions));
+			}
+
+			return defaultAcl;
+		}
+
 		/// <summary>
 		/// Helper method to check if this process has activated the given mode
 		/// </summary>
diff --git a/Engine/Source/Programs/Horde/Horde.Build/Startup.cs b/Engine/Source/Programs/Horde/Horde.Build/Startup.cs
--- a/Engine/Source/Programs/Horde/Horde.Build/Startup.cs
+++ b/Engine/Source/Programs/Horde/Horde.Build/Startup.cs
@@ -364,7 +364,6 @@ namespace Horde.Build
 			services.AddSingleton<IDeviceCollection, DeviceCollection>();
 			services.AddSingleton<INoticeCollection, NoticeCollection>();
 
-			services.AddSingleton<ConfigCollection>();
 			services.AddSingleton<ToolCollection>();
 
 			services.AddSingleton<IConfigSource, InMemoryConfigSource>();
@@ -416,7 +415,6 @@ namespace Horde.Build
 			});
 			services.AddSingleton<ICommitService, CommitService>();
 			services.AddSingleton<IClock, Clock>();
-			services.AddSingleton<ConfigUpdateService>();
 			services.AddSingleton<IDowntimeService, DowntimeService>();
 			services.AddSingleton<IssueService>();
 			services.AddSingleton<JobService>();
@@ -439,7 +437,6 @@ namespace Horde.Build
 
 			services.AddSingleton<PerforceLoadBalancer>();
 			services.AddSingleton<PoolService>();
-			services.AddSingleton<ProjectService>();
 			services.AddSingleton<ReplicationService>();
 			services.AddSingleton<ScheduleService>();
 
@@ -450,7 +447,6 @@ namespace Horde.Build
 				services.AddSingleton<INotificationSink, SlackNotificationSink>(sp => sp.GetRequiredService<SlackNotificationSink>());
 			}
 
-			services.AddSingleton<StreamService>();
 			services.AddSingleton<DeviceService>();						
 			services.AddSingleton<NoticeService>();
 			services.AddSingleton<StorageService>();
@@ -669,7 +665,6 @@ namespace Horde.Build
 					services.AddHostedService(provider => provider.GetRequiredService<SlackNotificationSink>());
 				}
 				services.AddHostedService(provider => provider.GetRequiredService<ConfigService>());
-				services.AddHostedService(provider => provider.GetRequiredService<ConfigUpdateService>());
 				services.AddHostedService<TelemetryService>();
 				services.AddHostedService(provider => provider.GetRequiredService<DeviceService>());
 				services.AddHostedService(provider => provider.GetRequiredService<TestDataService>());
@@ -956,9 +951,6 @@ namespace Horde.Build
 				BsonSerializer.RegisterSerializationProvider(new BsonSerializationProvider());
 				BsonSerializer.RegisterSerializationProvider(new StringIdSerializationProvider());
 				BsonSerializer.RegisterSerializationProvider(new ObjectIdSerializationProvider());
-
-                // Register all the custom class maps
-                BsonClassMap.RegisterClassMap<AclV2>(AclV2.ConfigureClassMap);
 			}
 		}
 
diff --git a/Engine/Source/Programs/Horde/Horde.Build/Storage/StorageConfig.cs b/Engine/Source/Programs/Horde/Horde.Build/Storage/StorageConfig.cs
--- a/Engine/Source/Programs/Horde/Horde.Build/Storage/StorageConfig.cs
+++ b/Engine/Source/Programs/Horde/Horde.Build/Storage/StorageConfig.cs
@@ -4,8 +4,13 @@ using System;
 using System.Collections.Generic;
 using System.ComponentModel.DataAnnotations;
 using System.Diagnostics;
+using System.Diagnostics.CodeAnalysis;
+using System.Linq;
+using System.Security.Claims;
+using System.Text.Json.Serialization;
 using EpicGames.Horde.Storage;
 using Horde.Build.Acls;
+using Horde.Build.Server;
 using Horde.Build.Storage.Backends;
 using Horde.Build.Utilities;
 
@@ -53,6 +58,41 @@ namespace Horde.Build.Storage
 		/// List of namespaces for storage
 		/// </summary>
 		public List<NamespaceConfig> Namespaces { get; set; } = new List<NamespaceConfig>();
+
+		private readonly Dictionary<BackendId, BackendConfig> _backendLookup = new Dictionary<BackendId, BackendConfig>();
+		private readonly Dictionary<NamespaceId, NamespaceConfig> _namespaceLookup = new Dictionary<NamespaceId, NamespaceConfig>();
+
+		/// <summary>
+		/// Called after the config has been read from disk
+		/// </summary>
+		internal void PostLoad(GlobalConfig globalConfig)
+		{
+			foreach (BackendConfig backendConfig in Backends)
+			{
+				_backendLookup.Add(backendConfig.Id, backendConfig);
+			}
+			foreach (NamespaceConfig namespaceConfig in Namespaces)
+			{
+				namespaceConfig.PostLoad(globalConfig);
+				_namespaceLookup.Add(namespaceConfig.Id, namespaceConfig);
+			}
+		}
+
+		/// <summary>
+		/// Gets a backend with the given id
+		/// </summary>
+		/// <param name="backendId">Identifier for the backend</param>
+		/// <param name="backendConfig">Receives the backend config on success</param>
+		/// <returns>True on success</returns>
+		public bool TryGetBackend(BackendId backendId, [NotNullWhen(true)] out BackendConfig? backendConfig) => _backendLookup.TryGetValue(backendId, out backendConfig);
+
+		/// <summary>
+		/// Gets a namespace with the given id
+		/// </summary>
+		/// <param name="namespaceId">Identifier for the backend</param>
+		/// <param name="namespaceConfig">Receives the backend config on success</param>
+		/// <returns>True on success</returns>
+		public bool TryGetNamespace(NamespaceId namespaceId, [NotNullWhen(true)] out NamespaceConfig? namespaceConfig) => _namespaceLookup.TryGetValue(namespaceId, out namespaceConfig);
 	}
 
 	/// <summary>
@@ -109,6 +149,12 @@ namespace Horde.Build.Storage
 	[DebuggerDisplay("{Id}")]
 	public class NamespaceConfig
 	{
+		/// <summary>
+		/// Owner of this config object
+		/// </summary>
+		[JsonIgnore]
+		public GlobalConfig GlobalConfig { get; private set; } = null!;
+
 		/// <summary>
 		/// Identifier for this namespace
 		/// </summary>
@@ -145,5 +191,24 @@ namespace Horde.Build.Storage
 		/// Access list for this namespace
 		/// </summary>
 		public AclConfig Acl { get; set; } = new AclConfig();
+
+		/// <summary>
+		/// Callback once the configuration has been read from disk
+		/// </summary>
+		/// <param name="globalConfig"></param>
+		public void PostLoad(GlobalConfig globalConfig)
+		{
+			GlobalConfig = globalConfig;
+		}
+
+		/// <summary>
+		/// Authorizes a user to perform a given action
+		/// </summary>
+		/// <param name="action">The action being performed</param>
+		/// <param name="user">The principal to validate</param>
+		public bool Authorize(AclAction action, ClaimsPrincipal user)
+		{
+			return Acl?.Authorize(action, user) ?? GlobalConfig.Authorize(action, user);
+		}
 	}
 }
