Horde: Refactor codebase to use new config system.

diff --git a/Engine/Source/Programs/Horde/Horde.Build/Storage/StorageController.cs b/Engine/Source/Programs/Horde/Horde.Build/Storage/StorageController.cs
--- a/Engine/Source/Programs/Horde/Horde.Build/Storage/StorageController.cs
+++ b/Engine/Source/Programs/Horde/Horde.Build/Storage/StorageController.cs
@@ -11,10 +11,12 @@ using Amazon.EC2.Model;
 using EpicGames.Core;
 using EpicGames.Horde.Storage;
 using Horde.Build.Acls;
+using Horde.Build.Server;
 using Horde.Build.Utilities;
 using Microsoft.AspNetCore.Authorization;
 using Microsoft.AspNetCore.Http;
 using Microsoft.AspNetCore.Mvc;
+using Microsoft.Extensions.Options;
 using Microsoft.Extensions.Primitives;
 
 namespace Horde.Build.Storage
@@ -147,14 +149,15 @@ namespace Horde.Build.Storage
 	public class StorageController : HordeControllerBase
 	{
 		readonly StorageService _storageService;
+		readonly IOptionsSnapshot<GlobalConfig> _globalConfig;
 
 		/// <summary>
 		/// Constructor
 		/// </summary>
-		/// <param name="storageService"></param>
-		public StorageController(StorageService storageService)
+		public StorageController(StorageService storageService, IOptionsSnapshot<GlobalConfig> globalConfig)
 		{
 			_storageService = storageService;
+			_globalConfig = globalConfig;
 		}
 
 		/// <summary>
@@ -168,9 +171,14 @@ namespace Horde.Build.Storage
 		[Route("/api/v1/storage/{namespaceId}/blobs")]
 		public async Task<ActionResult<WriteBlobResponse>> WriteBlobAsync(NamespaceId namespaceId, IFormFile? file, [FromForm] string? prefix = default, CancellationToken cancellationToken = default)
 		{
-			if (!await _storageService.AuthorizeAsync(namespaceId, User, AclAction.WriteBlobs, null, cancellationToken))
+			NamespaceConfig? namespaceConfig;
+			if (!_globalConfig.Value.Storage.TryGetNamespace(namespaceId, out namespaceConfig))
 			{
-				return Forbid(AclAction.WriteBlobs);
+				return NotFound(namespaceId);
+			}
+			if (!namespaceConfig.Authorize(AclAction.WriteBlobs, User))
+			{
+				return Forbid(AclAction.WriteBlobs, namespaceId);
 			}
 
 			IStorageClientImpl client = await _storageService.GetClientAsync(namespaceId, cancellationToken);
@@ -242,9 +250,14 @@ namespace Horde.Build.Storage
 		[Route("/api/v1/storage/{namespaceId}/blobs/{*locator}")]
 		public async Task<ActionResult> ReadBlobAsync(NamespaceId namespaceId, BlobLocator locator, [FromQuery] int? offset = null, [FromQuery] int? length = null, CancellationToken cancellationToken = default)
 		{
-			if (!await _storageService.AuthorizeAsync(namespaceId, User, AclAction.ReadBlobs, null, cancellationToken))
+			NamespaceConfig? namespaceConfig;
+			if (!_globalConfig.Value.Storage.TryGetNamespace(namespaceId, out namespaceConfig))
 			{
-				return Forbid(AclAction.WriteBlobs);
+				return NotFound(namespaceId);
+			}
+			if (!namespaceConfig.Authorize(AclAction.ReadBlobs, User))
+			{
+				return Forbid(AclAction.ReadBlobs, namespaceId);
 			}
 
 			IStorageClientImpl client = await _storageService.GetClientAsync(namespaceId, cancellationToken);
@@ -284,9 +297,14 @@ namespace Horde.Build.Storage
 		[Route("/api/v1/storage/{namespaceId}/nodes")]
 		public async Task<ActionResult<FindNodesResponse>> FindNodesAsync(NamespaceId namespaceId, string alias, CancellationToken cancellationToken = default)
 		{
-			if (!await _storageService.AuthorizeAsync(namespaceId, User, AclAction.ReadBlobs, null, cancellationToken))
+			NamespaceConfig? namespaceConfig;
+			if (!_globalConfig.Value.Storage.TryGetNamespace(namespaceId, out namespaceConfig))
+			{
+				return NotFound(namespaceId);
+			}
+			if (!namespaceConfig.Authorize(AclAction.ReadBlobs, User))
 			{
-				return Forbid(AclAction.WriteBlobs);
+				return Forbid(AclAction.ReadBlobs, namespaceId);
 			}
 
 			IStorageClientImpl client = await _storageService.GetClientAsync(namespaceId, cancellationToken);
@@ -316,9 +334,14 @@ namespace Horde.Build.Storage
 		[Route("/api/v1/storage/{namespaceId}/refs/{*refName}")]
 		public async Task<ActionResult> WriteRefAsync(NamespaceId namespaceId, RefName refName, [FromBody] WriteRefRequest request, CancellationToken cancellationToken)
 		{
-			if (!await _storageService.AuthorizeAsync(namespaceId, User, AclAction.WriteRefs, null, cancellationToken))
+			NamespaceConfig? namespaceConfig;
+			if (!_globalConfig.Value.Storage.TryGetNamespace(namespaceId, out namespaceConfig))
+			{
+				return NotFound(namespaceId);
+			}
+			if (!namespaceConfig.Authorize(AclAction.WriteRefs, User))
 			{
-				return Forbid(AclAction.WriteBlobs);
+				return Forbid(AclAction.WriteRefs, namespaceId);
 			}
 
 			IStorageClient client = await _storageService.GetClientAsync(namespaceId, cancellationToken);
@@ -338,9 +361,14 @@ namespace Horde.Build.Storage
 		[Route("/api/v1/storage/{namespaceId}/refs/{*refName}")]
 		public async Task<ActionResult<ReadRefResponse>> ReadRefAsync(NamespaceId namespaceId, RefName refName, CancellationToken cancellationToken)
 		{
-			if (!await _storageService.AuthorizeAsync(namespaceId, User, AclAction.ReadRefs, null, cancellationToken))
+			NamespaceConfig? namespaceConfig;
+			if (!_globalConfig.Value.Storage.TryGetNamespace(namespaceId, out namespaceConfig))
+			{
+				return NotFound(namespaceId);
+			}
+			if (!namespaceConfig.Authorize(AclAction.ReadRefs, User))
 			{
-				return Forbid(AclAction.ReadRefs);
+				return Forbid(AclAction.ReadRefs, namespaceId);
 			}
 
 			IStorageClient client = await _storageService.GetClientAsync(namespaceId, cancellationToken);
diff --git a/Engine/Source/Programs/Horde/Horde.Build/Storage/StorageService.cs b/Engine/Source/Programs/Horde/Horde.Build/Storage/StorageService.cs
--- a/Engine/Source/Programs/Horde/Horde.Build/Storage/StorageService.cs
+++ b/Engine/Source/Programs/Horde/Horde.Build/Storage/StorageService.cs
@@ -28,6 +28,7 @@ using Microsoft.Extensions.Configuration;
 using Microsoft.Extensions.DependencyInjection;
 using Microsoft.Extensions.Hosting;
 using Microsoft.Extensions.Logging;
+using Microsoft.Extensions.Options;
 using MongoDB.Bson;
 using MongoDB.Bson.Serialization.Attributes;
 using MongoDB.Driver;
@@ -271,13 +272,11 @@ namespace Horde.Build.Storage
 		{
 			public NamespaceConfig Config { get; }
 			public StorageClient Client { get; }
-			public Acl? Acl { get; }
 
 			public NamespaceInfo(NamespaceConfig config, StorageClient client)
 			{
 				Config = config;
 				Client = client;
-				Acl = Acl.Merge(null, config.Acl);
 			}
 
 			public void Dispose() => (Client as IDisposable)?.Dispose();
@@ -443,6 +442,7 @@ namespace Horde.Build.Storage
 		readonly IClock _clock;
 		readonly IMemoryCache _cache;
 		readonly IServiceProvider _serviceProvider;
+		readonly IOptionsMonitor<GlobalConfig> _globalConfig;
 		readonly ILogger _logger;
 
 		readonly IMongoCollection<BlobInfo> _blobCollection;
@@ -464,7 +464,7 @@ namespace Horde.Build.Storage
 		/// <summary>
 		/// Constructor
 		/// </summary>
-		public StorageService(GlobalsService globalsService, MongoService mongoService, RedisService redisService, AclService aclService, IClock clock, IMemoryCache cache, IServiceProvider serviceProvider, ILogger<StorageService> logger)
+		public StorageService(GlobalsService globalsService, MongoService mongoService, RedisService redisService, AclService aclService, IClock clock, IMemoryCache cache, IServiceProvider serviceProvider, IOptionsMonitor<GlobalConfig> globalConfig, ILogger<StorageService> logger)
 		{
 			_globalsService = globalsService;
 			_redisService = redisService;
@@ -472,7 +472,8 @@ namespace Horde.Build.Storage
 			_clock = clock;
 			_cache = cache;
 			_serviceProvider = serviceProvider;
-			_cachedState = new AsyncCachedValue<State>(() => GetNextState(), TimeSpan.FromMinutes(1.0));
+			_cachedState = new AsyncCachedValue<State>(() => Task.FromResult(GetNextState()), TimeSpan.FromMinutes(1.0));
+			_globalConfig = globalConfig;
 			_logger = logger;
 
 			List<MongoIndex<BlobInfo>> blobIndexes = new List<MongoIndex<BlobInfo>>();
@@ -566,38 +567,14 @@ namespace Horde.Build.Storage
 		/// <inheritdoc/>
 		async ValueTask<IStorageClient> IStorageClientFactory.GetClientAsync(NamespaceId namespaceId, CancellationToken cancellationToken) => await GetClientAsync(namespaceId, cancellationToken);
 
-		/// <summary>
-		/// Authorize operations for the given store
-		/// </summary>
-		/// <param name="namespaceId"></param>
-		/// <param name="user"></param>
-		/// <param name="action"></param>
-		/// <param name="cache">Global permissions cache</param>
-		/// <param name="cancellationToken">Cancellation token for the operation</param>
-		/// <returns></returns>
-		public async Task<bool> AuthorizeAsync(NamespaceId namespaceId, ClaimsPrincipal user, AclAction action, GlobalPermissionsCache? cache, CancellationToken cancellationToken)
-		{
-			NamespaceInfo namespaceInfo = await GetNamespaceInfoAsync(namespaceId, cancellationToken);
-
-			bool? result = namespaceInfo.Acl?.Authorize(action, user);
-			if (result == null)
-			{
-				return await _aclService.AuthorizeAsync(action, user, cache);
-			}
-			else
-			{
-				return result.Value;
-			}
-		}
-
 		#region Config
 
-		async Task<State> GetNextState()
+		State GetNextState()
 		{
-			IGlobals globals = await _globalsService.GetAsync();
-			if (_lastState == null || !String.Equals(_lastConfigRevision, globals.ConfigRevision, StringComparison.Ordinal))
+			GlobalConfig globalConfig = _globalConfig.CurrentValue;
+			if (_lastState == null || !String.Equals(_lastConfigRevision, globalConfig.Revision, StringComparison.Ordinal))
 			{
-				StorageConfig storageConfig = globals.Config.Storage;
+				StorageConfig storageConfig = globalConfig.Storage;
 
 				// Create a lookup for backend config objects by their id
 				Dictionary<BackendId, BackendConfig> backendIdToConfig = new Dictionary<BackendId, BackendConfig>();
@@ -643,7 +620,7 @@ namespace Horde.Build.Storage
 				_lastState?.Dispose();
 				_lastState = nextState;
 
-				_lastConfigRevision = globals.ConfigRevision;
+				_lastConfigRevision = globalConfig.Revision;
 			}
 			return _lastState;
 		}
diff --git a/Engine/Source/Programs/Horde/Horde.Build/Streams/IStream.cs b/Engine/Source/Programs/Horde/Horde.Build/Streams/IStream.cs
--- a/Engine/Source/Programs/Horde/Horde.Build/Streams/IStream.cs
+++ b/Engine/Source/Programs/Horde/Horde.Build/Streams/IStream.cs
@@ -53,31 +53,11 @@ namespace Horde.Build.Streams
 		/// </summary>
 		StreamId Id { get; }
 
-		/// <summary>
-		/// The project that this stream belongs to
-		/// </summary>
-		ProjectId ProjectId { get; }
-
-		/// <summary>
-		/// The stream name
-		/// </summary>
-		string Name { get; }
-
-		/// <summary>
-		/// The revision of config file used for this stream
-		/// </summary>
-		string ConfigRevision { get; }
-
 		/// <summary>
 		/// Configuration settings for the stream
 		/// </summary>
 		StreamConfig Config { get; }
 
-		/// <summary>
-		/// Whether this stream has been deleted
-		/// </summary>
-		bool Deleted { get; }
-
 		/// <summary>
 		/// List of templates available for this stream
 		/// </summary>
@@ -92,11 +72,6 @@ namespace Horde.Build.Streams
 		/// Comment/reason for why the stream was paused
 		/// </summary>
 		string? PauseComment { get; }
-
-		/// <summary>
-		/// The ACL for this object
-		/// </summary>
-		Acl? Acl { get; }
 	}
 
 	/// <summary>
@@ -114,11 +89,6 @@ namespace Horde.Build.Streams
 		/// </summary>
 		TemplateRefConfig Config { get; }
 
-		/// <summary>
-		/// Hash of the template definition
-		/// </summary>
-		ContentHash Hash { get; }
-
 		/// <summary>
 		/// List of schedules for this template
 		/// </summary>
@@ -128,11 +98,6 @@ namespace Horde.Build.Streams
 		/// List of template step states
 		/// </summary>
 		IReadOnlyList<ITemplateStep> StepStates { get; }
-
-		/// <summary>
-		/// Custom permissions for this template
-		/// </summary>
-		Acl? Acl { get; set; }
 	}
 
 	/// <summary>
@@ -198,99 +163,27 @@ namespace Horde.Build.Streams
 			return schedule.Config.GetNextTriggerTimeUtc(schedule.LastTriggerTimeUtc, timeZone);
 		}
 
-		/// <summary>
-		/// Tries to get an agent workspace definition from the given type name
-		/// </summary>
-		/// <param name="stream">The stream object</param>
-		/// <param name="agentType">The agent type</param>
-		/// <param name="workspace">Receives the agent workspace definition</param>
-		/// <returns>True if the agent type was valid, and an agent workspace could be created</returns>
-		public static bool TryGetAgentWorkspace(this IStream stream, AgentConfig agentType, [NotNullWhen(true)] out (AgentWorkspace, bool)? workspace)
-		{
-			// Get the workspace settings
-			if (agentType.Workspace == null)
-			{
-				// Use the default settings (fast switching workspace, clean 
-				workspace = (new AgentWorkspace(null, null, GetDefaultWorkspaceIdentifier(stream), stream.Name, null, false, null), true);
-				return true;
-			}
-			else
-			{
-				// Try to get the matching workspace type
-				WorkspaceConfig? workspaceType;
-				if (!stream.Config.WorkspaceTypes.TryGetValue(agentType.Workspace, out workspaceType))
-				{
-					workspace = null;
-					return false;
-				}
-
-				// Get the workspace identifier
-				string identifier;
-				if (workspaceType.Identifier != null)
-				{
-					identifier = workspaceType.Identifier;
-				}
-				else if (workspaceType.Incremental)
-				{
-					identifier = $"{stream.GetEscapedName()}+{agentType.Workspace}";
-				}
-				else
-				{
-					identifier = GetDefaultWorkspaceIdentifier(stream);
-				}
-
-				// Create the new workspace
-				workspace = (new AgentWorkspace(workspaceType.Cluster, workspaceType.UserName, identifier, workspaceType.Stream ?? stream.Name, workspaceType.View, workspaceType.Incremental, workspaceType.Method), workspaceType.UseAutoSdk);
-				return true;
-			}
-		}
-
-		/// <summary>
-		/// The escaped name of this stream. Removes all non-identifier characters.
-		/// </summary>
-		/// <param name="stream">The stream object</param>
-		/// <returns>Escaped name for the stream</returns>
-		public static string GetEscapedName(this IStream stream)
-		{
-			return Regex.Replace(stream.Name, @"[^a-zA-Z0-9_]", "+");
-		}
-
-		/// <summary>
-		/// Gets the default identifier for workspaces created for this stream. Just includes an escaped depot name.
-		/// </summary>
-		/// <param name="stream">The stream object</param>
-		/// <returns>The default workspace identifier</returns>
-		private static string GetDefaultWorkspaceIdentifier(IStream stream)
-		{
-			return Regex.Replace(stream.GetEscapedName(), @"^(\+\+[^+]*).*$", "$1");
-		}
-
 		/// <summary>
 		/// Converts to a public response object
 		/// </summary>
 		/// <param name="stream">The stream object</param>
-		/// <param name="includeAcl">Whether to include the ACL in the response object</param>
 		/// <param name="apiTemplateRefs">The template refs for this stream. Passed separately because they have their own ACL.</param>
 		/// <returns>New response instance</returns>
-		public static GetStreamResponse ToApiResponse(this IStream stream, bool includeAcl, List<GetTemplateRefResponse> apiTemplateRefs)
+		public static GetStreamResponse ToApiResponse(this IStream stream, List<GetTemplateRefResponse> apiTemplateRefs)
 		{
-			List<TabConfig> apiTabs = stream.Config.Tabs;
-			Dictionary<string, AgentConfig> apiAgentTypes = stream.Config.AgentTypes.ToDictionary(x => x.Key, x => x.Value);
-			Dictionary<string, WorkspaceConfig> apiWorkspaceTypes = stream.Config.WorkspaceTypes.ToDictionary(x => x.Key, x => x.Value);
-			GetAclResponse? apiAcl = (includeAcl && stream.Acl != null)? new GetAclResponse(stream.Acl) : null;
-			return new GetStreamResponse(stream.Id.ToString(), stream.ProjectId.ToString(), stream.Name, stream.ConfigRevision, stream.Config.Order, stream.Config.NotificationChannel, stream.Config.NotificationChannelFilter, stream.Config.TriageChannel, stream.Config.DefaultPreflight, apiTabs, apiAgentTypes, apiWorkspaceTypes, apiTemplateRefs, apiAcl, stream.PausedUntil, stream.PauseComment, stream.Config.Workflows);
+			return new GetStreamResponse(stream, apiTemplateRefs);
 		}
 
 		/// <summary>
 		/// Converts to an RPC response object
 		/// </summary>
-		/// <param name="stream">The stream object</param>
+		/// <param name="streamConfig">The stream config object</param>
 		/// <returns>New response instance</returns>
-		public static HordeCommon.Rpc.GetStreamResponse ToRpcResponse(this IStream stream)
+		public static HordeCommon.Rpc.GetStreamResponse ToRpcResponse(this StreamConfig streamConfig)
 		{
 			HordeCommon.Rpc.GetStreamResponse response = new HordeCommon.Rpc.GetStreamResponse();
-			response.Name = stream.Name;
-			response.AgentTypes.Add(stream.Config.AgentTypes.ToDictionary(x => x.Key, x => x.Value.ToRpcResponse()));
+			response.Name = streamConfig.Name;
+			response.AgentTypes.Add(streamConfig.AgentTypes.ToDictionary(x => x.Key, x => x.Value.ToRpcResponse()));
 			return response;
 		}
 
@@ -305,20 +198,4 @@ namespace Horde.Build.Streams
 			return stream.PausedUntil != null && stream.PausedUntil > currentTime;
 		}
 	}
-
-	/// <summary>
-	/// Projection of a stream definition to just include permissions info
-	/// </summary>
-	public interface IStreamPermissions
-	{
-		/// <summary>
-		/// ACL for the stream
-		/// </summary>
-		public Acl? Acl { get; }
-
-		/// <summary>
-		/// The project containing this stream
-		/// </summary>
-		public ProjectId ProjectId { get; }
-	}
 }
diff --git a/Engine/Source/Programs/Horde/Horde.Build/Streams/IStreamCollection.cs b/Engine/Source/Programs/Horde/Horde.Build/Streams/IStreamCollection.cs
--- a/Engine/Source/Programs/Horde/Horde.Build/Streams/IStreamCollection.cs
+++ b/Engine/Source/Programs/Horde/Horde.Build/Streams/IStreamCollection.cs
@@ -2,16 +2,14 @@
 
 using System;
 using System.Collections.Generic;
+using System.Linq;
 using System.Threading.Tasks;
 using Horde.Build.Jobs;
-using Horde.Build.Projects;
 using Horde.Build.Utilities;
 
 namespace Horde.Build.Streams
 {
 	using JobId = ObjectId<IJob>;
-	using ProjectId = StringId<IProject>;
-	using StreamId = StringId<IStream>;
 	using TemplateId = StringId<ITemplateRef>;
 
 	/// <summary>
@@ -19,42 +17,19 @@ namespace Horde.Build.Streams
 	/// </summary>
 	public interface IStreamCollection
 	{
-		/// <summary>
-		/// Creates or replaces a stream configuration
-		/// </summary>
-		/// <param name="id">Unique id for the new stream</param>
-		/// <param name="stream">The current stream value. If not-null, this will attempt to replace the existing instance.</param>
-		/// <param name="revision">The config file revision</param>
-		/// <param name="projectId">The project id</param>
-		/// <returns></returns>
-		Task<IStream?> TryCreateOrReplaceAsync(StreamId id, IStream? stream, string revision, ProjectId projectId);
-
 		/// <summary>
 		/// Gets a stream by ID
 		/// </summary>
-		/// <param name="streamId">Unique id of the stream</param>
+		/// <param name="streamConfig">The stream config object</param>
 		/// <returns>The stream document</returns>
-		Task<IStream?> GetAsync(StreamId streamId);
+		Task<IStream> GetAsync(StreamConfig streamConfig);
 
 		/// <summary>
-		/// Gets a stream's permissions by ID
+		/// Gets a stream by ID
 		/// </summary>
-		/// <param name="streamId">Unique id of the stream</param>
+		/// <param name="streamConfigs">The stream config object</param>
 		/// <returns>The stream document</returns>
-		Task<IStreamPermissions?> GetPermissionsAsync(StreamId streamId);
-
-		/// <summary>
-		/// Enumerates all streams
-		/// </summary>
-		/// <returns></returns>
-		Task<List<IStream>> FindAllAsync();
-
-		/// <summary>
-		/// Gets all the available streams for a project
-		/// </summary>
-		/// <param name="projectIds">Unique id of the projects to query</param>
-		/// <returns>List of stream documents</returns>
-		Task<List<IStream>> FindForProjectsAsync(ProjectId[] projectIds);
+		Task<List<IStream>> GetAsync(IReadOnlyList<StreamConfig> streamConfigs);
 
 		/// <summary>
 		/// Updates user-facing properties for an existing stream
@@ -84,37 +59,80 @@ namespace Horde.Build.Streams
 		/// <param name="stepStates">The stream states to update, pass an empty list to clear all step states, otherwise will be a partial update based on included step updates</param>
 		/// <returns></returns>
 		Task<IStream?> TryUpdateTemplateRefAsync(IStream streamInterface, TemplateId templateRefId, List<UpdateStepStateRequest>? stepStates = null);
+	}
 
+	/// <summary>
+	/// Extension methods for <see cref="StreamCollection"/>
+	/// </summary>
+	public static class StreamCollectionExtensions
+	{
 		/// <summary>
-		/// Delete a stream
+		/// Updates an existing stream
 		/// </summary>
-		/// <param name="streamId">Unique id of the stream</param>
-		/// <returns>Async task</returns>
-		Task DeleteAsync(StreamId streamId);
-	}
+		/// <param name="streamCollection">The stream collection</param>
+		/// <param name="stream">The stream to update</param>
+		/// <param name="newPausedUntil">The new datetime for pausing builds</param>
+		/// <param name="newPauseComment">The reason for pausing</param>
+		/// <returns>Async task object</returns>
+		public static async Task<IStream?> UpdatePauseStateAsync(this IStreamCollection streamCollection, IStream? stream, DateTime? newPausedUntil = null, string? newPauseComment = null)
+		{
+			for (; stream != null; stream = await streamCollection.GetAsync(stream.Config))
+			{
+				IStream? newStream = await streamCollection.TryUpdatePauseStateAsync(stream, newPausedUntil, newPauseComment);
+				if (newStream != null)
+				{
+					return newStream;
+				}
+			}
+			return null;
+		}
 
-	static class StreamCollectionExtensions
-	{
 		/// <summary>
-		/// Creates or replaces a stream configuration
+		/// Attempts to update the last trigger time for a schedule
 		/// </summary>
 		/// <param name="streamCollection">The stream collection</param>
-		/// <param name="id">Unique id for the new stream</param>
-		/// <param name="stream">The current stream value. If not-null, this will attempt to replace the existing instance.</param>
-		/// <param name="revision">The config file revision</param>
-		/// <param name="projectId">The project id</param>
-		/// <returns></returns>
-		public static async Task<IStream> CreateOrReplaceAsync(this IStreamCollection streamCollection, StreamId id, IStream? stream, string revision, ProjectId projectId)
+		/// <param name="stream">The stream to update</param>
+		/// <param name="templateRefId">The template ref id</param>
+		/// <param name="lastTriggerTimeUtc"></param>
+		/// <param name="lastTriggerChange"></param>
+		/// <param name="addJobs">Jobs to add</param>
+		/// <param name="removeJobs">Jobs to remove</param>
+		/// <returns>True if the stream was updated</returns>
+		public static async Task<IStream?> UpdateScheduleTriggerAsync(this IStreamCollection streamCollection, IStream stream, TemplateId templateRefId, DateTime? lastTriggerTimeUtc = null, int? lastTriggerChange = null, List<JobId>? addJobs = null, List<JobId>? removeJobs = null)
 		{
-			for (; ; )
+			IStream? newStream = stream;
+			while (newStream != null)
 			{
-				stream = await streamCollection.TryCreateOrReplaceAsync(id, stream, revision, projectId);
-				if (stream != null)
+				ITemplateRef? templateRef;
+				if (!newStream.Templates.TryGetValue(templateRefId, out templateRef))
+				{
+					break;
+				}
+				if (templateRef.Schedule == null)
+				{
+					break;
+				}
+
+				IEnumerable<JobId> newActiveJobs = templateRef.Schedule.ActiveJobs;
+				if (removeJobs != null)
+				{
+					newActiveJobs = newActiveJobs.Except(removeJobs);
+				}
+				if (addJobs != null)
 				{
-					return stream;
+					newActiveJobs = newActiveJobs.Union(addJobs);
 				}
-				stream = await streamCollection.GetAsync(id);
+
+				newStream = await streamCollection.TryUpdateScheduleTriggerAsync(newStream, templateRefId, lastTriggerTimeUtc, lastTriggerChange, newActiveJobs.ToList());
+
+				if (newStream != null)
+				{
+					return newStream;
+				}
+
+				newStream = await streamCollection.GetAsync(stream.Config);
 			}
+			return null;
 		}
 	}
 }
diff --git a/Engine/Source/Programs/Horde/Horde.Build/Streams/StreamCollection.cs b/Engine/Source/Programs/Horde/Horde.Build/Streams/StreamCollection.cs
--- a/Engine/Source/Programs/Horde/Horde.Build/Streams/StreamCollection.cs
+++ b/Engine/Source/Programs/Horde/Horde.Build/Streams/StreamCollection.cs
@@ -2,26 +2,16 @@
 
 using System;
 using System.Collections.Generic;
-using System.Diagnostics.CodeAnalysis;
 using System.IO;
 using System.Linq;
 using System.Threading.Tasks;
 using EpicGames.Core;
-using Horde.Build.Acls;
-using Horde.Build.Configuration;
 using Horde.Build.Jobs;
-using Horde.Build.Jobs.Schedules;
-using Horde.Build.Jobs.Templates;
-using Horde.Build.Perforce;
 using Horde.Build.Projects;
 using Horde.Build.Server;
 using Horde.Build.Users;
 using Horde.Build.Utilities;
 using HordeCommon;
-using Microsoft.Extensions.Logging;
-using MongoDB.Bson;
-using MongoDB.Bson.IO;
-using MongoDB.Bson.Serialization;
 using MongoDB.Bson.Serialization.Attributes;
 using MongoDB.Driver;
 
@@ -46,35 +36,24 @@ namespace Horde.Build.Streams
 			[BsonRequired, BsonId]
 			public StreamId Id { get; set; }
 
-			[BsonRequired]
-			public ProjectId ProjectId { get; set; }
-
-			public string ConfigRevision { get; set; } = String.Empty;
-
 			public Dictionary<TemplateId, TemplateRefDoc> Templates { get; set; } = new Dictionary<TemplateId, TemplateRefDoc>();
 			public DateTime? PausedUntil { get; set; }
 			public string? PauseComment { get; set; }
 
-			public Acl? Acl { get; set; }
 			public int UpdateIndex { get; set; }
-			public bool Deleted { get; set; }
 
 			[BsonConstructor]
 			private StreamDoc()
 			{
 			}
 
-			public StreamDoc(StreamId id, ProjectId projectId)
+			public StreamDoc(StreamId id)
 			{
 				Id = id;
-				ProjectId = projectId;
 			}
 
 			#region IStream Implementation
 
-			[BsonIgnore]
-			public string Name => Config?.Name!;
-
 			[BsonIgnore]
 			public StreamConfig Config { get; private set; } = null!;
 
@@ -90,30 +69,30 @@ namespace Horde.Build.Streams
 				}
 			}
 
-			public async Task PostLoadAsync(ConfigCollection configCollection, ILogger logger)
+			public bool PostLoad(StreamConfig config, DateTime utcNow)
 			{
-				try
-				{
-					Config = await configCollection.GetConfigAsync<StreamConfig>(ConfigRevision);
-				}
-				catch (Exception)
+				Config = config;
+
+				bool replaceDocument = false;
+
+				Dictionary<TemplateId, TemplateRefDoc> templateRefDocs = new Dictionary<TemplateId, TemplateRefDoc>(config.Templates.Count);
+				foreach (TemplateRefConfig templateRefConfig in config.Templates)
 				{
-					if (Deleted)
+					TemplateRefDoc? templateRefDoc;
+					if (!Templates.TryGetValue(templateRefConfig.Id, out templateRefDoc))
 					{
-						logger.LogWarning("Unable to get stream config for {StreamId} at {Revision}; using default.", Id, ConfigRevision);
-						Config = new StreamConfig();
+						templateRefDoc = new TemplateRefDoc();
+						replaceDocument = true;
 					}
-					else
-					{
-						logger.LogError("Unable to get stream config for {StreamId} at {Revision}", Id, ConfigRevision);
-						throw;
-					}
-				}
+					replaceDocument |= templateRefDoc.PostLoad(templateRefConfig, utcNow);
 
-				foreach (KeyValuePair<TemplateId, TemplateRefDoc> pair in Templates)
-				{
-					pair.Value.PostLoad(this, pair.Key);
+					templateRefDocs.Add(templateRefConfig.Id, templateRefDoc);
 				}
+
+				replaceDocument |= Templates.Count != templateRefDocs.Count;
+				Templates = templateRefDocs;
+
+				return replaceDocument;
 			}
 
 			#endregion
@@ -121,56 +100,51 @@ namespace Horde.Build.Streams
 
 		class TemplateRefDoc : ITemplateRef
 		{
-			[BsonRequired]
-			public ContentHash Hash { get; set; } = null!;
-
 			[BsonIgnoreIfNull]
 			public TemplateScheduleDoc? Schedule { get; set; }
 
 			[BsonIgnoreIfNull]
 			public List<TemplateStepDoc>? StepStates { get; set; }
 
-			[BsonIgnoreIfNull]
-			public Acl? Acl { get; set; }
-
 			#region ITemplateRef implementation
 
-			[BsonIgnore]
-			StreamDoc? _owner;
-
 			[BsonIgnore]
 			public TemplateId Id { get; private set; }
 
 			[BsonIgnore]
-			TemplateRefConfig? _config;
-
-			[BsonIgnore]
-			public TemplateRefConfig Config
-			{
-				get
-				{
-					if (_config == null && _owner != null && _owner.Config != null)
-					{
-						// This should generally always succeed, but adding a fallback to handle legacy data.
-						_config = _owner.Config.Templates.FirstOrDefault(x => x.Id == Id) ?? new TemplateRefConfig { Id = Id, Name = Id.ToString() };
-					}
-					return _config!;
-				}
-			}
+			public TemplateRefConfig Config { get; private set; } = null!;
 
 			ITemplateSchedule? ITemplateRef.Schedule => Schedule;
 			IReadOnlyList<ITemplateStep> ITemplateRef.StepStates => (IReadOnlyList<ITemplateStep>?)StepStates ?? Array.Empty<ITemplateStep>();
 
-			public void PostLoad(StreamDoc owner, TemplateId id)
+			public bool PostLoad(TemplateRefConfig config, DateTime utcNow)
 			{
-				_owner = owner;
-				Id = id;
+				Id = config.Id;
+				Config = config;
+
+				bool replaceDocument = false;
+
+				if (Config.Schedule == null)
+				{
+					Schedule = null;
+					replaceDocument = true;
+				}
+				else if (Schedule == null)
+				{
+					Schedule = new TemplateScheduleDoc();
+					Schedule.LastTriggerTimeUtc = utcNow;
+					replaceDocument = true;
+				}
+
 				Schedule?.PostLoad(this);
 
 				if (StepStates != null)
 				{
-					StepStates.RemoveAll(x => x.PausedByUserId == null);
+					int count = StepStates.RemoveAll(x => x.PausedByUserId == null);
+					replaceDocument |= count > 0;
 				}
+
+				return replaceDocument;
 			}
 
 			#endregion
@@ -233,99 +207,77 @@ namespace Horde.Build.Streams
 			}
 		}
 
-		/// <summary>
-		/// Projection of a stream definition to just include permissions info
-		/// </summary>
-		[SuppressMessage("Design", "CA1812: Class is never instantiated")]
-		private class StreamPermissions : IStreamPermissions
-		{
-			public Acl? Acl { get; set; }
-			public ProjectId ProjectId { get; set; }
-
-			public static readonly ProjectionDefinition<StreamDoc> Projection = Builders<StreamDoc>.Projection.Include(x => x.Acl).Include(x => x.ProjectId);
-		}
-
 		readonly IMongoCollection<StreamDoc> _streams;
-		readonly ConfigCollection _configCollection;
 		readonly IClock _clock;
-		readonly ITemplateCollection _templateCollection;
-		readonly ILogger<StreamCollection> _logger;
 
 		/// <summary>
 		/// Constructor
 		/// </summary>
 		/// <param name="mongoService">The database service instance</param>
-		/// <param name="configCollection"></param>
 		/// <param name="clock"></param>
-		/// <param name="templateCollection"></param>
-		/// <param name="logger"></param>
-		public StreamCollection(MongoService mongoService, ConfigCollection configCollection, IClock clock, ITemplateCollection templateCollection, ILogger<StreamCollection> logger)
+		public StreamCollection(MongoService mongoService, IClock clock)
 		{
 			_streams = mongoService.GetCollection<StreamDoc>("Streams");
-			_configCollection = configCollection;
 			_clock = clock;
-			_templateCollection = templateCollection;
-			_logger = logger;
 		}
 
-		Task PostLoadAsync(StreamDoc stream) => stream.PostLoadAsync(_configCollection, _logger);
-
 		/// <inheritdoc/>
-		public async Task<IStream?> TryCreateOrReplaceAsync(StreamId id, IStream? stream, string revision, ProjectId projectId)
-		{
-			StreamConfig config = await _configCollection.GetConfigAsync<StreamConfig>(revision);
-			Validate(id, config);
+		public async Task<IStream> GetAsync(StreamConfig streamConfig) => await GetInternalAsync(streamConfig);
 
-			StreamDoc? streamDoc = (StreamDoc?)stream;
-			Dictionary<TemplateId, TemplateRefDoc> templateRefs = await CreateTemplateRefsAsync(streamDoc, config.Templates);
+		async Task<StreamDoc> GetInternalAsync(StreamConfig streamConfig)
+		{
+			FindOneAndUpdateOptions<StreamDoc, StreamDoc> options = new FindOneAndUpdateOptions<StreamDoc, StreamDoc>();
+			options.IsUpsert = true;
+			options.ReturnDocument = ReturnDocument.After;
 
-			Acl? acl = Acl.Merge(new Acl(), config.Acl);
-			if (streamDoc == null)
+			for (; ; )
 			{
-				return await TryCreateAsync(id, projectId, revision, templateRefs, acl);
-			}
-			else
-			{
-				return await TryReplaceAsync(streamDoc, projectId, revision, config, templateRefs, acl);
+				StreamDoc stream = await _streams.FindOneAndUpdateAsync<StreamDoc>(x => x.Id == streamConfig.Id, Builders<StreamDoc>.Update.SetOnInsert(x => x.UpdateIndex, 0), options);
+				if (await PostLoadAsync(stream, streamConfig))
+				{
+					return stream;
+				}
 			}
 		}
 
-		/// <summary>
-		/// Creates a list of template refs from a set of request objects
-		/// </summary>
-		/// <param name="requests">Request objects</param>
-		/// <param name="stream">The current stream state</param>
-		/// <returns>List of new template references</returns>
-		async Task<Dictionary<TemplateId, TemplateRefDoc>> CreateTemplateRefsAsync(StreamDoc? stream, List<TemplateRefConfig> requests)
-		{
-			Dictionary<TemplateId, TemplateRefDoc> newTemplateRefs = new Dictionary<TemplateId, TemplateRefDoc>();
-			foreach (TemplateRefConfig request in requests)
-			{
-				// Create the template
-				ITemplate template = await _templateCollection.AddAsync(request.Name, request.Priority, request.AllowPreflights, request.UpdateIssues, request.PromoteIssuesByDefault, request.InitialAgentType, request.SubmitNewChange, request.SubmitDescription, request.Arguments, request.Parameters.ConvertAll(x => x.ToModel()));
+		/// <inheritdoc/>
+		public async Task<List<IStream>> GetAsync(IReadOnlyList<StreamConfig> streamConfigs) => await GetInternalAsync(streamConfigs).ConvertAllAsync<StreamDoc, IStream>();
 
-				// Get an identifier for the new template ref
-				TemplateId templateId = request.Id;
+		async Task<List<StreamDoc>> GetInternalAsync(IReadOnlyList<StreamConfig> streamConfigs)
+		{
+			FilterDefinition<StreamDoc> filter = Builders<StreamDoc>.Filter.In(x => x.Id, streamConfigs.Select(x => x.Id));
+			List<StreamDoc> matches = await _streams.Find(filter).ToListAsync();
 
-				// Get the existing template ref or create a new one
-				TemplateRefDoc? templateRef;
-				if (stream == null || !stream.Templates.TryGetValue(templateId, out templateRef))
+			List<StreamDoc> results = new List<StreamDoc>(streamConfigs.Count);
+			foreach (StreamConfig streamConfig in streamConfigs)
+			{
+				StreamDoc? stream = matches.FirstOrDefault(x => x.Id == streamConfig.Id);
+				if (stream == null)
 				{
-					templateRef = new TemplateRefDoc();
+					stream = await GetInternalAsync(streamConfig);
 				}
-
-				// Update the hash
-				templateRef.Hash = template.Id;
-				if (request.Schedule != null)
+				else
 				{
-					templateRef.Schedule ??= new TemplateScheduleDoc();
-					templateRef.Schedule.LastTriggerTimeUtc = _clock.UtcNow;
+					stream.PostLoad(streamConfig, _clock.UtcNow);
 				}
+				results.Add(stream);
+			}
+			return results;
+		}
 
-				// Add it to the new lookup
-				newTemplateRefs.Add(templateId, templateRef);
+		async ValueTask<bool> PostLoadAsync(StreamDoc streamDoc, StreamConfig streamConfig)
+		{
+			if (streamDoc.PostLoad(streamConfig, _clock.UtcNow))
+			{
+				int updateIndex = streamDoc.UpdateIndex++;
+
+				ReplaceOneResult result = await _streams.ReplaceOneAsync(x => x.Id == streamDoc.Id && x.UpdateIndex == updateIndex, streamDoc);
+				if (result.MatchedCount != 1)
+				{
+					return false;
+				}
 			}
-			return newTemplateRefs;
+			return true;
 		}
 
 		/// <inheritdoc/>
@@ -416,89 +368,6 @@ namespace Horde.Build.Streams
 
 		}
 
-		/// <inheritdoc/>
-		async Task<IStream?> TryCreateAsync(StreamId id, ProjectId projectId, string configRevision, Dictionary<TemplateId, TemplateRefDoc> templateRefs, Acl? acl)
-		{
-			StreamDoc newStream = new StreamDoc(id, projectId);
-			newStream.ConfigRevision = configRevision;
-			newStream.Templates = templateRefs;
-			newStream.Acl = acl;
-
-			try
-			{
-				await _streams.InsertOneAsync(newStream);
-				await PostLoadAsync(newStream);
-				return newStream;
-			}
-			catch (MongoWriteException ex)
-			{
-				if (ex.WriteError.Category == ServerErrorCategory.DuplicateKey)
-				{
-					return null;
-				}
-				else
-				{
-					throw;
-				}
-			}
-		}
-
-		/// <inheritdoc/>
-		async Task<IStream?> TryReplaceAsync(StreamDoc stream, ProjectId projectId, string configRevision, StreamConfig config, Dictionary<TemplateId, TemplateRefDoc> templateRefs, Acl? acl)
-		{
-			UpdateDefinitionBuilder<StreamDoc> updateBuilder = Builders<StreamDoc>.Update;
-
-			List<UpdateDefinition<StreamDoc>> updates = new List<UpdateDefinition<StreamDoc>>();
-			updates.Add(updateBuilder.Set(x => x.ProjectId, projectId));
-			updates.Add(updateBuilder.Set(x => x.ConfigRevision, configRevision));
-			updates.Add(updateBuilder.Set(x => x.Templates, templateRefs));
-			updates.Add(updateBuilder.SetOrUnsetNullRef(x => x.Acl, acl));
-			updates.Add(updateBuilder.Unset(x => x.Deleted));
-
-			return await TryUpdateStreamAsync(stream, updateBuilder.Combine(updates));
-		}
-
-		/// <inheritdoc/>
-		public async Task<IStream?> GetAsync(StreamId streamId)
-		{
-			StreamDoc? stream = await _streams.Find<StreamDoc>(x => x.Id == streamId).FirstOrDefaultAsync();
-			if (stream != null)
-			{
-				await PostLoadAsync(stream);
-			}
-			return stream;
-		}
-
-		/// <inheritdoc/>
-		public async Task<IStreamPermissions?> GetPermissionsAsync(StreamId streamId)
-		{
-			return await _streams.Find<StreamDoc>(x => x.Id == streamId).Project<StreamPermissions>(StreamPermissions.Projection).FirstOrDefaultAsync();
-		}
-
-		/// <inheritdoc/>
-		public async Task<List<IStream>> FindAllAsync()
-		{
-			List<StreamDoc> results = await _streams.Find(Builders<StreamDoc>.Filter.Ne(x => x.Deleted, true)).ToListAsync();
-			foreach (StreamDoc result in results)
-			{
-				await PostLoadAsync(result);
-			}
-			return results.ConvertAll<IStream>(x => x);
-		}
-
-		/// <inheritdoc/>
-		public async Task<List<IStream>> FindForProjectsAsync(ProjectId[] projectIds)
-		{
-			FilterDefinition<StreamDoc> filter = Builders<StreamDoc>.Filter.In(x => x.ProjectId, projectIds) & Builders<StreamDoc>.Filter.Ne(x => x.Deleted, true);
-
-			List<StreamDoc> results = await _streams.Find(filter).ToListAsync();
-			foreach (StreamDoc result in results)
-			{
-				await PostLoadAsync(result);
-			}
-			return results.ConvertAll<IStream>(x => x);
-		}
-
 		/// <inheritdoc/>
 		public async Task<IStream?> TryUpdatePauseStateAsync(IStream streamInterface, DateTime? newPausedUntil, string? newPauseComment)
 		{
@@ -560,19 +429,13 @@ namespace Horde.Build.Streams
 			FindOneAndUpdateOptions<StreamDoc> options = new FindOneAndUpdateOptions<StreamDoc> { ReturnDocument = ReturnDocument.After };
 
 			StreamDoc? result = await _streams.FindOneAndUpdateAsync(filter, update, options);
-			if(result != null)
+			if (result != null)
 			{
-				await PostLoadAsync(result);
+				result.PostLoad(stream.Config, _clock.UtcNow);
 			}
 			return result;
 		}
 
-		/// <inheritdoc/>
-		public async Task DeleteAsync(StreamId streamId)
-		{
-			await _streams.UpdateOneAsync<StreamDoc>(x => x.Id == streamId, Builders<StreamDoc>.Update.Set(x => x.Deleted, true).Inc(x => x.UpdateIndex, 1));
-		}
-
 		/// <summary>
 		/// Checks the stream definition for consistency
 		/// </summary>
diff --git a/Engine/Source/Programs/Horde/Horde.Build/Streams/StreamConfig.cs b/Engine/Source/Programs/Horde/Horde.Build/Streams/StreamConfig.cs
--- a/Engine/Source/Programs/Horde/Horde.Build/Streams/StreamConfig.cs
+++ b/Engine/Source/Programs/Horde/Horde.Build/Streams/StreamConfig.cs
@@ -5,22 +5,27 @@ using System.Collections.Generic;
 using System.ComponentModel.DataAnnotations;
 using System.Diagnostics.CodeAnalysis;
 using System.Linq;
-using Datadog.Trace;
+using System.Security.Claims;
+using System.Text.Json;
+using System.Text.Json.Serialization;
+using System.Text.RegularExpressions;
 using EpicGames.Core;
 using EpicGames.Horde.Common;
-using EpicGames.Perforce;
 using Horde.Build.Acls;
 using Horde.Build.Agents.Pools;
+using Horde.Build.Configuration;
 using Horde.Build.Issues;
 using Horde.Build.Jobs.Graphs;
 using Horde.Build.Jobs.Templates;
 using Horde.Build.Perforce;
+using Horde.Build.Projects;
 using Horde.Build.Server;
 using Horde.Build.Utilities;
 using HordeCommon;
 
 namespace Horde.Build.Streams
 {
+	using StreamId = StringId<IStream>;
 	using PoolId = StringId<IPool>;
 	using TemplateId = StringId<ITemplateRef>;
 	using WorkflowId = StringId<WorkflowConfig>;
@@ -69,12 +74,40 @@ namespace Horde.Build.Streams
 	/// </summary>
 	[JsonSchema("https://unrealengine.com/horde/stream")]
 	[JsonSchemaCatalog("Horde Stream", "Horde stream configuration file", new[] { "*.stream.json", "Streams/*.json" })]
+	[ConfigIncludeRoot]
 	public class StreamConfig
 	{
+		/// <summary>
+		/// Accessor for the project containing this stream
+		/// </summary>
+		[JsonIgnore]
+		public ProjectConfig ProjectConfig { get; private set; } = null!;
+
+		/// <summary>
+		/// Identifier for the stream
+		/// </summary>
+		public StreamId Id { get; set; }
+
+		/// <summary>
+		/// Direct include path for the stream config. For backwards compatibility with old config files when including from a ProjectConfig object.
+		/// </summary>
+		[ConfigInclude, ConfigRelativePath]
+		public string? Path { get; set; }
+
+		/// <summary>
+		/// Includes for other configuration files
+		/// </summary>
+		public List<ConfigInclude> Include { get; set; } = new List<ConfigInclude>();
+
+		/// <summary>
+		/// Revision identifier for this configuration object
+		/// </summary>
+		[JsonIgnore]
+		public string Revision { get; set; } = String.Empty;
+
 		/// <summary>
 		/// Name of the stream
 		/// </summary>
-		[Required]
 		public string Name { get; set; } = String.Empty;
 
 		/// <summary>
@@ -186,6 +219,32 @@ namespace Horde.Build.Streams
 		/// Tokens to create for each job step
 		/// </summary>
 		public List<TokenConfig> Tokens { get; set; } = new List<TokenConfig>();
+		
+		/// <summary>
+		/// Callback after reading this stream configuration
+		/// </summary>
+		/// <param name="id">The stream id</param>
+		/// <param name="projectConfig">Owning project</param>
+		public void PostLoad(StreamId id, ProjectConfig projectConfig)
+		{
+			Id = id;
+			ProjectConfig = projectConfig;
+
+			foreach (TemplateRefConfig template in Templates)
+			{
+				template.PostLoad(this);
+			}
+		}
+
+		/// <summary>
+		/// Authorizes a user to perform a given action
+		/// </summary>
+		/// <param name="action">The action being performed</param>
+		/// <param name="user">The principal to validate</param>
+		public bool Authorize(AclAction action, ClaimsPrincipal user)
+		{
+			return Acl?.Authorize(action, user) ?? ProjectConfig.Authorize(action, user);
+		}
 
 		/// <summary>
 		/// Enumerates all commit tags, including the default tags for code and content.
@@ -312,8 +371,25 @@ namespace Horde.Build.Streams
 			workflowConfig = Workflows.FirstOrDefault(x => x.Id == workflowId);
 			return workflowConfig != null;
 		}
-	}
 
+		/// <summary>
+		/// The escaped name of this stream. Removes all non-identifier characters.
+		/// </summary>
+		/// <returns>Escaped name for the stream</returns>
+		public string GetEscapedName()
+		{
+			return Regex.Replace(Name, @"[^a-zA-Z0-9_]", "+");
+		}
+
+		/// <summary>
+		/// Gets the default identifier for workspaces created for this stream. Just includes an escaped depot name.
+		/// </summary>
+		/// <returns>The default workspace identifier</returns>
+		public string GetDefaultWorkspaceIdentifier()
+		{
+			return Regex.Replace(GetEscapedName(), @"^(\+\+[^+]*).*$", "$1");
+		}
+	}
 
 	/// <summary>
 	/// Information about a page to display in the dashboard for a stream
@@ -590,16 +666,16 @@ namespace Horde.Build.Streams
 	/// </summary>
 	public class TemplateRefConfig : TemplateConfig
 	{
-		TemplateId _id;
+		/// <summary>
+		/// The owning stream config
+		/// </summary>
+		[JsonIgnore]
+		public StreamConfig StreamConfig { get; private set; } = null!;
 
 		/// <summary>
 		/// Optional identifier for this ref. If not specified, an id will be generated from the name.
 		/// </summary>
-		public TemplateId Id
-		{
-			get => _id.IsEmpty ? TemplateId.Sanitize(Name) : _id;
-			set => _id = value;
-		}
+		public TemplateId Id { get; set; }
 
 		/// <summary>
 		/// Whether to show badges in UGS for these jobs
@@ -654,6 +730,30 @@ namespace Horde.Build.Streams
 		/// The ACL for this template
 		/// </summary>
 		public AclConfig? Acl { get; set; }
+
+		/// <summary>
+		/// Callback after the config is loaded
+		/// </summary>
+		/// <param name="streamConfig"></param>
+		public void PostLoad(StreamConfig streamConfig)
+		{
+			StreamConfig = streamConfig;
+
+			if (Id.IsEmpty)
+			{
+				Id = TemplateId.Sanitize(Name);
+			}
+		}
+
+		/// <summary>
+		/// Authorizes a user to perform a given action
+		/// </summary>
+		/// <param name="action">The action being performed</param>
+		/// <param name="user">The principal to validate</param>
+		public bool Authorize(AclAction action, ClaimsPrincipal user)
+		{
+			return Acl?.Authorize(action, user) ?? StreamConfig.Authorize(action, user);
+		}
 	}
 
 	/// <summary>
diff --git a/Engine/Source/Programs/Horde/Horde.Build/Streams/StreamMessages.cs b/Engine/Source/Programs/Horde/Horde.Build/Streams/StreamMessages.cs
--- a/Engine/Source/Programs/Horde/Horde.Build/Streams/StreamMessages.cs
+++ b/Engine/Source/Programs/Horde/Horde.Build/Streams/StreamMessages.cs
@@ -266,24 +266,18 @@ namespace Horde.Build.Streams
 		/// </summary>
 		public List<ChangeQueryConfig>? DefaultChange { get; set; }
 
-		/// <summary>
-		/// ACL for this template
-		/// </summary>
-		public GetAclResponse? Acl { get; set; }
-
 		/// <summary>
 		/// Constructor
 		/// </summary>
 		/// <param name="id">The template ref id</param>
 		/// <param name="templateRef">The template ref</param>
-		/// <param name="template">The template instance</param>
+		/// <param name="template">The actual template</param>
 		/// <param name="stepStates">The template step states</param>
-		/// <param name="includeAcl">Whether to include the ACL in the response</param>
-		public GetTemplateRefResponse(TemplateId id, ITemplateRef templateRef, ITemplate template, List<GetTemplateStepStateResponse>? stepStates, bool includeAcl)
+		public GetTemplateRefResponse(TemplateId id, ITemplateRef templateRef, ITemplate template, List<GetTemplateStepStateResponse>? stepStates)
 			: base(template)
 		{
 			Id = id.ToString();
-			Hash = templateRef.Hash.ToString();
+			Hash = template.Hash.ToString();
 			ShowUgsBadges = templateRef.Config.ShowUgsBadges;
 			ShowUgsAlerts = templateRef.Config.ShowUgsAlerts;
 			NotificationChannel = templateRef.Config.NotificationChannel;
@@ -292,7 +286,6 @@ namespace Horde.Build.Streams
 			ChainedJobs = (templateRef.Config.ChainedJobs != null && templateRef.Config.ChainedJobs.Count > 0) ? templateRef.Config.ChainedJobs : null;
 			StepStates = stepStates;
 			DefaultChange = templateRef.Config.DefaultChange;
-			Acl = (includeAcl && templateRef.Acl != null) ? new GetAclResponse(templateRef.Acl) : null;
 		}
 	}
 
@@ -376,11 +369,6 @@ namespace Horde.Build.Streams
 		/// </summary>
 		public List<GetTemplateRefResponse> Templates { get; set; }
 
-		/// <summary>
-		/// Custom permissions for this object
-		/// </summary>
-		public GetAclResponse? Acl { get; set; }
-
 		/// <summary>
 		/// Stream paused for new builds until this date
 		/// </summary>
@@ -399,43 +387,27 @@ namespace Horde.Build.Streams
 		/// <summary>
 		/// Constructor
 		/// </summary>
-		/// <param name="id">Unique id of the stream</param>
-		/// <param name="projectId">Unique id of the project containing the stream</param>
-		/// <param name="name">Name of the stream</param>
-		/// <param name="configRevision">The config file path on the server</param>
-		/// <param name="order">Order to display this stream</param>
-		/// <param name="notificationChannel"></param>
-		/// <param name="notificationChannelFilter"></param>
-		/// <param name="triageChannel"></param>
-		/// <param name="defaultPreflight">The default template to use for preflights</param>
-		/// <param name="tabs">List of tabs to display for this stream</param>
-		/// <param name="agentTypes">Map of agent type name to description</param>
-		/// <param name="workspaceTypes">Map of workspace name to description</param>
+		/// <param name="stream">The stream to construct from</param>
 		/// <param name="templates">Templates for this stream</param>
-		/// <param name="acl">Permissions for this object</param>
-		/// <param name="pausedUntil">Stream paused for new builds until this date</param>
-		/// <param name="pauseComment">Reason for stream being paused</param>
-		/// <param name="workflows">Workflows for this stream</param>
-		public GetStreamResponse(string id, string projectId, string name, string configRevision, int order, string? notificationChannel, string? notificationChannelFilter, string? triageChannel, DefaultPreflightConfig? defaultPreflight, List<TabConfig> tabs, Dictionary<string, AgentConfig> agentTypes, Dictionary<string, WorkspaceConfig>? workspaceTypes, List<GetTemplateRefResponse> templates, GetAclResponse? acl, DateTime? pausedUntil, string? pauseComment, List<WorkflowConfig> workflows)
+		public GetStreamResponse(IStream stream, List<GetTemplateRefResponse> templates)
 		{
-			Id = id;
-			ProjectId = projectId;
-			Name = name;
-			ConfigRevision = configRevision;
-			Order = order;
-			NotificationChannel = notificationChannel;
-			NotificationChannelFilter = notificationChannelFilter;
-			TriageChannel = triageChannel;
-			DefaultPreflightTemplate = defaultPreflight?.TemplateId?.ToString();
-			DefaultPreflight = defaultPreflight;
-			Tabs = tabs;
-			AgentTypes = agentTypes;
-			WorkspaceTypes = workspaceTypes;
+			Id = stream.Id.ToString();
+			ProjectId = stream.Config.ProjectConfig.Id.ToString();
+			Name = stream.Config.Name;
+			ConfigRevision = stream.Config.Revision;
+			Order = stream.Config.Order;
+			NotificationChannel = stream.Config.NotificationChannel;
+			NotificationChannelFilter = stream.Config.NotificationChannelFilter;
+			TriageChannel = stream.Config.TriageChannel;
+			DefaultPreflightTemplate = stream.Config.DefaultPreflight?.TemplateId?.ToString();
+			DefaultPreflight = stream.Config.DefaultPreflight;
+			Tabs = stream.Config.Tabs;
+			AgentTypes = stream.Config.AgentTypes;
+			WorkspaceTypes = stream.Config.WorkspaceTypes;
 			Templates = templates;
-			Acl = acl;
-			PausedUntil = pausedUntil;
-			PauseComment = pauseComment;
-			Workflows = workflows;
+			PausedUntil = stream.PausedUntil;
+			PauseComment = stream.PauseComment;
+			Workflows = stream.Config.Workflows;
 		}
 	}
 
@@ -449,21 +421,18 @@ namespace Horde.Build.Streams
 		/// <inheritdoc cref="IStream.Id"/>
 		public StreamId Id => _stream.Id;
 
-		/// <inheritdoc cref="IStream.ProjectId"/>
-		public ProjectId ProjectId => _stream.ProjectId;
+		/// <inheritdoc cref="ProjectConfig.Id"/>
+		public ProjectId ProjectId => _stream.Config.ProjectConfig.Id;
 
-		/// <inheritdoc cref="IStream.Name"/>
-		public string Name => _stream.Name;
+		/// <inheritdoc cref="StreamConfig.Name"/>
+		public string Name => _stream.Config.Name;
 
-		/// <inheritdoc cref="IStream.ConfigRevision"/>
-		public string ConfigRevision => _stream.ConfigRevision;
+		/// <inheritdoc cref="StreamConfig.Revision"/>
+		public string ConfigRevision => _stream.Config.Revision;
 
 		/// <inheritdoc cref="IStream.Config"/>
 		public StreamConfig? Config { get; }
 
-		/// <inheritdoc cref="IStream.Deleted"/>
-		public bool Deleted => _stream.Deleted;
-
 		/// <inheritdoc cref="IStream.Templates"/>
 		public List<GetTemplateRefResponseV2> Templates { get; }
 
@@ -473,11 +442,11 @@ namespace Horde.Build.Streams
 		/// <inheritdoc cref="IStream.PauseComment"/>
 		public string? PauseComment => _stream.PauseComment;
 
-		internal GetStreamResponseV2(IStream stream, bool includeConfig)
+		internal GetStreamResponseV2(IStream stream, bool includeConfig, List<GetTemplateRefResponseV2> templates)
 		{
 			_stream = stream;
 
-			Templates = stream.Templates.Values.Select(x => new GetTemplateRefResponseV2(x)).ToList();
+			Templates = templates;
 
 			if (includeConfig)
 			{
@@ -491,16 +460,18 @@ namespace Horde.Build.Streams
 	/// </summary>
 	public class GetTemplateRefResponseV2
 	{
+		readonly ITemplate _template;
 		readonly ITemplateRef _templateRef;
 
 		/// <inheritdoc cref="ITemplateRef.Id"/>
 		public TemplateId Id => _templateRef.Id;
 
-		/// <inheritdoc cref="ITemplateRef.Hash"/>
-		public ContentHash Hash => _templateRef.Hash;
+		/// <inheritdoc cref="ITemplate.Hash"/>
+		public ContentHash Hash => _template.Hash;
 
-		internal GetTemplateRefResponseV2(ITemplateRef templateRef)
+		internal GetTemplateRefResponseV2(ITemplate template, ITemplateRef templateRef)
 		{
+			_template = template;
 			_templateRef = templateRef;
 		}
 	}
diff --git a/Engine/Source/Programs/Horde/Horde.Build/Streams/StreamsController.cs b/Engine/Source/Programs/Horde/Horde.Build/Streams/StreamsController.cs
--- a/Engine/Source/Programs/Horde/Horde.Build/Streams/StreamsController.cs
+++ b/Engine/Source/Programs/Horde/Horde.Build/Streams/StreamsController.cs
@@ -10,10 +10,12 @@ using Horde.Build.Jobs;
 using Horde.Build.Jobs.Templates;
 using Horde.Build.Perforce;
 using Horde.Build.Projects;
+using Horde.Build.Server;
 using Horde.Build.Users;
 using Horde.Build.Utilities;
 using Microsoft.AspNetCore.Authorization;
 using Microsoft.AspNetCore.Mvc;
+using Microsoft.Extensions.Options;
 using OpenTracing;
 using OpenTracing.Util;
 
@@ -31,24 +33,24 @@ namespace Horde.Build.Streams
 	[Route("[controller]")]
 	public class StreamsController : HordeControllerBase
 	{
-		private readonly StreamService _streamService;
+		private readonly IStreamCollection _streamCollection;
 		private readonly ICommitService _commitService;
 		private readonly ITemplateCollection _templateCollection;
 		private readonly IJobStepRefCollection _jobStepRefCollection;
 		private readonly IUserCollection _userCollection;
-		private readonly AclService _aclService;
+		private readonly IOptionsSnapshot<GlobalConfig> _globalConfig;
 
 		/// <summary>
 		/// Constructor
 		/// </summary>
-		public StreamsController(AclService aclService, StreamService streamService, ICommitService commitService, ITemplateCollection templateCollection, IJobStepRefCollection jobStepRefCollection, IUserCollection userCollection)
+		public StreamsController(IStreamCollection streamCollection, ICommitService commitService, ITemplateCollection templateCollection, IJobStepRefCollection jobStepRefCollection, IUserCollection userCollection, IOptionsSnapshot<GlobalConfig> globalConfig)
 		{
-			_streamService = streamService;
+			_streamCollection = streamCollection;
 			_commitService = commitService;
 			_templateCollection = templateCollection;
 			_jobStepRefCollection = jobStepRefCollection;
 			_userCollection = userCollection;
-			_aclService = aclService;
+			_globalConfig = globalConfig;
 		}
 
 		/// <summary>
@@ -62,21 +64,7 @@ namespace Horde.Build.Streams
 		[ProducesResponseType(typeof(List<GetStreamResponse>), 200)]
 		public async Task<ActionResult<List<object>>> GetStreamsAsync([FromQuery(Name = "ProjectId")] string[] projectIds, [FromQuery] PropertyFilter? filter = null)
 		{
-			ProjectId[] projectIdValues = Array.ConvertAll(projectIds, x => new ProjectId(x));
-
-			List<IStream> streams = await _streamService.GetStreamsAsync(projectIdValues);
-			ProjectPermissionsCache permissionsCache = new ProjectPermissionsCache();
-
-			List<GetStreamResponse> responses = new List<GetStreamResponse>();
-			foreach (IStream stream in streams)
-			{
-				if (await _streamService.AuthorizeAsync(stream, AclAction.ViewStream, User, permissionsCache))
-				{
-					GetStreamResponse response = await CreateGetStreamResponse(stream, permissionsCache);
-					responses.Add(response);
-				}
-			}
-			return responses.OrderBy(x => x.Order).ThenBy(x => x.Name).Select(x => PropertyFilter.Apply(x, filter)).ToList();
+			return await GetStreamsInternalAsync(projectIds, true, filter);
 		}
 
 		/// <summary>
@@ -90,23 +78,43 @@ namespace Horde.Build.Streams
 		[ProducesResponseType(typeof(List<GetStreamResponse>), 200)]
 		public async Task<ActionResult<List<object>>> GetStreamsAsyncV2([FromQuery(Name = "ProjectId")] string[] projectIds, [FromQuery] PropertyFilter? filter = null)
 		{
-			ProjectId[] projectIdValues = Array.ConvertAll(projectIds, x => new ProjectId(x));
+			return await GetStreamsInternalAsync(projectIds, false, filter);
+		}
 
-			List<IStream> streams = await _streamService.GetStreamsAsync(projectIdValues);
-			ProjectPermissionsCache permissionsCache = new ProjectPermissionsCache();
+		async Task<List<object>> GetStreamsInternalAsync([FromQuery(Name = "ProjectId")] string[] projectIds, bool includeConfig, PropertyFilter? filter = null)
+		{
+			ProjectId[] projectIdValues = Array.ConvertAll(projectIds, x => new ProjectId(x));
 
 			List<GetStreamResponseV2> responses = new List<GetStreamResponseV2>();
-			foreach (IStream stream in streams)
+			foreach (ProjectConfig projectConfig in _globalConfig.Value.Projects)
 			{
-				if (await _streamService.AuthorizeAsync(stream, AclAction.ViewStream, User, permissionsCache))
+				if (projectIdValues.Length == 0 || projectIdValues.Contains(projectConfig.Id))
 				{
-					GetStreamResponseV2 response = new GetStreamResponseV2(stream, false);
-					responses.Add(response);
+					foreach (StreamConfig streamConfig in projectConfig.Streams)
+					{
+						if (streamConfig.Authorize(AclAction.ViewStream, User))
+						{
+							IStream stream = await _streamCollection.GetAsync(streamConfig);
+							GetStreamResponseV2 response = await CreateStreamResponseAsync(stream, includeConfig);
+							responses.Add(response);
+						}
+					}
 				}
 			}
 			return responses.OrderBy(x => x.Id).Select(x => PropertyFilter.Apply(x, filter)).ToList();
 		}
 
+		async Task<GetStreamResponseV2> CreateStreamResponseAsync(IStream stream, bool includeConfig)
+		{
+			List<GetTemplateRefResponseV2> templates = new List<GetTemplateRefResponseV2>();
+			foreach (ITemplateRef templateRef in stream.Templates.Values)
+			{
+				ITemplate template = await _templateCollection.GetOrAddAsync(templateRef.Config);
+				templates.Add(new GetTemplateRefResponseV2(template, templateRef));
+			}
+			return new GetStreamResponseV2(stream, includeConfig, templates);
+		}
+
 		/// <summary>
 		/// Retrieve information about a specific stream.
 		/// </summary>
@@ -118,19 +126,18 @@ namespace Horde.Build.Streams
 		[ProducesResponseType(typeof(GetStreamResponse), 200)]
 		public async Task<ActionResult<object>> GetStreamAsync(StreamId streamId, [FromQuery] PropertyFilter? filter = null)
 		{
-			IStream? stream = await _streamService.GetStreamAsync(streamId);
-			if (stream == null)
+			StreamConfig? streamConfig;
+			if (!_globalConfig.Value.TryGetStream(streamId, out streamConfig))
 			{
 				return NotFound(streamId);
 			}
-
-			ProjectPermissionsCache permissionsCache = new ProjectPermissionsCache();
-			if (!await _streamService.AuthorizeAsync(stream, AclAction.ViewStream, User, permissionsCache))
+			if (!streamConfig.Authorize(AclAction.ViewStream, User))
 			{
 				return Forbid(AclAction.ViewStream, streamId);
 			}
 
-			return PropertyFilter.Apply(await CreateGetStreamResponse(stream, permissionsCache), filter);
+			IStream stream = await _streamCollection.GetAsync(streamConfig);
+			return PropertyFilter.Apply(await CreateGetStreamResponse(stream), filter);
 		}
 
 		/// <summary>
@@ -145,70 +152,60 @@ namespace Horde.Build.Streams
 		[ProducesResponseType(typeof(GetStreamResponseV2), 200)]
 		public async Task<ActionResult<object>> GetStreamAsyncV2(StreamId streamId, [FromQuery] string? config = null, [FromQuery] PropertyFilter? filter = null)
 		{
-			IStream? stream = await _streamService.GetStreamAsync(streamId);
-			if (stream == null)
+			StreamConfig? streamConfig;
+			if (!_globalConfig.Value.TryGetStream(streamId, out streamConfig))
 			{
 				return NotFound(streamId);
 			}
-
-			ProjectPermissionsCache permissionsCache = new ProjectPermissionsCache();
-			if (!await _streamService.AuthorizeAsync(stream, AclAction.ViewStream, User, permissionsCache))
+			if (!streamConfig.Authorize(AclAction.ViewStream, User))
 			{
 				return Forbid(AclAction.ViewStream, streamId);
 			}
 
-			bool includeConfig = !String.Equals(config, stream.ConfigRevision, StringComparison.Ordinal); 
-			return PropertyFilter.Apply(new GetStreamResponseV2(stream, includeConfig), filter);
+			bool includeConfig = !String.Equals(config, streamConfig.Revision, StringComparison.Ordinal);
+			IStream stream = await _streamCollection.GetAsync(streamConfig);
+			return PropertyFilter.Apply(await CreateStreamResponseAsync(stream, includeConfig), filter);
 		}
 
 		/// <summary>
 		/// Create a stream response object, including all the templates
 		/// </summary>
 		/// <param name="stream">Stream to create response for</param>
-		/// <param name="cache">Permissions cache</param>
 		/// <returns>Response object</returns>
-		async Task<GetStreamResponse> CreateGetStreamResponse(IStream stream, ProjectPermissionsCache cache)
+		async Task<GetStreamResponse> CreateGetStreamResponse(IStream stream)
 		{
 			using IScope scope = GlobalTracer.Instance.BuildSpan("CreateGetStreamResponse").StartActive();
 			scope.Span.SetTag("streamId", stream.Id);
 			
-			bool includeAcl = stream.Acl != null && await _streamService.AuthorizeAsync(stream, AclAction.ViewPermissions, User, cache);
-
 			List<GetTemplateRefResponse> apiTemplateRefs = new List<GetTemplateRefResponse>();
 			foreach (KeyValuePair<TemplateId, ITemplateRef> pair in stream.Templates)
 			{
 				using IScope templateScope = GlobalTracer.Instance.BuildSpan("CreateGetStreamResponse.Template").StartActive();
 				templateScope.Span.SetTag("templateName", pair.Value.Config.Name);
-				
-				if (await _streamService.AuthorizeAsync(stream, pair.Value, AclAction.ViewTemplate, User, cache))
+
+				ITemplateRef templateRef = pair.Value;
+				if (templateRef.Config.Authorize(AclAction.ViewTemplate, User))
 				{
-					ITemplate? template = await _templateCollection.GetAsync(pair.Value.Hash);
-					
-					if (template != null)						
+					List<GetTemplateStepStateResponse>? stepStates = null;
+					if (templateRef.StepStates != null)
 					{
-						ITemplateRef tref = pair.Value;
-
-						List<GetTemplateStepStateResponse>? stepStates = null;
-						if (tref.StepStates != null)
+						for (int i = 0; i < templateRef.StepStates.Count; i++)
 						{
-							for (int i = 0; i < tref.StepStates.Count; i++)
-							{
-								ITemplateStep state = tref.StepStates[i];
+							ITemplateStep state = templateRef.StepStates[i];
 
-								stepStates ??= new List<GetTemplateStepStateResponse>();
+							stepStates ??= new List<GetTemplateStepStateResponse>();
 
-								GetThinUserInfoResponse? pausedByUserInfo = new GetThinUserInfoResponse(await _userCollection.GetCachedUserAsync(state.PausedByUserId));
-								stepStates.Add(new GetTemplateStepStateResponse(state, pausedByUserInfo));
-							}
+							GetThinUserInfoResponse? pausedByUserInfo = new GetThinUserInfoResponse(await _userCollection.GetCachedUserAsync(state.PausedByUserId));
+							stepStates.Add(new GetTemplateStepStateResponse(state, pausedByUserInfo));
 						}
-
-						bool includeTemplateAcl = pair.Value.Acl != null && await _streamService.AuthorizeAsync(stream, pair.Value, AclAction.ViewPermissions, User, cache);
-						apiTemplateRefs.Add(new GetTemplateRefResponse(pair.Key, pair.Value, template, stepStates, includeTemplateAcl));
 					}
+
+					ITemplate? template = await _templateCollection.GetOrAddAsync(pair.Value.Config);
+					apiTemplateRefs.Add(new GetTemplateRefResponse(pair.Key, pair.Value, template, stepStates));
 				}
 			}
 
-			return stream.ToApiResponse(includeAcl, apiTemplateRefs);
+			return stream.ToApiResponse(apiTemplateRefs);
 		}
 
 		/// <summary>
@@ -226,12 +223,12 @@ namespace Horde.Build.Streams
 		[ProducesResponseType(typeof(List<GetCommitResponse>), 200)]
 		public async Task<ActionResult<List<object>>> GetChangesAsync(StreamId streamId, [FromQuery] int? min = null, [FromQuery] int? max = null, [FromQuery] int results = 50, [FromQuery] string? tags = null, PropertyFilter? filter = null)
 		{
-			IStream? stream = await _streamService.GetStreamAsync(streamId);
-			if (stream == null)
+			StreamConfig? streamConfig;
+			if (!_globalConfig.Value.TryGetStream(streamId, out streamConfig))
 			{
 				return NotFound(streamId);
 			}
-			if (!await _streamService.AuthorizeAsync(stream, AclAction.ViewChanges, User, null))
+			if (!streamConfig.Authorize(AclAction.ViewChanges, User))
 			{
 				return Forbid(AclAction.ViewChanges, streamId);
 			}
@@ -242,7 +239,7 @@ namespace Horde.Build.Streams
 				commitTags = tags.Split(';', StringSplitOptions.RemoveEmptyEntries).Select(x => new CommitTag(x)).ToList();
 			}
 
-			List<ICommit> commits = await _commitService.GetCollection(stream).FindAsync(min, max, results, commitTags).ToListAsync();
+			List<ICommit> commits = await _commitService.GetCollection(streamConfig).FindAsync(min, max, results, commitTags).ToListAsync();
 
 			List<GetCommitResponse> responses = new List<GetCommitResponse>();
 			foreach (ICommit commit in commits)
@@ -265,17 +262,17 @@ namespace Horde.Build.Streams
 		[ProducesResponseType(typeof(GetCommitResponse), 200)]
 		public async Task<ActionResult<object>> GetChangeDetailsAsync(StreamId streamId, int changeNumber, PropertyFilter? filter = null)
 		{
-			IStream? stream = await _streamService.GetStreamAsync(streamId);
-			if (stream == null)
+			StreamConfig? streamConfig;
+			if (!_globalConfig.Value.TryGetStream(streamId, out streamConfig))
 			{
 				return NotFound(streamId);
 			}
-			if (!await _streamService.AuthorizeAsync(stream, AclAction.ViewChanges, User, null))
+			if (!streamConfig.Authorize(AclAction.ViewChanges, User))
 			{
 				return Forbid(AclAction.ViewChanges, streamId);
 			}
 
-			ICommit? changeDetails = await _commitService.GetCollection(stream).GetAsync(changeNumber);
+			ICommit? changeDetails = await _commitService.GetCollection(streamConfig).GetAsync(changeNumber);
 			if(changeDetails == null)
 			{
 				return NotFound("CL {Change} not found in stream {StreamId}", changeNumber, streamId);
@@ -303,12 +300,12 @@ namespace Horde.Build.Streams
 		[ProducesResponseType(typeof(List<GetJobStepRefResponse>), 200)]
 		public async Task<ActionResult<List<object>>> GetStepHistoryAsync(StreamId streamId, [FromQuery] string templateId, [FromQuery] string step, [FromQuery] int? change = null, [FromQuery] int count = 10, [FromQuery] PropertyFilter? filter = null)
 		{
-			IStream? stream = await _streamService.GetStreamAsync(streamId);
-			if (stream == null)
+			StreamConfig? streamConfig;
+			if (!_globalConfig.Value.TryGetStream(streamId, out streamConfig))
 			{
 				return NotFound(streamId);
 			}
-			if (!await _streamService.AuthorizeAsync(stream, AclAction.ViewJob, User, null))
+			if (!streamConfig.Authorize(AclAction.ViewJob, User))
 			{
 				return Forbid(AclAction.ViewJob, streamId);
 			}
@@ -319,29 +316,6 @@ namespace Horde.Build.Streams
 			return steps.ConvertAll(x => PropertyFilter.Apply(new GetJobStepRefResponse(x), filter));
 		}
 
-		/// <summary>
-		/// Deletes a stream
-		/// </summary>
-		/// <param name="streamId">Id of the stream to update.</param>
-		/// <returns>Http result code</returns>
-		[HttpDelete]
-		[Route("/api/v1/streams/{streamId}")]
-		public async Task<ActionResult> DeleteStreamAsync(StreamId streamId)
-		{
-			IStream? stream = await _streamService.GetStreamAsync(streamId);
-			if (stream == null)
-			{
-				return NotFound(streamId);
-			}
-			if (!await _streamService.AuthorizeAsync(stream, AclAction.DeleteStream, User, null))
-			{
-				return Forbid(AclAction.DeleteStream, streamId);
-			}
-
-			await _streamService.DeleteStreamAsync(streamId);
-			return new OkResult();
-		}
-
 		/// <summary>
 		/// Gets a template for a stream
 		/// </summary>
@@ -354,28 +328,23 @@ namespace Horde.Build.Streams
 		[ProducesResponseType(typeof(List<GetTemplateResponse>), 200)]
 		public async Task<ActionResult<object>> GetTemplateAsync(StreamId streamId, TemplateId templateId, [FromQuery] PropertyFilter? filter = null)
 		{
-			IStream? stream = await _streamService.GetStreamAsync(streamId);
-			if (stream == null)
+			StreamConfig? streamConfig;
+			if (!_globalConfig.Value.TryGetStream(streamId, out streamConfig))
 			{
 				return NotFound(streamId);
 			}
 
-			ITemplateRef? templateRef;
-			if (!stream.Templates.TryGetValue(templateId, out templateRef))
+			TemplateRefConfig? templateConfig;
+			if (!streamConfig.TryGetTemplate(templateId, out templateConfig))
 			{
 				return NotFound(streamId, templateId);
 			}
-			if (!await _streamService.AuthorizeAsync(stream, templateRef, AclAction.ViewTemplate, User, null))
+			if (!templateConfig.Authorize(AclAction.ViewTemplate, User))
 			{
 				return Forbid(AclAction.ViewTemplate, streamId);
 			}
 
-			ITemplate? template = await _templateCollection.GetAsync(templateRef.Hash);
-			if(template == null)
-			{
-				return NotFound();
-			}
-
+			ITemplate template = await _templateCollection.GetOrAddAsync(templateConfig);
 			return new GetTemplateResponse(template).ApplyFilter(filter);
 		}
 
@@ -387,24 +356,25 @@ namespace Horde.Build.Streams
 		[Route("/api/v1/streams/{streamId}/templates/{templateRefId}")]
 		public async Task<ActionResult> UpdateStreamTemplateRefAsync(StreamId streamId, TemplateId templateRefId, [FromBody] UpdateTemplateRefRequest update)
 		{
-			if (!await _aclService.AuthorizeAsync(AclAction.AdminWrite, User))
+			if (!_globalConfig.Value.Authorize(AclAction.AdminWrite, User))
 			{
 				return Forbid(AclAction.AdminWrite);
 			}
 
-			IStream? stream = await _streamService.GetStreamAsync(streamId);
-			if (stream == null)
+			StreamConfig? streamConfig;
+			if (!_globalConfig.Value.TryGetStream(streamId, out streamConfig))
 			{
 				return NotFound(streamId);
 			}
 
+
+			IStream stream = await _streamCollection.GetAsync(streamConfig);
 			if (!stream.Templates.ContainsKey(templateRefId))
 			{
 				return NotFound(streamId, templateRefId);
 			}
 
-			await _streamService.TryUpdateTemplateRefAsync(stream, templateRefId, update.StepStates);
-
+			await _streamCollection.TryUpdateTemplateRefAsync(stream, templateRefId, update.StepStates);
 			return Ok();
 		}
 	}
diff --git a/Engine/Source/Programs/Horde/Horde.Build/Tasks/ConformTaskSource.cs b/Engine/Source/Programs/Horde/Horde.Build/Tasks/ConformTaskSource.cs
--- a/Engine/Source/Programs/Horde/Horde.Build/Tasks/ConformTaskSource.cs
+++ b/Engine/Source/Programs/Horde/Horde.Build/Tasks/ConformTaskSource.cs
@@ -47,12 +47,13 @@ namespace Horde.Build.Tasks
 		readonly ILogFileService _logService;
 		readonly ILogger _logger;
 		readonly IOptionsMonitor<ServerSettings> _settings;
+		readonly IOptionsMonitor<GlobalConfig> _globalConfig;
 		readonly ITicker _tickConformList;
 
 		/// <summary>
 		/// Constructor
 		/// </summary>
-		public ConformTaskSource(MongoService mongoService, GlobalsService globalsService, IAgentCollection agentCollection, PoolService poolService, ILogFileService logService, PerforceLoadBalancer perforceLoadBalancer, IClock clock, IOptionsMonitor<ServerSettings> settings, ILogger<ConformTaskSource> logger)
+		public ConformTaskSource(MongoService mongoService, GlobalsService globalsService, IAgentCollection agentCollection, PoolService poolService, ILogFileService logService, PerforceLoadBalancer perforceLoadBalancer, IClock clock, IOptionsMonitor<ServerSettings> settings, IOptionsMonitor<GlobalConfig> globalConfig, ILogger<ConformTaskSource> logger)
 		{
 			_mongoService = mongoService;
 			_globalsService = globalsService;
@@ -62,6 +63,7 @@ namespace Horde.Build.Tasks
 			_perforceLoadBalancer = perforceLoadBalancer;
 			_logService = logService;
 			_settings = settings;
+			_globalConfig = globalConfig;
 			_logger = logger;
 			_tickConformList = clock.AddSharedTicker<ConformTaskSource>(TimeSpan.FromMinutes(1.0), CleanConformListAsync, logger);
 
@@ -208,12 +210,12 @@ namespace Horde.Build.Tasks
 		/// <inheritdoc/>
 		public async Task<bool> GetWorkspacesAsync(IAgent agent, IList<HordeCommon.Rpc.Messages.AgentWorkspace> workspaces)
 		{
-			IGlobals globals = await _globalsService.GetAsync();
+			GlobalConfig globalConfig = _globalConfig.CurrentValue;
 
-			HashSet<AgentWorkspace> conformWorkspaces = await _poolService.GetWorkspacesAsync(agent, DateTime.UtcNow);
+			HashSet<AgentWorkspace> conformWorkspaces = await _poolService.GetWorkspacesAsync(agent, DateTime.UtcNow, globalConfig);
 			foreach (AgentWorkspace conformWorkspace in conformWorkspaces)
 			{
-				PerforceCluster? cluster = globals.Config.FindPerforceCluster(conformWorkspace.Cluster);
+				PerforceCluster? cluster = globalConfig.FindPerforceCluster(conformWorkspace.Cluster);
 				if (cluster == null || !await agent.TryAddWorkspaceMessage(conformWorkspace, cluster, _perforceLoadBalancer, workspaces))
 				{
 					return false;
@@ -238,11 +240,11 @@ namespace Horde.Build.Tasks
 		/// <returns>True if the resource was allocated, false otherwise</returns>
 		private async Task<bool> AllocateConformLeaseAsync(AgentId agentId, IEnumerable<HordeCommon.Rpc.Messages.AgentWorkspace> workspaces, LeaseId leaseId)
 		{
-			IGlobals globals = await _globalsService.GetAsync();
+			GlobalConfig globalConfig = _globalConfig.CurrentValue;
 			for (; ; )
 			{
 				ConformList currentValue = await _conformList.GetAsync();
-				if (globals.Config.MaxConformCount != 0 && currentValue.Entries.Count + currentValue.Servers.Sum(x => x.Entries.Count) >= globals.Config.MaxConformCount)
+				if (globalConfig.MaxConformCount != 0 && currentValue.Entries.Count + currentValue.Servers.Sum(x => x.Entries.Count) >= globalConfig.MaxConformCount)
 				{
 					return false;
 				}
@@ -252,7 +254,7 @@ namespace Horde.Build.Tasks
 				{
 					if (servers.Add(workspace.ServerAndPort))
 					{
-						PerforceCluster? cluster = globals.Config.FindPerforceCluster(workspace.Cluster);
+						PerforceCluster? cluster = globalConfig.FindPerforceCluster(workspace.Cluster);
 						if (cluster == null)
 						{
 							_logger.LogWarning("Unable to find perforce cluster '{Cluster}' for conform", workspace.Cluster);
@@ -325,6 +327,8 @@ namespace Horde.Build.Tasks
 		/// <returns>True if the agent should be conformed</returns>
 		private async Task<bool> IsConformPendingAsync(IAgent agent, DateTime utcNow)
 		{
+			GlobalConfig globalConfig = _globalConfig.CurrentValue;
+
 			// If a conform was manually requested, allow it to run even if the agent is disabled
 			if (agent.RequestConform || agent.RequestFullConform)
 			{
@@ -366,10 +370,10 @@ namespace Horde.Build.Tasks
 				}
 
 				// Check if the workspaces have changed (first check against a cached list of workspaces, then an accurate one)
-				HashSet<AgentWorkspace> workspaces = await _poolService.GetWorkspacesAsync(agent, utcNow - TimeSpan.FromSeconds(30.0));
+				HashSet<AgentWorkspace> workspaces = await _poolService.GetWorkspacesAsync(agent, utcNow - TimeSpan.FromSeconds(30.0), globalConfig);
 				if (!workspaces.SetEquals(agent.Workspaces))
 				{
-					workspaces = await _poolService.GetWorkspacesAsync(agent, utcNow);
+					workspaces = await _poolService.GetWorkspacesAsync(agent, utcNow, globalConfig);
 					if (!workspaces.SetEquals(agent.Workspaces))
 					{
 						return true;
diff --git a/Engine/Source/Programs/Horde/Horde.Build/Tasks/UpgradeTaskSource.cs b/Engine/Source/Programs/Horde/Horde.Build/Tasks/UpgradeTaskSource.cs
--- a/Engine/Source/Programs/Horde/Horde.Build/Tasks/UpgradeTaskSource.cs
+++ b/Engine/Source/Programs/Horde/Horde.Build/Tasks/UpgradeTaskSource.cs
@@ -30,17 +30,17 @@ namespace Horde.Build.Tasks
 		public override TaskSourceFlags Flags => TaskSourceFlags.AllowWhenDisabled | TaskSourceFlags.AllowDuringDowntime;
 
 		readonly AgentSoftwareService _agentSoftwareService;
-		readonly GlobalsService _globalsService;
 		readonly ILogFileService _logService;
-		readonly IOptionsMonitor<ServerSettings> _settings;
+		readonly IOptionsMonitor<GlobalConfig> _globalConfig;
+		readonly IOptions<ServerSettings> _serverSettings;
 		readonly IClock _clock;
 
-		public UpgradeTaskSource(AgentSoftwareService agentSoftwareService, GlobalsService globalsService, ILogFileService logService, IOptionsMonitor<ServerSettings> settings, IClock clock)
+		public UpgradeTaskSource(AgentSoftwareService agentSoftwareService, ILogFileService logService, IOptionsMonitor<GlobalConfig> globalConfig, IOptions<ServerSettings> serverSettings, IClock clock)
 		{
 			_agentSoftwareService = agentSoftwareService;
-			_globalsService = globalsService;
 			_logService = logService;
-			_settings = settings;
+			_globalConfig = globalConfig;
+			_serverSettings = serverSettings;
 			_clock = clock;
 
 			OnLeaseStartedProperties.Add(nameof(UpgradeTask.LogId), x => new LogId(x.LogId));
@@ -48,7 +48,7 @@ namespace Horde.Build.Tasks
 
 		public override async Task<Task<AgentLease?>> AssignLeaseAsync(IAgent agent, CancellationToken cancellationToken)
 		{
-			if (!_settings.CurrentValue.EnableUpgradeTasks)
+			if (!_serverSettings.Value.EnableUpgradeTasks)
 			{
 				return Skip(cancellationToken);
 			}
@@ -92,11 +92,10 @@ namespace Horde.Build.Tasks
 		/// <returns>Unique id of the client version this agent should be running</returns>
 		public async Task<string?> GetRequiredSoftwareVersion(IAgent agent)
 		{
-			IGlobals globals = await _globalsService.GetAsync();
-			GlobalConfig config = globals.Config;
+			GlobalConfig globalConfig = _globalConfig.CurrentValue;
 
 			AgentSoftwareChannelName channelName = AgentSoftwareService.DefaultChannelName;
-			foreach (AgentSoftwareConfig softwareConfig in config.Software)
+			foreach (AgentSoftwareConfig softwareConfig in globalConfig.Software)
 			{
 				if (softwareConfig.Condition != null && agent.SatisfiesCondition(softwareConfig.Condition))
 				{
diff --git a/Engine/Source/Programs/Horde/Horde.Build/Tools/ITool.cs b/Engine/Source/Programs/Horde/Horde.Build/Tools/ITool.cs
--- a/Engine/Source/Programs/Horde/Horde.Build/Tools/ITool.cs
+++ b/Engine/Source/Programs/Horde/Horde.Build/Tools/ITool.cs
@@ -24,35 +24,15 @@ namespace Horde.Build.Tools
 		/// </summary>
 		ToolId Id { get; }
 
-		/// <summary>
-		/// Name of the tool
-		/// </summary>
-		string Name { get; }
-
-		/// <summary>
-		/// Description for the tool
-		/// </summary>
-		string Description { get; }
-
-		/// <summary>
-		/// Last time that the tool was updated
-		/// </summary>
-		DateTime LastUpdateTime { get; }
-
 		/// <summary>
 		/// Current deployments of this tool, sorted by time.
 		/// </summary>
 		IReadOnlyList<IToolDeployment> Deployments { get; }
 
 		/// <summary>
-		/// Whether this tool should be exposed for download on a public endpoint without authentication
-		/// </summary>
-		bool Public { get; }
-
-		/// <summary>
-		/// Access list for this tool
+		/// Config object for this tool
 		/// </summary>
-		public AclV2? Acl { get; }
+		ToolConfig Config { get; }
 	}
 
 	/// <summary>
@@ -132,22 +112,6 @@ namespace Horde.Build.Tools
 	/// </summary>
 	public static class ToolExtensions
 	{
-		/// <summary>
-		/// Test whether a user can perform an action on this tool 
-		/// </summary>
-		public static Task<bool> AuthorizeAsync(this ITool tool, AclAction action, ClaimsPrincipal user, AclService aclService, GlobalPermissionsCache? cache)
-		{
-			bool? result = tool.Acl?.Authorize(action, user);
-			if (result == null)
-			{
-				return aclService.AuthorizeAsync(action, user, cache);
-			}
-			else
-			{
-				return Task.FromResult(result.Value);
-			}
-		}
-
 		/// <summary>
 		/// Get the progress fraction for a particular deployment and time
 		/// </summary>
diff --git a/Engine/Source/Programs/Horde/Horde.Build/Tools/ToolCollection.cs b/Engine/Source/Programs/Horde/Horde.Build/Tools/ToolCollection.cs
--- a/Engine/Source/Programs/Horde/Horde.Build/Tools/ToolCollection.cs
+++ b/Engine/Source/Programs/Horde/Horde.Build/Tools/ToolCollection.cs
@@ -3,7 +3,6 @@
 using EpicGames.Core;
 using EpicGames.Horde.Storage;
 using EpicGames.Serialization;
-using Horde.Build.Acls;
 using Horde.Build.Server;
 using Horde.Build.Storage;
 using Horde.Build.Utilities;
@@ -33,51 +32,26 @@ namespace Horde.Build.Tools
 	{
 		class Tool : VersionedDocument<ToolId, Tool>, ITool
 		{
-			/// <summary>
-			/// Name of the tool
-			/// </summary>
-			[BsonElement("nam")]
-			public string Name { get; set; }
-
-			/// <summary>
-			/// Description for the tool
-			/// </summary>
-			[BsonElement("dsc")]
-			public string Description { get; set; }
+			[BsonIgnore]
+			public ToolConfig Config { get; set; } = null!;
 
 			[BsonElement("dep")]
 			public List<ToolDeployment> Deployments { get; set; } = new List<ToolDeployment>();
 
+			// ITool interface
 			IReadOnlyList<IToolDeployment> ITool.Deployments => Deployments;
 
-			[BsonElement("pub")]
-			public bool Public { get; set; }
-
-			[BsonElement("acl")]
-			public AclV2? Acl
-			{
-				get => _acl;
-				set => _acl = (value is not null && !value.IsDefault()) ? value : null;
-			}
-
-			private AclV2? _acl;
-
 			[BsonConstructor]
 			public Tool(ToolId id)
 				: base(id)
 			{
-				Id = id;
-				Name = id.ToString();
-				Description = String.Empty;
+				Config = null!;
 			}
 
 			public Tool(ToolConfig config)
 				: base(config.Id)
 			{
-				Name = config.Name;
-				Description = config.Description;
-				Public = config.Public;
-				Acl = config.Acl;
+				Config = config;
 			}
 
 			/// <inheritdoc/>
@@ -185,7 +159,6 @@ namespace Horde.Build.Tools
 			public List<ToolId> Ids { get; set; } = new List<ToolId>();
 		}
 
-		private readonly RedisService _redisService;
 		private readonly VersionedCollection<ToolId, Tool> _tools;
 		private readonly ILegacyStorageClient _storage;
 		private readonly IClock _clock;
@@ -193,16 +166,14 @@ namespace Horde.Build.Tools
 		private readonly ILogger _logger;
 
 		private static readonly RedisKey s_baseKey = "tools/v1/";
-		private static readonly RedisKey s_indexKey = s_baseKey.Append("index");
 
 		private static readonly IReadOnlyDictionary<int, Type> s_types = RegisterTypes();
 
 		/// <summary>
 		/// Constructor
 		/// </summary>
-		public ToolCollection(MongoService mongoService, RedisService redisService, ILegacyStorageClient storage, IClock clock, IOptions<ServerSettings> options, ILogger<ToolCollection> logger)
+		public ToolCollection(MongoService mongoService, RedisService redisService, ILegacyStorageClient storage, IClock clock, ILogger<ToolCollection> logger)
 		{
-			_redisService = redisService;
 			_tools = new VersionedCollection<ToolId, Tool>(mongoService, "Tools", redisService, s_baseKey, s_types);
 			_storage = storage;
 			_clock = clock;
@@ -223,188 +194,36 @@ namespace Horde.Build.Tools
 			{
 				cm.AutoMap();
 				cm.MapCreator(t => new Tool(t.Id));
-				cm.MapMember(x => x.Acl).SetIgnoreIfDefault(true).SetDefaultValue(new AclV2());
 			});
 
 			return versionToType;
 		}
 
-		/// <summary>
-		/// Configures all the available tools
-		/// </summary>
-		/// <param name="tools">The list of configured tools</param>
-		public async Task ConfigureAsync(List<ToolConfig> tools)
-		{
-			List<ToolId> ids = await FindAllIdsAsync();
-
-			List<ToolId> removeIds = ids.Except(tools.Select(x => x.Id)).ToList();
-			foreach (ToolId removeId in removeIds)
-			{
-				_logger.LogInformation("Removing tool {ToolId}", removeId);
-				await _tools.DeleteAsync(removeId);
-			}
-
-			foreach (ToolConfig tool in tools)
-			{
-				if (!ids.Contains(tool.Id))
-				{
-					_logger.LogInformation("Creating tool {ToolId}", tool.Id);
-				}
-				await AddOrUpdateAsync(tool);
-			}
-		}
-
-		/// <summary>
-		/// Finds or adds a tool with the given identifier
-		/// </summary>
-		/// <param name="options">Options for the tool</param>
-		/// <returns>Document describing the tool</returns>
-		private async Task<ITool> AddOrUpdateAsync(ToolConfig options)
-		{
-			for (; ; )
-			{
-				// Check for the common case where the tool already exists and is set to the right values already
-				Tool tool = await _tools.FindOrAddAsync(options.Id, () => new Tool(options));
-				if (tool.Name.Equals(options.Name, StringComparison.Ordinal) && tool.Description.Equals(options.Description, StringComparison.Ordinal) && tool.Public == options.Public && AclV2.Equals(tool.Acl, options.Acl))
-				{
-					return tool;
-				}
-
-				// Otherwise attempt to do an in-place upgrade
-				UpdateDefinition<Tool> update = Builders<Tool>.Update.Set(x => x.Name, options.Name).Set(x => x.Description, options.Description).Set(x => x.Public, options.Public);
-				if (AclV2.IsNullOrDefault(options.Acl))
-				{
-					update = update.Unset(x => x.Acl);
-				}
-				else
-				{
-					update = update.Set(x => x.Acl, options.Acl);
-				}
-
-				// Update the existing tool
-				Tool? updatedTool = await _tools.UpdateAsync(tool, update);
-				if (updatedTool != null)
-				{
-					await ClearCachedIndexAsync();
-					return updatedTool;
-				}
-			}
-		}
-
-		/// <summary>
-		/// Removes a tool with the given identifier
-		/// </summary>
-		/// <param name="id">The tool to remove</param>
-		/// <returns>True if a tool with the given identifier was deleted</returns>
-		public async Task<bool> DeleteAsync(ToolId id)
-		{
-			bool deleted = await _tools.DeleteAsync(id);
-			if (deleted)
-			{
-				await ClearCachedIndexAsync();
-			}
-			return deleted;
-		}
-
-		private async Task ClearCachedIndexAsync()
-		{
-			CachedIndex newIndexData = new CachedIndex { Rev = ObjectId.GenerateNewId().ToString() };
-			await _redisService.GetDatabase().StringSetAsync(s_indexKey, CbSerializer.SerializeToByteArray(newIndexData));
-		}
-
-		/// <summary>
-		/// Finds all tools matching a set of criteria
-		/// </summary>
-		/// <returns>Sequence of tool documents</returns>
-		async Task<List<ToolId>> FindAllIdsAsync()
-		{
-			IDatabase redis = _redisService.GetDatabase();
-			CachedIndex? index;
-			for (; ; )
-			{
-				// Get the cached index, and check if it has valid data
-				RedisValue value = await redis.StringGetAsync(s_indexKey);
-				if (!value.IsNullOrEmpty)
-				{
-					index = CbSerializer.Deserialize<CachedIndex>((byte[])value!);
-					if (index.Empty || index.Ids.Count > 0)
-					{
-						break;
-					}
-				}
-
-				// Create a new index object
-				index = new CachedIndex();
-				index.Rev = ObjectId.GenerateNewId().ToString();
-				index.Ids = await _tools.BaseCollection.Find(FilterDefinition<VersionedDocument<ToolId, Tool>>.Empty).Project(x => x.Id).ToListAsync();
-				index.Empty = index.Ids.Count == 0;
-
-				// Try to replace the existing one
-				RedisValue newValue = CbSerializer.SerializeToByteArray(index);
-				if (value.IsNull)
-				{
-					if (await redis.StringSetAsync(s_indexKey, newValue, when: When.NotExists))
-					{
-						break;
-					}
-				}
-				else
-				{
-					ITransaction transaction = redis.CreateTransaction();
-					transaction.AddCondition(Condition.StringEqual(s_indexKey, value));
-					_ = transaction.StringSetAsync(s_indexKey, newValue, flags: CommandFlags.FireAndForget);
-
-					if (await transaction.ExecuteAsync())
-					{
-						break;
-					}
-				}
-			}
-			return index.Ids;
-		}
-
-		/// <summary>
-		/// Finds all tools matching a set of criteria
-		/// </summary>
-		/// <returns>Sequence of tool documents</returns>
-		public async Task<List<ITool>> FindAllAsync()
-		{
-			// Fetch all the tools in parallel
-			IEnumerable<ToolId> toolIds =await FindAllIdsAsync();
-			List<Task<ITool?>> tasks = toolIds.Select(x => Task.Run(() => GetAsync(x))).ToList();
-
-			// Return all the successful results
-			List<ITool> results = new List<ITool>();
-			foreach (Task<ITool?> task in tasks)
-			{
-				ITool? tool = await task;
-				if (tool != null)
-				{
-					results.Add(tool);
-				}
-			}
-			return results;
-		}
-
 		/// <summary>
 		/// Gets a tool with the given identifier
 		/// </summary>
 		/// <param name="id">The tool identifier</param>
+		/// <param name="globalConfig">The current global configuration</param>
 		/// <returns></returns>
-		public async Task<ITool?> GetAsync(ToolId id) => await GetInternalAsync(id);
+		public async Task<ITool?> GetAsync(ToolId id, GlobalConfig globalConfig) => await GetInternalAsync(id, globalConfig);
 
 		/// <summary>
 		/// Gets a tool with the given identifier
 		/// </summary>
-		/// <param name="id">The tool identifier</param>
+		/// <param name="toolId">The tool identifier</param>
+		/// <param name="globalConfig">The current global configuration</param>
 		/// <returns></returns>
-		async Task<Tool?> GetInternalAsync(ToolId id)
+		async Task<Tool?> GetInternalAsync(ToolId toolId, GlobalConfig globalConfig)
 		{
-			Tool? tool = await _tools.GetAsync(id);
-			if (tool != null)
+			ToolConfig? toolConfig;
+			if (!globalConfig.TryGetTool(toolId, out toolConfig))
 			{
-				tool.UpdateTemporalState(_clock.UtcNow);
+				return null;
 			}
+
+			Tool tool = await _tools.FindOrAddAsync(toolId, () => new Tool(toolId));
+			tool.Config = toolConfig;
+			tool.UpdateTemporalState(_clock.UtcNow);
 			return tool;
 		}
 
@@ -414,8 +233,9 @@ namespace Horde.Build.Tools
 		/// <param name="tool">The tool to update</param>
 		/// <param name="options">Options for the new deployment</param>
 		/// <param name="stream">Stream containing the tool data</param>
+		/// <param name="globalConfig">The current configuration</param>
 		/// <returns>Updated tool document, or null if it does not exist</returns>
-		public async Task<ITool?> CreateDeploymentAsync(ITool tool, ToolDeploymentConfig options, Stream stream)
+		public async Task<ITool?> CreateDeploymentAsync(ITool tool, ToolDeploymentConfig options, Stream stream, GlobalConfig globalConfig)
 		{
 			// Upload the tool data first
 			IoHash hash = await _storage.WriteBlobAsync(_namespaceId, stream);
@@ -449,7 +269,7 @@ namespace Horde.Build.Tools
 					break;
 				}
 
-				newTool = await GetInternalAsync(tool.Id);
+				newTool = await GetInternalAsync(tool.Id, globalConfig);
 				if (newTool == null)
 				{
 					return null;
diff --git a/Engine/Source/Programs/Horde/Horde.Build/Tools/ToolConfig.cs b/Engine/Source/Programs/Horde/Horde.Build/Tools/ToolConfig.cs
--- a/Engine/Source/Programs/Horde/Horde.Build/Tools/ToolConfig.cs
+++ b/Engine/Source/Programs/Horde/Horde.Build/Tools/ToolConfig.cs
@@ -3,7 +3,10 @@
 using System;
 using System.ComponentModel.DataAnnotations;
 using System.Diagnostics;
+using System.Security.Claims;
+using System.Text.Json.Serialization;
 using Horde.Build.Acls;
+using Horde.Build.Server;
 using Horde.Build.Utilities;
 
 namespace Horde.Build.Tools
@@ -16,22 +19,38 @@ namespace Horde.Build.Tools
 	[DebuggerDisplay("{Id}")]
 	public class ToolConfig
 	{
-		/// <inheritdoc cref="VersionedDocument{ToolId, Tool}.Id"/>
+		/// <summary>
+		/// The global config containing this tool
+		/// </summary>
+		[JsonIgnore]
+		public GlobalConfig GlobalConfig { get; private set; } = null!;
+
+		/// <summary>
+		/// Unique identifier for the tool
+		/// </summary>
 		[Required]
 		public ToolId Id { get; set; }
 
-		/// <inheritdoc cref="ITool.Name"/>
+		/// <summary>
+		/// Name of the tool
+		/// </summary>
 		[Required]
 		public string Name { get; set; }
 
-		/// <inheritdoc cref="ITool.Description"/>
+		/// <summary>
+		/// Description for the tool
+		/// </summary>
 		public string Description { get; set; }
 
-		/// <inheritdoc cref="ITool.Public"/>
+		/// <summary>
+		/// Whether this tool should be exposed for download on a public endpoint without authentication
+		/// </summary>
 		public bool Public { get; set; }
 
-		/// <inheritdoc cref="ITool.Acl"/>
-		public AclV2? Acl { get; set; }
+		/// <summary>
+		/// Permissions for the tool
+		/// </summary>
+		public AclConfig? Acl { get; set; }
 
 		/// <summary>
 		/// Default constructor for serialization
@@ -51,6 +70,25 @@ namespace Horde.Build.Tools
 			Name = id.ToString();
 			Description = String.Empty;
 		}
+
+		/// <summary>
+		/// Called after the config has been read
+		/// </summary>
+		/// <param name="globalConfig">Parent GlobalConfig object</param>
+		public void PostLoad(GlobalConfig globalConfig)
+		{
+			GlobalConfig = globalConfig;
+		}
+
+		/// <summary>
+		/// Authorizes a user to perform a given action
+		/// </summary>
+		/// <param name="action">The action being performed</param>
+		/// <param name="user">The principal to validate</param>
+		public bool Authorize(AclAction action, ClaimsPrincipal user)
+		{
+			return Acl?.Authorize(action, user) ?? GlobalConfig.Authorize(action, user);
+		}
 	}
 
 	/// <summary>
