Horde: Refactor codebase to use new config system.

diff --git a/Engine/Source/Programs/Horde/Horde.Build/Tools/ToolsController.cs b/Engine/Source/Programs/Horde/Horde.Build/Tools/ToolsController.cs
--- a/Engine/Source/Programs/Horde/Horde.Build/Tools/ToolsController.cs
+++ b/Engine/Source/Programs/Horde/Horde.Build/Tools/ToolsController.cs
@@ -3,12 +3,14 @@
 using EpicGames.Core;
 using EpicGames.Horde.Storage;
 using Horde.Build.Acls;
+using Horde.Build.Server;
 using Horde.Build.Utilities;
 using HordeCommon;
 using Microsoft.AspNetCore.Authorization;
 using Microsoft.AspNetCore.Http;
 using Microsoft.AspNetCore.Mvc;
 using Microsoft.Extensions.Logging;
+using Microsoft.Extensions.Options;
 using System;
 using System.Collections.Generic;
 using System.IO;
@@ -31,17 +33,17 @@ namespace Horde.Build.Tools
 		/// <inheritdoc cref="VersionedDocument{TId, TLatest}.Id"/>
 		public ToolId Id => _tool.Id;
 
-		/// <inheritdoc cref="ITool.Name"/>
-		public string Name => _tool.Name;
+		/// <inheritdoc cref="ToolConfig.Name"/>
+		public string Name => _tool.Config.Name;
 
-		/// <inheritdoc cref="ITool.Description"/>
-		public string Description => _tool.Description;
+		/// <inheritdoc cref="ToolConfig.Description"/>
+		public string Description => _tool.Config.Description;
 
 		/// <inheritdoc cref="ITool.Deployments"/>
 		public List<GetToolDeploymentResponse> Deployments { get; }
 
-		/// <inheritdoc cref="ITool.Public"/>
-		public bool Public => _tool.Public;
+		/// <inheritdoc cref="ToolConfig.Public"/>
+		public bool Public => _tool.Config.Public;
 
 		/// <summary>
 		/// Constructor
@@ -136,16 +138,16 @@ namespace Horde.Build.Tools
 	[Route("[controller]")]
 	public class ToolsController : HordeControllerBase
 	{
-		readonly AclService _aclService;
 		readonly ToolCollection _toolCollection;
+		readonly IOptionsSnapshot<GlobalConfig> _globalConfig;
 
 		/// <summary>
 		/// Constructor
 		/// </summary>
-		public ToolsController(AclService aclService, ToolCollection toolCollection)
+		public ToolsController(ToolCollection toolCollection, IOptionsSnapshot<GlobalConfig> globalConfig)
 		{
-			_aclService = aclService;
 			_toolCollection = toolCollection;
+			_globalConfig = globalConfig;
 		}
 
 		/// <summary>
@@ -156,19 +158,19 @@ namespace Horde.Build.Tools
 		[Route("/api/v1/tools/{id}/deployments")]
 		public async Task<ActionResult<CreateDeploymentResponse>> CreateDeploymentAsync(ToolId id, [FromForm] ToolDeploymentConfig options, [FromForm] IFormFile file)
 		{
-			ITool? tool = await _toolCollection.GetAsync(id);
+			ITool? tool = await _toolCollection.GetAsync(id, _globalConfig.Value);
 			if (tool == null)
 			{
 				return NotFound(id);
 			}
-			if (!await tool.AuthorizeAsync(AclAction.UploadTool, User, _aclService, null))
+			if (!tool.Config.Authorize(AclAction.UploadTool, User))
 			{
 				return Forbid(AclAction.UploadTool, id);
 			}
 
 			using (Stream stream = file.OpenReadStream())
 			{
-				tool = await _toolCollection.CreateDeploymentAsync(tool, options, stream);
+				tool = await _toolCollection.CreateDeploymentAsync(tool, options, stream, _globalConfig.Value);
 				if (tool == null)
 				{
 					return NotFound(id);
@@ -186,12 +188,12 @@ namespace Horde.Build.Tools
 		[Route("/api/v1/tools/{id}/deployments/{deploymentId}")]
 		public async Task<ActionResult> GetDeploymentAsync(ToolId id, ToolDeploymentId deploymentId, [FromQuery] GetToolAction action = GetToolAction.Info)
 		{
-			ITool? tool = await _toolCollection.GetAsync(id);
+			ITool? tool = await _toolCollection.GetAsync(id, _globalConfig.Value);
 			if (tool == null)
 			{
 				return NotFound(id);
 			}
-			if (!await tool.AuthorizeAsync(AclAction.DownloadTool, User, _aclService, null))
+			if (!tool.Config.Authorize(AclAction.DownloadTool, User))
 			{
 				return Forbid(AclAction.DownloadTool, id);
 			}
@@ -226,12 +228,12 @@ namespace Horde.Build.Tools
 		[Route("/api/v1/tools/{id}/deployments/{deploymentId}")]
 		public async Task<ActionResult> UpdateDeploymentAsync(ToolId id, ToolDeploymentId deploymentId, [FromBody] UpdateDeploymentRequest request)
 		{
-			ITool? tool = await _toolCollection.GetAsync(id);
+			ITool? tool = await _toolCollection.GetAsync(id, _globalConfig.Value);
 			if (tool == null)
 			{
 				return NotFound(id);
 			}
-			if (!await tool.AuthorizeAsync(AclAction.UploadTool, User, _aclService, null))
+			if (!tool.Config.Authorize(AclAction.UploadTool, User))
 			{
 				return Forbid(AclAction.UploadTool, id);
 			}
@@ -254,18 +256,18 @@ namespace Horde.Build.Tools
 	[ApiController]
 	public class PublicToolsController : HordeControllerBase
 	{
-		readonly AclService _aclService;
 		readonly ToolCollection _toolCollection;
+		readonly IOptionsSnapshot<GlobalConfig> _globalConfig;
 		readonly IClock _clock;
 
 		/// <summary>
 		/// Constructor
 		/// </summary>
-		public PublicToolsController(AclService aclService, ToolCollection toolCollection, IClock clock)
+		public PublicToolsController(ToolCollection toolCollection, IClock clock, IOptionsSnapshot<GlobalConfig> globalConfig)
 		{
-			_aclService = aclService;
 			_toolCollection = toolCollection;
 			_clock = clock;
+			_globalConfig = globalConfig;
 		}
 
 		/// <summary>
@@ -276,12 +278,12 @@ namespace Horde.Build.Tools
 		[Route("/api/v1/tools/{id}")]
 		public async Task<ActionResult> GetToolAsync(ToolId id, GetToolAction action = GetToolAction.Info)
 		{
-			ITool? tool = await _toolCollection.GetAsync(id);
+			ITool? tool = await _toolCollection.GetAsync(id, _globalConfig.Value);
 			if (tool == null)
 			{
 				return NotFound(id);
 			}
-			if (!await AuthorizeDownloadAsync(tool))
+			if (!AuthorizeDownload(tool))
 			{
 				return Forbid(AclAction.DownloadTool, id);
 			}
@@ -312,12 +314,12 @@ namespace Horde.Build.Tools
 		[Route("/api/v1/tools/{id}/deployments")]
 		public async Task<ActionResult> FindDeploymentAsync(ToolId id, [FromQuery] double phase = 0.0, [FromQuery] GetToolAction action = GetToolAction.Info)
 		{
-			ITool? tool = await _toolCollection.GetAsync(id);
+			ITool? tool = await _toolCollection.GetAsync(id, _globalConfig.Value);
 			if (tool == null)
 			{
 				return NotFound(id);
 			}
-			if (!await AuthorizeDownloadAsync(tool))
+			if (!AuthorizeDownload(tool))
 			{
 				return Forbid(AclAction.DownloadTool, id);
 			}
@@ -353,15 +355,15 @@ namespace Horde.Build.Tools
 			return new FileStreamResult(stream, deployment.MimeType);
 		}
 
-		async ValueTask<bool> AuthorizeDownloadAsync(ITool tool)
+		bool AuthorizeDownload(ITool tool)
 		{
-			if (!tool.Public)
+			if (!tool.Config.Public)
 			{
 				if (User.Identity == null || !User.Identity.IsAuthenticated)
 				{
 					return false;
 				}
-				if (!await tool.AuthorizeAsync(AclAction.DownloadTool, User, _aclService, null))
+				if (!tool.Config.Authorize(AclAction.DownloadTool, User))
 				{
 					return false;
 				}
diff --git a/Engine/Source/Programs/Shared/EpicGames.Core/JsonSchema.cs b/Engine/Source/Programs/Shared/EpicGames.Core/JsonSchema.cs
--- a/Engine/Source/Programs/Shared/EpicGames.Core/JsonSchema.cs
+++ b/Engine/Source/Programs/Shared/EpicGames.Core/JsonSchema.cs
@@ -5,6 +5,7 @@ using System.Collections.Generic;
 using System.IO;
 using System.Reflection;
 using System.Text.Json;
+using System.Text.Json.Serialization;
 using System.Xml;
 
 namespace EpicGames.Core
@@ -435,9 +436,12 @@ namespace EpicGames.Core
 			PropertyInfo[] properties = type.GetProperties(BindingFlags.Instance | BindingFlags.Public);
 			foreach (PropertyInfo property in properties)
 			{
-				string? description = GetPropertyDescription(type, property.Name, xmlDoc);
-				JsonSchemaType propertyType = CreateSchemaType(property.PropertyType, typeCache, xmlDoc);
-				obj.Properties.Add(new JsonSchemaProperty(property.Name, description, propertyType));
+				if (property.GetCustomAttribute<JsonIgnoreAttribute>() == null)
+				{
+					string? description = GetPropertyDescription(type, property.Name, xmlDoc);
+					JsonSchemaType propertyType = CreateSchemaType(property.PropertyType, typeCache, xmlDoc);
+					obj.Properties.Add(new JsonSchemaProperty(property.Name, description, propertyType));
+				}
 			}
 		}
 	}
diff --git a/Engine/Source/Programs/Shared/EpicGames.Core/ListExtensions.cs b/Engine/Source/Programs/Shared/EpicGames.Core/ListExtensions.cs
--- a/Engine/Source/Programs/Shared/EpicGames.Core/ListExtensions.cs
+++ b/Engine/Source/Programs/Shared/EpicGames.Core/ListExtensions.cs
@@ -3,6 +3,7 @@
 using System;
 using System.Collections;
 using System.Collections.Generic;
+using System.Threading.Tasks;
 
 namespace EpicGames.Core
 {
@@ -63,5 +64,47 @@ namespace EpicGames.Core
 		{
 			list.Sort((x, y) => comparer.Compare(selector(x), selector(y)));
 		}
+
+		/// <summary>
+		/// Convert all elements of a list to a base type
+		/// </summary>
+		/// <typeparam name="TInput">Input element type</typeparam>
+		/// <typeparam name="TOutput">Output element type</typeparam>
+		/// <param name="input">List to convert</param>
+		/// <returns>Converted list</returns>
+		public static List<TOutput> ConvertAll<TInput, TOutput>(this List<TInput> input) where TInput : TOutput
+		{
+			List<TOutput> output = new List<TOutput>(input.Count);
+			foreach (TInput inputItem in input)
+			{
+				output.Add(inputItem);
+			}
+			return output;
+		}
+
+		/// <summary>
+		/// Convert all elements of a list to a base type
+		/// </summary>
+		/// <typeparam name="TInput">Input element type</typeparam>
+		/// <typeparam name="TOutput">Output element type</typeparam>
+		/// <param name="input">List to convert</param>
+		/// <returns>Converted list</returns>
+		public static async Task<List<TOutput>> ConvertAllAsync<TInput, TOutput>(this Task<List<TInput>> input) where TInput : TOutput
+		{
+			return ConvertAll<TInput, TOutput>(await input);
+		}
+
+		/// <summary>
+		/// Convert all elements of a list to a base type
+		/// </summary>
+		/// <typeparam name="TInput">Input element type</typeparam>
+		/// <typeparam name="TOutput">Output element type</typeparam>
+		/// <param name="input">List to convert</param>
+		/// <param name="converter">Converter for elements</param>
+		/// <returns>Converted list</returns>
+		public static async Task<List<TOutput>> ConvertAllAsync<TInput, TOutput>(this Task<List<TInput>> input, Converter<TInput, TOutput> converter)
+		{
+			return (await input).ConvertAll(converter);
+		}
 	}
 }
