Horde: Force config files to be re-parsed when switching server versions.

diff --git a/Engine/Source/Programs/Horde/Horde.Build/Configuration/ConfigService.cs b/Engine/Source/Programs/Horde/Horde.Build/Configuration/ConfigService.cs
--- a/Engine/Source/Programs/Horde/Horde.Build/Configuration/ConfigService.cs
+++ b/Engine/Source/Programs/Horde/Horde.Build/Configuration/ConfigService.cs
@@ -38,6 +38,9 @@ namespace Horde.Build.Configuration
 
 			[ProtoMember(4)]
 			public Dictionary<Uri, string> Dependencies { get; set; } = new Dictionary<Uri, string>();
+
+			[ProtoMember(5)]
+			public string ServerVersion { get; set; } = String.Empty;
 		}
 
 		// Implements a token for config change notifications
@@ -317,6 +320,7 @@ namespace Horde.Build.Configuration
 			try
 			{
 				ConfigSnapshot snapshot = new ConfigSnapshot();
+				snapshot.ServerVersion = Program.Version.ToString();
 
 				GlobalConfig globalConfig = await ConfigType.ReadAsync<GlobalConfig>(globalConfigUri, context, cancellationToken);
 				snapshot.Data = JsonSerializer.SerializeToUtf8Bytes(globalConfig, context.JsonOptions);
@@ -343,6 +347,12 @@ namespace Horde.Build.Configuration
 		/// <returns>True if the snapshot is out of date</returns>
 		async Task<bool> IsOutOfDateAsync(ConfigSnapshot snapshot, CancellationToken cancellationToken)
 		{
+			// Always re-read the config file when switching server versions
+			if (!snapshot.ServerVersion.Equals(Program.Version.ToString(), StringComparison.Ordinal))
+			{
+				return true;
+			}
+
 			// Group the dependencies by scheme in order to allow the source to batch-query them
 			foreach(IGrouping<string, KeyValuePair<Uri, string>> group in snapshot.Dependencies.GroupBy(x => x.Key.Scheme))
 			{
