Horde: Add an admin endpoint for returning the complete config object.

diff --git a/Engine/Source/Programs/Horde/Horde.Build/Server/AdminController.cs b/Engine/Source/Programs/Horde/Horde.Build/Server/AdminController.cs
--- a/Engine/Source/Programs/Horde/Horde.Build/Server/AdminController.cs
+++ b/Engine/Source/Programs/Horde/Horde.Build/Server/AdminController.cs
@@ -83,6 +83,21 @@ namespace Horde.Build.Server
 			return await _aclService.IssueBearerTokenAsync(User.Claims, GetDefaultExpiryTime());
 		}
 
+		/// <summary>
+		/// Returns the fully parsed config object.
+		/// </summary>
+		[HttpGet]
+		[Route("/api/v1/admin/config")]
+		public ActionResult<object> GetConfig()
+		{
+			if (!_globalConfig.Value.Authorize(AclAction.AdminRead, User))
+			{
+				return Forbid(AclAction.AdminRead);
+			}
+
+			return _globalConfig.Value;
+		}
+
 		/// <summary>
 		/// Issues a token for the given roles. Issues a token for the current user if not specified.
 		/// </summary>
