Horde: Fix exception casting objects while building project responses.

diff --git a/Engine/Source/Programs/Horde/Horde.Build/Projects/ProjectMessages.cs b/Engine/Source/Programs/Horde/Horde.Build/Projects/ProjectMessages.cs
--- a/Engine/Source/Programs/Horde/Horde.Build/Projects/ProjectMessages.cs
+++ b/Engine/Source/Programs/Horde/Horde.Build/Projects/ProjectMessages.cs
@@ -138,25 +138,25 @@ namespace Horde.Build.Projects
 		/// <param name="projectConfig">The project instance</param>
 		/// <param name="includeStreams">Whether to include streams in the response</param>
 		/// <param name="includeCategories">Whether to include categories in the response</param>
-		/// <param name="streams">The list of streams</param>
+		/// <param name="streamConfigs">The list of streams</param>
 		/// <returns>Response instance</returns>
-		public static GetProjectResponse FromConfig(ProjectConfig projectConfig, bool includeStreams, bool includeCategories, List<StreamConfig>? streams)
+		public static GetProjectResponse FromConfig(ProjectConfig projectConfig, bool includeStreams, bool includeCategories, List<StreamConfig>? streamConfigs)
 		{
 			List<GetProjectStreamResponse>? streamResponses = null;
 			if (includeStreams)
 			{
-				streamResponses = streams!.ConvertAll(x => new GetProjectStreamResponse(x.Id.ToString(), x.Name));
+				streamResponses = streamConfigs!.ConvertAll(x => new GetProjectStreamResponse(x.Id.ToString(), x.Name));
 			}
 
 			List<GetProjectCategoryResponse>? categoryResponses = null;
 			if (includeCategories)
 			{
 				categoryResponses = projectConfig.Categories.ConvertAll(x => new GetProjectCategoryResponse(x));
-				if (streams != null)
+				if (streamConfigs != null)
 				{
-					foreach (IStream stream in streams)
+					foreach (StreamConfig streamConfig in streamConfigs)
 					{
-						GetProjectCategoryResponse? categoryResponse = categoryResponses.FirstOrDefault(x => MatchCategory(stream.Config.Name, x));
+						GetProjectCategoryResponse? categoryResponse = categoryResponses.FirstOrDefault(x => MatchCategory(streamConfig.Name, x));
 						if (categoryResponse == null)
 						{
 							int row = (categoryResponses.Count > 0) ? categoryResponses.Max(x => x.Row) : 0;
@@ -173,7 +173,7 @@ namespace Horde.Build.Projects
 							categoryResponse = new GetProjectCategoryResponse(otherCategory);
 							categoryResponses.Add(categoryResponse);
 						}
-						categoryResponse.Streams!.Add(stream.Id.ToString());
+						categoryResponse.Streams!.Add(streamConfig.Id.ToString());
 					}
 				}
 			}
