Horde: Fix streams controller returning "V2" responses from "V1" endpoint.

diff --git a/Engine/Source/Programs/Horde/Horde.Build/Streams/StreamsController.cs b/Engine/Source/Programs/Horde/Horde.Build/Streams/StreamsController.cs
--- a/Engine/Source/Programs/Horde/Horde.Build/Streams/StreamsController.cs
+++ b/Engine/Source/Programs/Horde/Horde.Build/Streams/StreamsController.cs
@@ -64,7 +64,26 @@ namespace Horde.Build.Streams
 		[ProducesResponseType(typeof(List<GetStreamResponse>), 200)]
 		public async Task<ActionResult<List<object>>> GetStreamsAsync([FromQuery(Name = "ProjectId")] string[] projectIds, [FromQuery] PropertyFilter? filter = null)
 		{
-			return await GetStreamsInternalAsync(projectIds, true, filter);
+			ProjectId[] projectIdValues = Array.ConvertAll(projectIds, x => new ProjectId(x));
+
+			List<GetStreamResponse> responses = new List<GetStreamResponse>();
+			foreach (ProjectConfig projectConfig in _globalConfig.Value.Projects)
+			{
+				if (projectIdValues.Length == 0 || projectIdValues.Contains(projectConfig.Id))
+				{
+					foreach (StreamConfig streamConfig in projectConfig.Streams)
+					{
+						if (streamConfig.Authorize(AclAction.ViewStream, User))
+						{
+							IStream stream = await _streamCollection.GetAsync(streamConfig);
+							GetStreamResponse response = await CreateGetStreamResponseAsync(stream);
+							responses.Add(response);
+						}
+					}
+				}
+			}
+
+			return responses.OrderBy(x => x.Id).Select(x => PropertyFilter.Apply(x, filter)).ToList();
 		}
 
 		/// <summary>
@@ -77,11 +96,6 @@ namespace Horde.Build.Streams
 		[Route("/api/v2/streams")]
 		[ProducesResponseType(typeof(List<GetStreamResponse>), 200)]
 		public async Task<ActionResult<List<object>>> GetStreamsAsyncV2([FromQuery(Name = "ProjectId")] string[] projectIds, [FromQuery] PropertyFilter? filter = null)
-		{
-			return await GetStreamsInternalAsync(projectIds, false, filter);
-		}
-
-		async Task<List<object>> GetStreamsInternalAsync([FromQuery(Name = "ProjectId")] string[] projectIds, bool includeConfig, PropertyFilter? filter = null)
 		{
 			ProjectId[] projectIdValues = Array.ConvertAll(projectIds, x => new ProjectId(x));
 
@@ -95,24 +109,14 @@ namespace Horde.Build.Streams
 						if (streamConfig.Authorize(AclAction.ViewStream, User))
 						{
 							IStream stream = await _streamCollection.GetAsync(streamConfig);
-							GetStreamResponseV2 response = await CreateStreamResponseAsync(stream, includeConfig);
+							GetStreamResponseV2 response = await CreateGetStreamResponseV2Async(stream, false);
 							responses.Add(response);
 						}
 					}
 				}
 			}
-			return responses.OrderBy(x => x.Id).Select(x => PropertyFilter.Apply(x, filter)).ToList();
-		}
 
-		async Task<GetStreamResponseV2> CreateStreamResponseAsync(IStream stream, bool includeConfig)
-		{
-			List<GetTemplateRefResponseV2> templates = new List<GetTemplateRefResponseV2>();
-			foreach (ITemplateRef templateRef in stream.Templates.Values)
-			{
-				ITemplate template = await _templateCollection.GetOrAddAsync(templateRef.Config);
-				templates.Add(new GetTemplateRefResponseV2(template, templateRef));
-			}
-			return new GetStreamResponseV2(stream, includeConfig, templates);
+			return responses.OrderBy(x => x.Id).Select(x => PropertyFilter.Apply(x, filter)).ToList();
 		}
 
 		/// <summary>
@@ -137,7 +141,7 @@ namespace Horde.Build.Streams
 			}
 
 			IStream stream = await _streamCollection.GetAsync(streamConfig);
-			return PropertyFilter.Apply(await CreateGetStreamResponse(stream), filter);
+			return PropertyFilter.Apply(await CreateGetStreamResponseAsync(stream), filter);
 		}
 
 		/// <summary>
@@ -164,7 +168,7 @@ namespace Horde.Build.Streams
 
 			bool includeConfig = !String.Equals(config, streamConfig.Revision, StringComparison.Ordinal);
 			IStream stream = await _streamCollection.GetAsync(streamConfig);
-			return PropertyFilter.Apply(await CreateStreamResponseAsync(stream, includeConfig), filter);
+			return PropertyFilter.Apply(await CreateGetStreamResponseV2Async(stream, includeConfig), filter);
 		}
 
 		/// <summary>
@@ -172,7 +176,7 @@ namespace Horde.Build.Streams
 		/// </summary>
 		/// <param name="stream">Stream to create response for</param>
 		/// <returns>Response object</returns>
-		async Task<GetStreamResponse> CreateGetStreamResponse(IStream stream)
+		async Task<GetStreamResponse> CreateGetStreamResponseAsync(IStream stream)
 		{
 			using IScope scope = GlobalTracer.Instance.BuildSpan("CreateGetStreamResponse").StartActive();
 			scope.Span.SetTag("streamId", stream.Id);
@@ -208,6 +212,17 @@ namespace Horde.Build.Streams
 			return stream.ToApiResponse(apiTemplateRefs);
 		}
 
+		async Task<GetStreamResponseV2> CreateGetStreamResponseV2Async(IStream stream, bool includeConfig)
+		{
+			List<GetTemplateRefResponseV2> templates = new List<GetTemplateRefResponseV2>();
+			foreach (ITemplateRef templateRef in stream.Templates.Values)
+			{
+				ITemplate template = await _templateCollection.GetOrAddAsync(templateRef.Config);
+				templates.Add(new GetTemplateRefResponseV2(template, templateRef));
+			}
+			return new GetStreamResponseV2(stream, includeConfig, templates);
+		}
+
 		/// <summary>
 		/// Gets a list of changes for a stream
 		/// </summary>
