Horde: Add support for parsing dictionary objects in config files.

diff --git a/Engine/Source/Programs/Horde/Horde.Build/Configuration/ConfigType.cs b/Engine/Source/Programs/Horde/Horde.Build/Configuration/ConfigType.cs
--- a/Engine/Source/Programs/Horde/Horde.Build/Configuration/ConfigType.cs
+++ b/Engine/Source/Programs/Horde/Horde.Build/Configuration/ConfigType.cs
@@ -278,6 +278,53 @@ namespace Horde.Build.Configuration
 			}
 		}
 
+		class DictionaryProperty : Property
+		{
+			readonly ConfigType _elementType;
+
+			public DictionaryProperty(string name, PropertyInfo propertyInfo, ConfigType elementType)
+				: base(name, propertyInfo)
+			{
+				_elementType = elementType;
+			}
+
+			public override bool HasIncludes() => _elementType is ClassConfigType elementType && elementType.HasIncludes();
+
+			public override async Task MergeAsync(object target, JsonNode? node, ConfigContext context, CancellationToken cancellationToken)
+			{
+				IDictionary? dictionary = (IDictionary?)PropertyInfo.GetValue(target);
+				if (dictionary == null)
+				{
+					object value = Activator.CreateInstance(PropertyInfo.PropertyType)!;
+					PropertyInfo.SetValue(target, value);
+					dictionary = (IDictionary)value;
+				}
+
+				JsonObject obj = (JsonObject)node!;
+				foreach ((string key, JsonNode? element) in obj)
+				{
+					context.EnterScope($"{Name}[{key}]");
+
+					object? elementValue = await _elementType.ReadAsync(element, context, cancellationToken);
+					dictionary.Add(key, elementValue);
+
+					context.LeaveScope();
+				}
+			}
+
+			public override async Task ParseIncludesAsync(JsonNode jsonNode, object targetObject, ClassConfigType targetType, ConfigContext context, CancellationToken cancellationToken)
+			{
+				if (jsonNode is JsonObject jsonObjectValue)
+				{
+					ClassConfigType classElementType = (ClassConfigType)_elementType;
+					foreach (JsonObject jsonObjectElement in jsonObjectValue.Select(x => x.Value).OfType<JsonObject>())
+					{
+						await classElementType.ParseIncludesAsync(jsonObjectElement, targetObject, targetType, context, cancellationToken);
+					}
+				}
+			}
+		}
+
 		class ObjectProperty : Property
 		{
 			readonly ClassConfigType _classConfigType;
@@ -414,6 +461,11 @@ namespace Horde.Build.Configuration
 				Type elementType = propertyType.GetGenericArguments()[0];
 				return new ListProperty(name, propertyInfo, FindOrAddValueType(elementType));
 			}
+			else if (propertyType.IsGenericType && propertyType.GetGenericTypeDefinition() == typeof(Dictionary<,>))
+			{
+				Type elementType = propertyType.GetGenericArguments()[1];
+				return new DictionaryProperty(name, propertyInfo, FindOrAddValueType(elementType));
+			}
 			else
 			{
 				return new ObjectProperty(name, propertyInfo, FindOrAdd(propertyType));
