Horde: Additional logging for config updates.

diff --git a/Engine/Source/Programs/Horde/Horde.Build/Configuration/ConfigService.cs b/Engine/Source/Programs/Horde/Horde.Build/Configuration/ConfigService.cs
--- a/Engine/Source/Programs/Horde/Horde.Build/Configuration/ConfigService.cs
+++ b/Engine/Source/Programs/Horde/Horde.Build/Configuration/ConfigService.cs
@@ -237,8 +237,15 @@ namespace Horde.Build.Configuration
 			// avoids any latency with receiving notifications on the pub/sub channel about changes.
 			ReadOnlyMemory<byte> initialData = await ReadSnapshotDataAsync();
 
+			// Print info about the initial snapshot
+			ConfigSnapshot? snapshot = null;
+			if (initialData.Length > 0)
+			{
+				snapshot = DeserializeSnapshot(initialData);
+				_logger.LogInformation("Initial snapshot for update: {@Info}", GetSnapshotInfo(initialData.Span, snapshot));
+			}
+
 			// Update the snapshot until we're asked to stop
-			ConfigSnapshot? snapshot = (initialData.Length > 0) ? DeserializeSnapshot(initialData) : null;
 			while (!cancellationToken.IsCancellationRequested)
 			{
 				if (snapshot == null || await IsOutOfDateAsync(snapshot, cancellationToken))
@@ -339,6 +346,16 @@ namespace Horde.Build.Configuration
 			}
 		}
 
+		/// <summary>
+		/// Writes information about the current snapshot to the logger
+		/// </summary>
+		/// <param name="span">Data for the snapshot</param>
+		/// <param name="snapshot">Snapshot to log information on</param>
+		static object GetSnapshotInfo(ReadOnlySpan<byte> span, ConfigSnapshot snapshot)
+		{
+			return new { Revision = IoHash.Compute(span).ToString(), ServerVersion = snapshot.ServerVersion.ToString(), Dependencies = snapshot.Dependencies.ToArray() };
+		}
+
 		/// <summary>
 		/// Checks if the given snapshot is out of date
 		/// </summary>
@@ -348,8 +365,10 @@ namespace Horde.Build.Configuration
 		async Task<bool> IsOutOfDateAsync(ConfigSnapshot snapshot, CancellationToken cancellationToken)
 		{
 			// Always re-read the config file when switching server versions
-			if (!snapshot.ServerVersion.Equals(Program.Version.ToString(), StringComparison.Ordinal))
+			string newServerVersion = Program.Version.ToString();
+			if (!snapshot.ServerVersion.Equals(newServerVersion, StringComparison.Ordinal))
 			{
+				_logger.LogInformation("Config is out of date (server version {OldVersion} -> {NewVersion})", snapshot.ServerVersion, newServerVersion);
 				return true;
 			}
 
@@ -361,6 +380,7 @@ namespace Horde.Build.Configuration
 				IConfigSource? source;
 				if (!_sources.TryGetValue(group.Key, out source))
 				{
+					_logger.LogInformation("Config is out of date (missing source '{Source}')", group.Key);
 					return true;
 				}
 
@@ -369,6 +389,7 @@ namespace Horde.Build.Configuration
 				{
 					if (!files[idx].Revision.Equals(pairs[idx].Value, StringComparison.Ordinal))
 					{
+						_logger.LogInformation("Config is out of date (file {Path}: {OldVersion} -> {NewVersion})", files[idx].Uri, files[idx].Revision, pairs[idx].Value);
 						return true;
 					}
 				}
@@ -376,11 +397,12 @@ namespace Horde.Build.Configuration
 			return false;
 		}
 
-		async Task WriteSnapshotAsync(ConfigSnapshot index)
+		async Task WriteSnapshotAsync(ConfigSnapshot snapshot)
 		{
-			ReadOnlyMemory<byte> data = SerializeSnapshot(index);
+			ReadOnlyMemory<byte> data = SerializeSnapshot(snapshot);
 			await _redisService.GetDatabase().StringSetAsync(_snapshotKey, data);
 			await _redisService.PublishAsync(_updateChannel, RedisValue.EmptyString);
+			_logger.LogInformation("Published new config snapshot: {@Info}", GetSnapshotInfo(data.Span, snapshot)); 
 		}
 
 		async Task<ReadOnlyMemory<byte>> ReadSnapshotDataAsync()
