Horde: Fix formatting of Epic telemetry events.

diff --git a/Engine/Source/Programs/Horde/Horde.Build/Telemetry/EpicTelemetrySink.cs b/Engine/Source/Programs/Horde/Horde.Build/Telemetry/EpicTelemetrySink.cs
--- a/Engine/Source/Programs/Horde/Horde.Build/Telemetry/EpicTelemetrySink.cs
+++ b/Engine/Source/Programs/Horde/Horde.Build/Telemetry/EpicTelemetrySink.cs
@@ -73,6 +73,7 @@ namespace Horde.Build.Telemetry
 
 			_jsonOptions = new JsonSerializerOptions();
 			Startup.ConfigureJsonSerializer(_jsonOptions);
+			_jsonOptions.PropertyNamingPolicy = null;
 
 			_packetMemoryWriter = new ArrayMemoryWriter(65536);
 			_packetWriter = new Utf8JsonWriter(_packetMemoryWriter);
@@ -205,7 +206,7 @@ namespace Horde.Build.Telemetry
 				_packetWriter.WriteStartObject();
 				_packetWriter.WriteString(s_eventNamePropertyName, eventName);
 				_packetWriter.Flush();
-				_packetMemoryWriter.WriteFixedLengthBytes(eventData.Slice(0, eventData.Length - 1));
+				_packetMemoryWriter.WriteFixedLengthBytes(eventData);
 				_packetWriter.WriteEndObject();
 			}
 		}
