Horde: Add support for embedding external files into config files (used for project logos).

* Removed IProject/IProjectCollection, since there is no longer any tracked state for projects.

diff --git a/Engine/Source/Programs/Horde/Horde.Build.Tests/DeviceServiceTests.cs b/Engine/Source/Programs/Horde/Horde.Build.Tests/DeviceServiceTests.cs
--- a/Engine/Source/Programs/Horde/Horde.Build.Tests/DeviceServiceTests.cs
+++ b/Engine/Source/Programs/Horde/Horde.Build.Tests/DeviceServiceTests.cs
@@ -6,19 +6,15 @@ using Horde.Build.Devices;
 using Horde.Build.Jobs;
 using Microsoft.AspNetCore.Mvc;
 using Microsoft.VisualStudio.TestTools.UnitTesting;
-using Horde.Build.Projects;
 using Horde.Build.Utilities;
 using Horde.Build.Users;
 using System.Security.Claims;
 using Microsoft.AspNetCore.Http;
 using Horde.Build.Server;
-using System.Linq;
 using System;
 
 namespace Horde.Build.Tests
 {
-
-	using ProjectId = StringId<IProject>;
 	using DeviceId = StringId<IDevice>;
 	using DevicePlatformId = StringId<IDevicePlatform>;
 
diff --git a/Engine/Source/Programs/Horde/Horde.Build.Tests/Fixture.cs b/Engine/Source/Programs/Horde/Horde.Build.Tests/Fixture.cs
--- a/Engine/Source/Programs/Horde/Horde.Build.Tests/Fixture.cs
+++ b/Engine/Source/Programs/Horde/Horde.Build.Tests/Fixture.cs
@@ -22,7 +22,7 @@ namespace Horde.Build.Tests
 {
 	using JobId = ObjectId<IJob>;
 	using PoolId = StringId<IPool>;
-    using ProjectId = StringId<IProject>;
+    using ProjectId = StringId<ProjectConfig>;
     using StreamId = StringId<IStream>;
     using TemplateId = StringId<ITemplateRef>;
 
diff --git a/Engine/Source/Programs/Horde/Horde.Build.Tests/Fleet/JobQueueStrategyTest.cs b/Engine/Source/Programs/Horde/Horde.Build.Tests/Fleet/JobQueueStrategyTest.cs
--- a/Engine/Source/Programs/Horde/Horde.Build.Tests/Fleet/JobQueueStrategyTest.cs
+++ b/Engine/Source/Programs/Horde/Horde.Build.Tests/Fleet/JobQueueStrategyTest.cs
@@ -18,7 +18,7 @@ using Horde.Build.Agents;
 namespace Horde.Build.Tests.Fleet
 {
 	using PoolId = StringId<IPool>;
-	using ProjectId = StringId<IProject>;
+	using ProjectId = StringId<ProjectConfig>;
 	using StreamId = StringId<IStream>;
 
 	[TestClass]
diff --git a/Engine/Source/Programs/Horde/Horde.Build.Tests/IssueServiceTests.cs b/Engine/Source/Programs/Horde/Horde.Build.Tests/IssueServiceTests.cs
--- a/Engine/Source/Programs/Horde/Horde.Build.Tests/IssueServiceTests.cs
+++ b/Engine/Source/Programs/Horde/Horde.Build.Tests/IssueServiceTests.cs
@@ -35,7 +35,7 @@ namespace Horde.Build.Tests
 {
 	using JobId = ObjectId<IJob>;
 	using LogId = ObjectId<ILogFile>;
-	using ProjectId = StringId<IProject>;
+	using ProjectId = StringId<ProjectConfig>;
 	using StreamId = StringId<IStream>;
 	using TemplateId = StringId<ITemplateRef>;
 	using UserId = ObjectId<IUser>;
diff --git a/Engine/Source/Programs/Horde/Horde.Build.Tests/JobServiceTests.cs b/Engine/Source/Programs/Horde/Horde.Build.Tests/JobServiceTests.cs
--- a/Engine/Source/Programs/Horde/Horde.Build.Tests/JobServiceTests.cs
+++ b/Engine/Source/Programs/Horde/Horde.Build.Tests/JobServiceTests.cs
@@ -28,7 +28,7 @@ namespace Horde.Build.Tests
 {
 	using JobId = ObjectId<IJob>;
 	using LogId = ObjectId<ILogFile>;
-	using ProjectId = StringId<IProject>;
+	using ProjectId = StringId<ProjectConfig>;
 	using StreamId = StringId<IStream>;
 	using TemplateId = StringId<ITemplateRef>;
 
diff --git a/Engine/Source/Programs/Horde/Horde.Build.Tests/SchedulerTests.cs b/Engine/Source/Programs/Horde/Horde.Build.Tests/SchedulerTests.cs
--- a/Engine/Source/Programs/Horde/Horde.Build.Tests/SchedulerTests.cs
+++ b/Engine/Source/Programs/Horde/Horde.Build.Tests/SchedulerTests.cs
@@ -23,7 +23,7 @@ namespace Horde.Build.Tests
 {
 	using JobId = ObjectId<IJob>;
 	using LogId = ObjectId<ILogFile>;
-	using ProjectId = StringId<IProject>;
+	using ProjectId = StringId<ProjectConfig>;
 	using StreamId = StringId<IStream>;
 	using TemplateId = StringId<ITemplateRef>;
 
diff --git a/Engine/Source/Programs/Horde/Horde.Build.Tests/TestDataTests.cs b/Engine/Source/Programs/Horde/Horde.Build.Tests/TestDataTests.cs
--- a/Engine/Source/Programs/Horde/Horde.Build.Tests/TestDataTests.cs
+++ b/Engine/Source/Programs/Horde/Horde.Build.Tests/TestDataTests.cs
@@ -23,7 +23,7 @@ using Horde.Build.Server;
 
 namespace Horde.Build.Tests
 {
-	using ProjectId = StringId<IProject>;
+	using ProjectId = StringId<ProjectConfig>;
 	using StreamId = StringId<IStream>;
 	using TemplateId = StringId<ITemplateRef>;
 	using JobId = ObjectId<IJob>;
diff --git a/Engine/Source/Programs/Horde/Horde.Build.Tests/TestSetup.cs b/Engine/Source/Programs/Horde/Horde.Build.Tests/TestSetup.cs
--- a/Engine/Source/Programs/Horde/Horde.Build.Tests/TestSetup.cs
+++ b/Engine/Source/Programs/Horde/Horde.Build.Tests/TestSetup.cs
@@ -79,7 +79,6 @@ namespace Horde.Build.Tests
 		public IGraphCollection GraphCollection => ServiceProvider.GetRequiredService<IGraphCollection>();
 		public INotificationTriggerCollection NotificationTriggerCollection => ServiceProvider.GetRequiredService<INotificationTriggerCollection>();
 		public IStreamCollection StreamCollection => ServiceProvider.GetRequiredService<IStreamCollection>();
-		public IProjectCollection ProjectCollection => ServiceProvider.GetRequiredService<IProjectCollection>();
 		public IJobCollection JobCollection => ServiceProvider.GetRequiredService<IJobCollection>();
 		public IAgentCollection AgentCollection => ServiceProvider.GetRequiredService<IAgentCollection>();
 		public IJobStepRefCollection JobStepRefCollection => ServiceProvider.GetRequiredService<IJobStepRefCollection>();
@@ -211,7 +210,6 @@ namespace Horde.Build.Tests
 			services.AddSingleton<ILogFileCollection, LogFileCollection>();
 			services.AddSingleton<INotificationTriggerCollection, NotificationTriggerCollection>();
 			services.AddSingleton<IPoolCollection, PoolCollection>();
-			services.AddSingleton<IProjectCollection, ProjectCollection>();
 			services.AddSingleton<ISessionCollection, SessionCollection>();
 			services.AddSingleton<ISubscriptionCollection, SubscriptionCollection>();
 			services.AddSingleton<IStreamCollection, StreamCollection>();
diff --git a/Engine/Source/Programs/Horde/Horde.Build/Configuration/ConfigResource.cs b/Engine/Source/Programs/Horde/Horde.Build/Configuration/ConfigResource.cs
new file mode 100644
--- /dev/null
+++ b/Engine/Source/Programs/Horde/Horde.Build/Configuration/ConfigResource.cs
@@ -0,0 +1,84 @@
+// Copyright Epic Games, Inc. All Rights Reserved.
+
+using System;
+using System.Text.Json;
+using System.Text.Json.Serialization;
+using EpicGames.Core;
+
+namespace Horde.Build.Configuration
+{
+	/// <summary>
+	/// Special resource type in configuration files which stores data from an external source
+	/// </summary>
+	[JsonSchemaString]
+	[JsonConverter(typeof(ConfigResourceSerializer))]
+	public sealed class ConfigResource
+	{
+		/// <summary>
+		/// Path to the resource
+		/// </summary>
+		[ConfigRelativePath]
+		public string? Path { get; set; }
+
+		/// <summary>
+		/// Data for the resource. This will be stored externally to the config data.
+		/// </summary>
+		public ReadOnlyMemory<byte> Data { get; set; } = ReadOnlyMemory<byte>.Empty;
+	}
+
+	/// <summary>
+	/// Serializer for <see cref="ConfigResource"/> objects
+	/// </summary>
+	class ConfigResourceSerializer : JsonConverter<ConfigResource>
+	{
+		/// <inheritdoc/>
+		public override ConfigResource Read(ref Utf8JsonReader reader, Type typeToConvert, JsonSerializerOptions options)
+		{
+			ConfigResource resource = new ConfigResource();
+			if (reader.TokenType == JsonTokenType.String)
+			{
+				resource.Path = reader.GetString();
+			}
+			else if (reader.TokenType == JsonTokenType.StartObject)
+			{
+				while (reader.Read() && reader.TokenType != JsonTokenType.EndObject)
+				{
+					if (reader.TokenType == JsonTokenType.PropertyName)
+					{
+						if (reader.ValueTextEquals(nameof(ConfigResource.Path)))
+						{
+							reader.Read();
+							resource.Path = reader.GetString();
+						}
+						else if (reader.ValueTextEquals(nameof(ConfigResource.Data)))
+						{
+							reader.Read();
+							resource.Data = reader.GetBytesFromBase64();
+						}
+					}
+				}
+			}
+			else
+			{
+				throw new InvalidOperationException();
+			}
+			return resource;
+		}
+
+		/// <inheritdoc/>
+		public override void Write(Utf8JsonWriter writer, ConfigResource resource, JsonSerializerOptions options)
+		{
+			if (resource.Data.Length == 0)
+			{
+				writer.WriteStringValue(resource.Path);
+			}
+			else
+			{
+				writer.WriteStartObject();
+				writer.WriteString(nameof(ConfigResource.Path), resource.Path);
+				writer.WriteBase64String(nameof(ConfigResource.Data), resource.Data.Span);
+				writer.WriteEndObject();
+			}
+		}
+	}
+}
diff --git a/Engine/Source/Programs/Horde/Horde.Build/Configuration/ConfigService.cs b/Engine/Source/Programs/Horde/Horde.Build/Configuration/ConfigService.cs
--- a/Engine/Source/Programs/Horde/Horde.Build/Configuration/ConfigService.cs
+++ b/Engine/Source/Programs/Horde/Horde.Build/Configuration/ConfigService.cs
@@ -329,7 +329,10 @@ namespace Horde.Build.Configuration
 				ConfigSnapshot snapshot = new ConfigSnapshot();
 				snapshot.ServerVersion = Program.Version.ToString();
 
+				// Read the new config in
 				GlobalConfig globalConfig = await ConfigType.ReadAsync<GlobalConfig>(globalConfigUri, context, cancellationToken);
+
+				// Serialize it back out to a byte array
 				snapshot.Data = JsonSerializer.SerializeToUtf8Bytes(globalConfig, context.JsonOptions);
 
 				// Save all the dependencies
diff --git a/Engine/Source/Programs/Horde/Horde.Build/Configuration/ConfigType.cs b/Engine/Source/Programs/Horde/Horde.Build/Configuration/ConfigType.cs
--- a/Engine/Source/Programs/Horde/Horde.Build/Configuration/ConfigType.cs
+++ b/Engine/Source/Programs/Horde/Horde.Build/Configuration/ConfigType.cs
@@ -4,7 +4,6 @@ using System;
 using System.Collections;
 using System.Collections.Concurrent;
 using System.Collections.Generic;
-using System.ComponentModel.DataAnnotations;
 using System.Linq;
 using System.Reflection;
 using System.Text.Json;
@@ -231,6 +230,29 @@ namespace Horde.Build.Configuration
 			}
 		}
 
+		class ResourceProperty : ScalarProperty
+		{
+			public ResourceProperty(string name, PropertyInfo propertyInfo)
+				: base(name, propertyInfo)
+			{
+			}
+
+			public override async Task MergeAsync(object target, JsonNode? node, ConfigContext context, CancellationToken cancellationToken)
+			{
+				context.AddProperty(Name);
+
+				Uri uri = CombinePaths(context.CurrentFile, JsonSerializer.Deserialize<string>(node, context.JsonOptions) ?? String.Empty);
+				IConfigFile file = await ReadFileAsync(uri, context, cancellationToken);
+
+				ConfigResource resource = new ConfigResource();
+				resource.Path = uri.AbsoluteUri;
+				resource.Data = await file.ReadAsync(cancellationToken);
+				PropertyInfo.SetValue(target, resource);
+
+				context.Resources.Add(resource);
+			}
+		}
+
 		class ListProperty : Property
 		{
 			readonly ConfigType _elementType;
@@ -456,6 +478,10 @@ namespace Horde.Build.Configuration
 			{
 				return new ScalarProperty(name, propertyInfo);
 			}
+			else if (propertyType.IsAssignableTo(typeof(ConfigResource)))
+			{
+				return new ResourceProperty(name, propertyInfo);
+			}
 			else if (propertyType.IsGenericType && propertyType.GetGenericTypeDefinition() == typeof(List<>))
 			{
 				Type elementType = propertyType.GetGenericArguments()[0];
diff --git a/Engine/Source/Programs/Horde/Horde.Build/Devices/DeviceCollection.cs b/Engine/Source/Programs/Horde/Horde.Build/Devices/DeviceCollection.cs
--- a/Engine/Source/Programs/Horde/Horde.Build/Devices/DeviceCollection.cs
+++ b/Engine/Source/Programs/Horde/Horde.Build/Devices/DeviceCollection.cs
@@ -20,7 +20,7 @@ namespace Horde.Build.Devices
 	using DeviceId = StringId<IDevice>;
 	using DevicePlatformId = StringId<IDevicePlatform>;
 	using DevicePoolId = StringId<IDevicePool>;
-	using ProjectId = StringId<IProject>;
+	using ProjectId = StringId<ProjectConfig>;
 	using UserId = ObjectId<IUser>;
 	using StreamId = StringId<IStream>;
 
diff --git a/Engine/Source/Programs/Horde/Horde.Build/Devices/DeviceService.cs b/Engine/Source/Programs/Horde/Horde.Build/Devices/DeviceService.cs
--- a/Engine/Source/Programs/Horde/Horde.Build/Devices/DeviceService.cs
+++ b/Engine/Source/Programs/Horde/Horde.Build/Devices/DeviceService.cs
@@ -29,7 +29,7 @@ namespace Horde.Build.Devices
 	using DevicePlatformId = StringId<IDevicePlatform>;
 	using DevicePoolId = StringId<IDevicePool>;
 	using JobId = ObjectId<IJob>;
-	using ProjectId = StringId<IProject>;
+	using ProjectId = StringId<ProjectConfig>;
 	using UserId = ObjectId<IUser>;
 
 	/// <summary>
diff --git a/Engine/Source/Programs/Horde/Horde.Build/Devices/DevicesController.cs b/Engine/Source/Programs/Horde/Horde.Build/Devices/DevicesController.cs
--- a/Engine/Source/Programs/Horde/Horde.Build/Devices/DevicesController.cs
+++ b/Engine/Source/Programs/Horde/Horde.Build/Devices/DevicesController.cs
@@ -22,7 +22,7 @@ namespace Horde.Build.Devices
 	using DeviceId = StringId<IDevice>;
 	using DevicePlatformId = StringId<IDevicePlatform>;
 	using DevicePoolId = StringId<IDevicePool>;
-	using ProjectId = StringId<IProject>;
+	using ProjectId = StringId<ProjectConfig>;
 
 	/// <summary>
 	/// Controller for device service
diff --git a/Engine/Source/Programs/Horde/Horde.Build/Devices/IDevice.cs b/Engine/Source/Programs/Horde/Horde.Build/Devices/IDevice.cs
--- a/Engine/Source/Programs/Horde/Horde.Build/Devices/IDevice.cs
+++ b/Engine/Source/Programs/Horde/Horde.Build/Devices/IDevice.cs
@@ -14,7 +14,7 @@ namespace Horde.Build.Devices
 	using DeviceId = StringId<IDevice>;
 	using DevicePlatformId = StringId<IDevicePlatform>;
 	using DevicePoolId = StringId<IDevicePool>;
-	using ProjectId = StringId<IProject>;
+	using ProjectId = StringId<ProjectConfig>;
 	using UserId = ObjectId<IUser>;
 	using StreamId = StringId<IStream>;
 
diff --git a/Engine/Source/Programs/Horde/Horde.Build/Devices/IDeviceCollection.cs b/Engine/Source/Programs/Horde/Horde.Build/Devices/IDeviceCollection.cs
--- a/Engine/Source/Programs/Horde/Horde.Build/Devices/IDeviceCollection.cs
+++ b/Engine/Source/Programs/Horde/Horde.Build/Devices/IDeviceCollection.cs
@@ -13,7 +13,7 @@ namespace Horde.Build.Devices
 	using DeviceId = StringId<IDevice>;
 	using DevicePlatformId = StringId<IDevicePlatform>;
 	using DevicePoolId = StringId<IDevicePool>;
-	using ProjectId = StringId<IProject>;
+	using ProjectId = StringId<ProjectConfig>;
 	using UserId = ObjectId<IUser>;
 
 	/// <summary>
diff --git a/Engine/Source/Programs/Horde/Horde.Build/Projects/IProject.cs b/Engine/Source/Programs/Horde/Horde.Build/Projects/IProject.cs
deleted file mode 100644
diff --git a/Engine/Source/Programs/Horde/Horde.Build/Projects/IProjectCollection.cs b/Engine/Source/Programs/Horde/Horde.Build/Projects/IProjectCollection.cs
deleted file mode 100644
diff --git a/Engine/Source/Programs/Horde/Horde.Build/Projects/ProjectCollection.cs b/Engine/Source/Programs/Horde/Horde.Build/Projects/ProjectCollection.cs
deleted file mode 100644
diff --git a/Engine/Source/Programs/Horde/Horde.Build/Projects/ProjectConfig.cs b/Engine/Source/Programs/Horde/Horde.Build/Projects/ProjectConfig.cs
--- a/Engine/Source/Programs/Horde/Horde.Build/Projects/ProjectConfig.cs
+++ b/Engine/Source/Programs/Horde/Horde.Build/Projects/ProjectConfig.cs
@@ -16,7 +16,7 @@ using Horde.Build.Utilities;
 
 namespace Horde.Build.Projects
 {
-	using ProjectId = StringId<IProject>;
+	using ProjectId = StringId<ProjectConfig>;
 
 	/// <summary>
 	/// Stores configuration for a project
@@ -62,7 +62,7 @@ namespace Horde.Build.Projects
 		/// <summary>
 		/// Path to the project logo
 		/// </summary>
-		public string? Logo { get; set; }
+		public ConfigResource? Logo { get; set; }
 
 		/// <summary>
 		/// List of pools for this project
diff --git a/Engine/Source/Programs/Horde/Horde.Build/Projects/ProjectsController.cs b/Engine/Source/Programs/Horde/Horde.Build/Projects/ProjectsController.cs
--- a/Engine/Source/Programs/Horde/Horde.Build/Projects/ProjectsController.cs
+++ b/Engine/Source/Programs/Horde/Horde.Build/Projects/ProjectsController.cs
@@ -4,16 +4,18 @@ using System.Collections.Generic;
 using System.Linq;
 using System.Threading.Tasks;
 using Horde.Build.Acls;
+using Horde.Build.Configuration;
 using Horde.Build.Server;
 using Horde.Build.Streams;
 using Horde.Build.Utilities;
 using Microsoft.AspNetCore.Authorization;
 using Microsoft.AspNetCore.Mvc;
+using Microsoft.AspNetCore.StaticFiles;
 using Microsoft.Extensions.Options;
 
 namespace Horde.Build.Projects
 {
-	using ProjectId = StringId<IProject>;
+	using ProjectId = StringId<ProjectConfig>;
 
 	/// <summary>
 	/// Controller for the /api/v1/projects endpoint
@@ -23,15 +25,15 @@ namespace Horde.Build.Projects
 	[Route("[controller]")]
 	public class ProjectsController : HordeControllerBase
 	{
-		private readonly IProjectCollection _projectCollection;
 		private readonly IOptionsSnapshot<GlobalConfig> _globalConfig;
 
+		private static readonly FileExtensionContentTypeProvider s_contentTypeProvider = new FileExtensionContentTypeProvider();
+
 		/// <summary>
 		/// Constructor
 		/// </summary>
-		public ProjectsController(IProjectCollection projectCollection, IOptionsSnapshot<GlobalConfig> globalConfig)
+		public ProjectsController(IOptionsSnapshot<GlobalConfig> globalConfig)
 		{
-			_projectCollection = projectCollection;
 			_globalConfig = globalConfig;
 		}
 
@@ -105,7 +107,7 @@ namespace Horde.Build.Projects
 		/// <returns>Information about the requested project</returns>
 		[HttpGet]
 		[Route("/api/v1/projects/{projectId}/logo")]
-		public async Task<ActionResult<object>> GetProjectLogoAsync(ProjectId projectId)
+		public ActionResult<object> GetProjectLogo(ProjectId projectId)
 		{
 			ProjectConfig? projectConfig;
 			if (!_globalConfig.Value.TryGetProject(projectId, out projectConfig))
@@ -117,13 +119,19 @@ namespace Horde.Build.Projects
 				return Forbid(AclAction.ViewProject, projectId);
 			}
 
-			IProjectLogo? projectLogo = await _projectCollection.GetLogoAsync(projectId);
-			if (projectLogo == null)
+			ConfigResource? logoResource = projectConfig.Logo;
+			if (logoResource == null || logoResource.Path == null || logoResource.Data.Length == 0)
+			{
+				return NotFound("Missing logo resource data");
+			}
+
+			string? contentType;
+			if (!s_contentTypeProvider.TryGetContentType(logoResource.Path, out contentType))
 			{
-				return NotFound();
+				contentType = "application/octet-stream";
 			}
 
-			return new FileContentResult(projectLogo.Data, projectLogo.MimeType);
+			return new FileContentResult(logoResource.Data.ToArray(), contentType);
 		}
 	}
 }
diff --git a/Engine/Source/Programs/Horde/Horde.Build/Server/GlobalConfig.cs b/Engine/Source/Programs/Horde/Horde.Build/Server/GlobalConfig.cs
--- a/Engine/Source/Programs/Horde/Horde.Build/Server/GlobalConfig.cs
+++ b/Engine/Source/Programs/Horde/Horde.Build/Server/GlobalConfig.cs
@@ -27,7 +27,7 @@ using Microsoft.Extensions.Options;
 namespace Horde.Build.Server
 {
 	using UserId = ObjectId<IUser>;
-	using ProjectId = StringId<IProject>;
+	using ProjectId = StringId<ProjectConfig>;
 	using StreamId = StringId<IStream>;
 	using AgentSoftwareChannelName = StringId<AgentSoftwareChannels>;
 	using ToolId = StringId<ITool>;
diff --git a/Engine/Source/Programs/Horde/Horde.Build/Startup.cs b/Engine/Source/Programs/Horde/Horde.Build/Startup.cs
--- a/Engine/Source/Programs/Horde/Horde.Build/Startup.cs
+++ b/Engine/Source/Programs/Horde/Horde.Build/Startup.cs
@@ -107,7 +107,7 @@ namespace Horde.Build
 	{
 		static Startup()
 		{
-			ProtoBuf.Meta.RuntimeTypeModel.Default[typeof(StringId<IProject>)].SetSurrogate(typeof(StringIdProto<IProject>));
+			ProtoBuf.Meta.RuntimeTypeModel.Default[typeof(StringId<ProjectConfig>)].SetSurrogate(typeof(StringIdProto<ProjectConfig>));
 			ProtoBuf.Meta.RuntimeTypeModel.Default[typeof(StringId<IStream>)].SetSurrogate(typeof(StringIdProto<IStream>));
 		}
 
@@ -350,7 +350,6 @@ namespace Horde.Build
 			services.AddSingleton<ILogFileCollection, LogFileCollection>();
 			services.AddSingleton<INotificationTriggerCollection, NotificationTriggerCollection>();
 			services.AddSingleton<IPoolCollection, PoolCollection>();
-			services.AddSingleton<IProjectCollection, ProjectCollection>();
 			services.AddSingleton<ISessionCollection, SessionCollection>();
 			services.AddSingleton<IServiceAccountCollection, ServiceAccountCollection>();
 			services.AddSingleton<ISubscriptionCollection, SubscriptionCollection>();
diff --git a/Engine/Source/Programs/Horde/Horde.Build/Streams/IStream.cs b/Engine/Source/Programs/Horde/Horde.Build/Streams/IStream.cs
--- a/Engine/Source/Programs/Horde/Horde.Build/Streams/IStream.cs
+++ b/Engine/Source/Programs/Horde/Horde.Build/Streams/IStream.cs
@@ -17,7 +17,7 @@ using Horde.Build.Utilities;
 namespace Horde.Build.Streams
 {
 	using JobId = ObjectId<IJob>;
-	using ProjectId = StringId<IProject>;
+	using ProjectId = StringId<ProjectConfig>;
 	using StreamId = StringId<IStream>;
 	using TemplateId = StringId<ITemplateRef>;
 	using UserId = ObjectId<IUser>;
diff --git a/Engine/Source/Programs/Horde/Horde.Build/Streams/StreamCollection.cs b/Engine/Source/Programs/Horde/Horde.Build/Streams/StreamCollection.cs
--- a/Engine/Source/Programs/Horde/Horde.Build/Streams/StreamCollection.cs
+++ b/Engine/Source/Programs/Horde/Horde.Build/Streams/StreamCollection.cs
@@ -18,7 +18,7 @@ using MongoDB.Driver;
 namespace Horde.Build.Streams
 {
 	using JobId = ObjectId<IJob>;
-	using ProjectId = StringId<IProject>;
+	using ProjectId = StringId<ProjectConfig>;
 	using StreamId = StringId<IStream>;
 	using TemplateId = StringId<ITemplateRef>;
 	using UserId = ObjectId<IUser>;
diff --git a/Engine/Source/Programs/Horde/Horde.Build/Streams/StreamMessages.cs b/Engine/Source/Programs/Horde/Horde.Build/Streams/StreamMessages.cs
--- a/Engine/Source/Programs/Horde/Horde.Build/Streams/StreamMessages.cs
+++ b/Engine/Source/Programs/Horde/Horde.Build/Streams/StreamMessages.cs
@@ -20,7 +20,7 @@ namespace Horde.Build.Streams
 {
 	using JobId = ObjectId<IJob>;
 	using StreamId = StringId<IStream>;
-	using ProjectId = StringId<IProject>;
+	using ProjectId = StringId<ProjectConfig>;
 	using TemplateId = StringId<ITemplateRef>;
 
 	/// <summary>
diff --git a/Engine/Source/Programs/Horde/Horde.Build/Streams/StreamsController.cs b/Engine/Source/Programs/Horde/Horde.Build/Streams/StreamsController.cs
--- a/Engine/Source/Programs/Horde/Horde.Build/Streams/StreamsController.cs
+++ b/Engine/Source/Programs/Horde/Horde.Build/Streams/StreamsController.cs
@@ -21,7 +21,7 @@ using OpenTracing.Util;
 
 namespace Horde.Build.Streams
 {
-	using ProjectId = StringId<IProject>;
+	using ProjectId = StringId<ProjectConfig>;
 	using StreamId = StringId<IStream>;
 	using TemplateId = StringId<ITemplateRef>;
 
diff --git a/Engine/Source/Programs/Horde/Horde.Build/Utilities/HordeControllerBase.cs b/Engine/Source/Programs/Horde/Horde.Build/Utilities/HordeControllerBase.cs
--- a/Engine/Source/Programs/Horde/Horde.Build/Utilities/HordeControllerBase.cs
+++ b/Engine/Source/Programs/Horde/Horde.Build/Utilities/HordeControllerBase.cs
@@ -17,7 +17,7 @@ namespace Horde.Build.Utilities
 {
 	using JobId = ObjectId<IJob>;
 	using LeaseId = ObjectId<ILease>;
-	using ProjectId = StringId<IProject>;
+	using ProjectId = StringId<ProjectConfig>;
 	using StreamId = StringId<IStream>;
 	using ToolId = StringId<ITool>;
 	using ToolDeploymentId = ObjectId<IToolDeployment>;
