Horde: Build fix.

diff --git a/Engine/Source/Programs/Horde/Horde.Build/Configuration/ConfigType.cs b/Engine/Source/Programs/Horde/Horde.Build/Configuration/ConfigType.cs
--- a/Engine/Source/Programs/Horde/Horde.Build/Configuration/ConfigType.cs
+++ b/Engine/Source/Programs/Horde/Horde.Build/Configuration/ConfigType.cs
@@ -248,8 +248,6 @@ namespace Horde.Build.Configuration
 				resource.Path = uri.AbsoluteUri;
 				resource.Data = await file.ReadAsync(cancellationToken);
 				PropertyInfo.SetValue(target, resource);
-
-				context.Resources.Add(resource);
 			}
 		}
 
