Horde: Move compute service implementation into a "V1" namespace/folder.

diff --git a/Engine/Source/Programs/Horde/Horde.Build.Tests/Compute/RedisTaskSchedulerTest.cs b/Engine/Source/Programs/Horde/Horde.Build.Tests/Compute/RedisTaskSchedulerTest.cs
--- a/Engine/Source/Programs/Horde/Horde.Build.Tests/Compute/RedisTaskSchedulerTest.cs
+++ b/Engine/Source/Programs/Horde/Horde.Build.Tests/Compute/RedisTaskSchedulerTest.cs
@@ -2,7 +2,7 @@
 
 using System.Threading;
 using System.Threading.Tasks;
-using Horde.Build.Compute;
+using Horde.Build.Compute.V1;
 using Horde.Build.Server;
 using Microsoft.Extensions.Logging;
 using Microsoft.VisualStudio.TestTools.UnitTesting;
diff --git a/Engine/Source/Programs/Horde/Horde.Build.Tests/TestSetup.cs b/Engine/Source/Programs/Horde/Horde.Build.Tests/TestSetup.cs
--- a/Engine/Source/Programs/Horde/Horde.Build.Tests/TestSetup.cs
+++ b/Engine/Source/Programs/Horde/Horde.Build.Tests/TestSetup.cs
@@ -98,7 +98,7 @@ namespace Horde.Build.Tests
 		public AgentSoftwareService AgentSoftwareService => ServiceProvider.GetRequiredService<AgentSoftwareService>();
 		public AgentService AgentService => ServiceProvider.GetRequiredService<AgentService>();
 		public ICommitService CommitService => ServiceProvider.GetRequiredService<ICommitService>();
-		public ComputeService ComputeService => ServiceProvider.GetRequiredService<ComputeService>();
+		public Horde.Build.Compute.V1.ComputeService ComputeService => ServiceProvider.GetRequiredService<Horde.Build.Compute.V1.ComputeService>();
 		public GlobalsService GlobalsService => ServiceProvider.GetRequiredService<GlobalsService>();
 		public MongoService MongoService => ServiceProvider.GetRequiredService<MongoService>();
 		public ITemplateCollection TemplateCollection => ServiceProvider.GetRequiredService<ITemplateCollection>();
@@ -237,7 +237,7 @@ namespace Horde.Build.Tests
 			services.AddSingleton<AgentSoftwareService>();
 			services.AddSingleton<FleetService>();
 			services.AddSingleton<ConsistencyService>();
-			services.AddSingleton<ComputeService>();
+			services.AddSingleton<Horde.Build.Compute.V1.ComputeService>();
 			services.AddSingleton<RequestTrackerService>();
 			services.AddSingleton<GlobalsService>();
 			services.AddSingleton<CredentialService>();
diff --git a/Engine/Source/Programs/Horde/Horde.Build/Agents/Fleet/ComputeQueueAwsMetricStrategy.cs b/Engine/Source/Programs/Horde/Horde.Build/Agents/Fleet/ComputeQueueAwsMetricStrategy.cs
--- a/Engine/Source/Programs/Horde/Horde.Build/Agents/Fleet/ComputeQueueAwsMetricStrategy.cs
+++ b/Engine/Source/Programs/Horde/Horde.Build/Agents/Fleet/ComputeQueueAwsMetricStrategy.cs
@@ -9,7 +9,7 @@ using Amazon.CloudWatch;
 using Amazon.CloudWatch.Model;
 using EpicGames.Horde.Compute;
 using Horde.Build.Agents.Pools;
-using Horde.Build.Compute;
+using Horde.Build.Compute.V1;
 using Microsoft.Extensions.Logging;
 using OpenTracing;
 using OpenTracing.Util;
diff --git a/Engine/Source/Programs/Horde/Horde.Build/Compute/ComputeController.cs b/Engine/Source/Programs/Horde/Horde.Build/Compute/ComputeController.cs
deleted file mode 100644
diff --git a/Engine/Source/Programs/Horde/Horde.Build/Compute/ComputeService.cs b/Engine/Source/Programs/Horde/Horde.Build/Compute/ComputeService.cs
deleted file mode 100644
diff --git a/Engine/Source/Programs/Horde/Horde.Build/Compute/IComputeService.cs b/Engine/Source/Programs/Horde/Horde.Build/Compute/IComputeService.cs
deleted file mode 100644
diff --git a/Engine/Source/Programs/Horde/Horde.Build/Compute/RedisMessageQueue.cs b/Engine/Source/Programs/Horde/Horde.Build/Compute/RedisMessageQueue.cs
deleted file mode 100644
diff --git a/Engine/Source/Programs/Horde/Horde.Build/Compute/RedisTaskScheduler.cs b/Engine/Source/Programs/Horde/Horde.Build/Compute/RedisTaskScheduler.cs
deleted file mode 100644
diff --git a/Engine/Source/Programs/Horde/Horde.Build/Compute/V1/ComputeController.cs b/Engine/Source/Programs/Horde/Horde.Build/Compute/V1/ComputeController.cs
new file mode 100644
--- /dev/null
+++ b/Engine/Source/Programs/Horde/Horde.Build/Compute/V1/ComputeController.cs
@@ -0,0 +1,132 @@
+// Copyright Epic Games, Inc. All Rights Reserved.
+
+using System.Collections.Generic;
+using System.Threading;
+using System.Threading.Tasks;
+using EpicGames.Horde.Compute;
+using EpicGames.Horde.Compute.Impl;
+using Horde.Build.Acls;
+using Horde.Build.Server;
+using Horde.Build.Utilities;
+using Microsoft.AspNetCore.Authorization;
+using Microsoft.AspNetCore.Mvc;
+using Microsoft.Extensions.Options;
+
+namespace Horde.Build.Compute.V1
+{
+	/// <summary>
+	/// Controller for the /api/v1/compute endpoint
+	/// </summary>
+	[ApiController]
+	[Authorize]
+	[Route("[controller]")]
+	public class ComputeController : HordeControllerBase
+	{
+		readonly ComputeService _computeService;
+		readonly IOptionsSnapshot<GlobalConfig> _globalConfig;
+
+		/// <summary>
+		/// Constructor
+		/// </summary>
+		public ComputeController(ComputeService computeService, IOptionsSnapshot<GlobalConfig> globalConfig)
+		{
+			_computeService = computeService;
+			_globalConfig = globalConfig;
+		}
+
+		/// <summary>
+		/// Gets information about a cluster
+		/// </summary>
+		/// <param name="clusterId">The cluster to add to</param>
+		/// <returns></returns>
+		[HttpGet]
+		[Authorize]
+		[Route("/api/v1/compute/{clusterId}")]
+		public async Task<ActionResult<GetComputeClusterInfo>> GetClusterInfoAsync([FromRoute] ClusterId clusterId)
+		{
+			if (!_globalConfig.Value.Authorize(AclAction.ViewComputeTasks, User))
+			{
+				return Forbid(AclAction.ViewComputeTasks);
+			}
+
+			IComputeClusterInfo clusterInfo;
+			try
+			{
+				clusterInfo = await _computeService.GetClusterInfoAsync(clusterId);
+			}
+			catch (KeyNotFoundException)
+			{
+				return NotFound("Cluster '{ClusterId}' was found", clusterId);
+			}
+
+			GetComputeClusterInfo response = new GetComputeClusterInfo();
+			response.Id = clusterInfo.Id;
+			response.NamespaceId = clusterInfo.NamespaceId;
+			response.RequestBucketId = clusterInfo.RequestBucketId;
+			response.ResponseBucketId = clusterInfo.ResponseBucketId;
+
+			return response;
+		}
+
+		/// <summary>
+		/// Add tasks to be executed remotely
+		/// </summary>
+		/// <param name="clusterId">The cluster to add to</param>
+		/// <param name="request">The request parameters</param>
+		/// <returns></returns>
+		[HttpPost]
+		[Authorize]
+		[Route("/api/v1/compute/{clusterId}")]
+		public async Task<ActionResult> AddTasksAsync([FromRoute] ClusterId clusterId, [FromBody] AddTasksRequest request)
+		{
+			if(!_globalConfig.Value.Authorize(AclAction.AddComputeTasks, User))
+			{
+				return Forbid(AclAction.AddComputeTasks);
+			}
+			if (request.TaskRefIds.Count == 0)
+			{
+				return BadRequest("No task hashes specified");
+			}
+
+			await _computeService.AddTasksAsync(clusterId, request.ChannelId, request.TaskRefIds, request.RequirementsHash);
+			return Ok();
+		}
+
+		/// <summary>
+		/// Read updates for a particular channel
+		/// </summary>
+		/// <param name="clusterId"></param>
+		/// <param name="channelId">The channel to add to</param>
+		/// <param name="wait">Amount of time to wait before responding, in seconds</param>
+		/// <returns></returns>
+		[HttpPost]
+		[Authorize]
+		[Route("/api/v1/compute/{clusterId}/updates/{channelId}")]
+		public async Task<ActionResult<GetTaskUpdatesResponse>> GetUpdatesAsync([FromRoute] ClusterId clusterId, [FromRoute] ChannelId channelId, [FromQuery] int wait = 0)
+		{
+			if (!_globalConfig.Value.Authorize(AclAction.ViewComputeTasks, User))
+			{
+				return Forbid(AclAction.ViewComputeTasks);
+			}
+
+			List<ComputeTaskStatus> updates;
+			if (wait == 0)
+			{
+				updates = await _computeService.GetTaskUpdatesAsync(clusterId, channelId);
+			}
+			else
+			{
+				using CancellationTokenSource delaySource = new CancellationTokenSource(wait * 1000);
+				updates = await _computeService.WaitForTaskUpdatesAsync(clusterId, channelId, delaySource.Token);
+			}
+
+			GetTaskUpdatesResponse response = new GetTaskUpdatesResponse();
+			foreach (ComputeTaskStatus update in updates)
+			{
+				response.Updates.Add(update);
+			}
+
+			return response;
+		}
+	}
+}
diff --git a/Engine/Source/Programs/Horde/Horde.Build/Compute/V1/ComputeService.cs b/Engine/Source/Programs/Horde/Horde.Build/Compute/V1/ComputeService.cs
new file mode 100644
--- /dev/null
+++ b/Engine/Source/Programs/Horde/Horde.Build/Compute/V1/ComputeService.cs
@@ -0,0 +1,526 @@
+// Copyright Epic Games, Inc. All Rights Reserved.
+
+using System;
+using System.Collections.Generic;
+using System.Linq;
+using System.Threading;
+using System.Threading.Tasks;
+using EpicGames.Core;
+using EpicGames.Horde.Compute;
+using EpicGames.Horde.Storage;
+using EpicGames.Redis;
+using EpicGames.Serialization;
+using Google.Protobuf;
+using Google.Protobuf.WellKnownTypes;
+using Horde.Build.Agents;
+using Horde.Build.Agents.Leases;
+using Horde.Build.Agents.Pools;
+using Horde.Build.Server;
+using Horde.Build.Tasks;
+using Horde.Build.Utilities;
+using HordeCommon;
+using HordeCommon.Rpc.Tasks;
+using Microsoft.Extensions.Caching.Memory;
+using Microsoft.Extensions.Hosting;
+using Microsoft.Extensions.Logging;
+using Microsoft.Extensions.Options;
+using OpenTracing;
+using OpenTracing.Util;
+using StackExchange.Redis;
+
+namespace Horde.Build.Compute.V1
+{
+	using LeaseId = ObjectId<ILease>;
+
+	/// <summary>
+	/// Information about a particular task
+	/// </summary>
+	[RedisConverter(typeof(RedisCbConverter<>))]
+	class ComputeTaskInfo
+	{
+		[CbField("c")]
+		public ClusterId ClusterId { get; set; }
+
+		[CbField("h")]
+		public RefId TaskRefId { get; set; }
+
+		[CbField("ch")]
+		public ChannelId ChannelId { get; set; }
+
+		[CbField("q")]
+		public DateTime QueuedAt { get; set; }
+
+		private ComputeTaskInfo()
+		{
+		}
+
+		public ComputeTaskInfo(ClusterId clusterId, RefId taskRefId, ChannelId channelId, DateTime queuedAt)
+		{
+			ClusterId = clusterId;
+			TaskRefId = taskRefId;
+			ChannelId = channelId;
+			QueuedAt = queuedAt;
+		}
+	}
+
+	/// <summary>
+	/// Dispatches remote actions. Does not implement any cross-pod communication to satisfy leases; only agents connected to this server instance will be stored.
+	/// </summary>
+	public sealed class ComputeService : TaskSourceBase<ComputeTaskMessage>, IHostedService, IDisposable, IComputeService
+	{
+		[RedisConverter(typeof(QueueKeySerializer))]
+		class QueueKey
+		{
+			public ClusterId ClusterId { get; set; }
+			public IoHash RequirementsHash { get; set; }
+
+			public QueueKey(ClusterId clusterId, IoHash requirementsHash)
+			{
+				ClusterId = clusterId;
+				RequirementsHash = requirementsHash;
+			}
+
+			public override string ToString() => $"{ClusterId}/{RequirementsHash}";
+		}
+
+		class QueueKeySerializer : IRedisConverter<QueueKey>
+		{
+			public QueueKey FromRedisValue(RedisValue value)
+			{
+				string str = value.ToString();
+				int idx = str.LastIndexOf("/", StringComparison.Ordinal);
+				return new QueueKey(new ClusterId(str.Substring(0, idx)), IoHash.Parse(str.Substring(idx + 1)));
+			}
+
+			public RedisValue ToRedisValue(QueueKey value) => $"{value.ClusterId}/{value.RequirementsHash}";
+		}
+
+		class ClusterInfo : IComputeClusterInfo
+		{
+			public ClusterId Id { get; set; }
+			public NamespaceId NamespaceId { get; set; }
+			public BucketId RequestBucketId { get; set; }
+			public BucketId ResponseBucketId { get; set; }
+
+			public ClusterInfo(ComputeClusterConfig config)
+			{
+				Id = new ClusterId(config.Id);
+				NamespaceId = new NamespaceId(config.NamespaceId);
+				RequestBucketId = new BucketId(config.RequestBucketId);
+				ResponseBucketId = new BucketId(config.ResponseBucketId);
+			}
+		}
+
+		/// <inheritdoc/>
+		public override string Type => "Compute";
+
+		/// <inheritdoc/>
+		public override TaskSourceFlags Flags => TaskSourceFlags.None;
+
+		/// <summary>
+		/// ID of the default namespace
+		/// </summary>
+		public static NamespaceId DefaultNamespaceId { get; } = new NamespaceId("default");
+
+		readonly ILegacyStorageClient _storageClient;
+		readonly ITaskScheduler<QueueKey, ComputeTaskInfo> _taskScheduler;
+		readonly RedisMessageQueue<ComputeTaskStatus> _messageQueue;
+		readonly ITicker _expireTasksTicker;
+		readonly IMemoryCache _requirementsCache;
+		readonly IOptionsMonitor<GlobalConfig> _globalConfig;
+		readonly ILogger _logger;
+
+		static ComputeService()
+		{
+			RedisSerializer.RegisterConverter<ComputeTaskStatus, RedisCbConverter<ComputeTaskStatus>>();
+		}
+
+		/// <summary>
+		/// Constructor
+		/// </summary>
+		public ComputeService(GlobalsService globalsService, RedisService redisService, ILegacyStorageClient storageClient, IClock clock, IOptionsMonitor<GlobalConfig> globalConfig, ILogger<ComputeService> logger)
+		{
+			_storageClient = storageClient;
+			_taskScheduler = new RedisTaskScheduler<QueueKey, ComputeTaskInfo>(redisService.ConnectionPool, "compute/tasks/", logger);
+			_messageQueue = new RedisMessageQueue<ComputeTaskStatus>(redisService.GetDatabase(), "compute/messages/");
+			_expireTasksTicker = clock.AddTicker<ComputeService>(TimeSpan.FromMinutes(2.0), ExpireTasksAsync, logger);
+			_requirementsCache = new MemoryCache(new MemoryCacheOptions());
+			_globalConfig = globalConfig;
+			_logger = logger;
+
+			OnLeaseStartedProperties.Add(x => x.TaskRefId);
+		}
+
+		static ComputeTaskStatus CreateStatus(RefId taskRefId, ComputeTaskState state)
+		{
+			ComputeTaskStatus status = new ComputeTaskStatus();
+			status.TaskRefId = taskRefId;
+			status.Time = DateTime.UtcNow;
+			status.State = state;
+			return status;
+		}
+
+		/// <inheritdoc/>
+		public Task StartAsync(CancellationToken token) => _expireTasksTicker.StartAsync();
+
+		/// <inheritdoc/>
+		public Task StopAsync(CancellationToken token) => _expireTasksTicker.StopAsync();
+
+		/// <summary>
+		/// Expire tasks that are in inactive queues (ie. no machines can execute them)
+		/// </summary>
+		/// <param name="cancellationToken"></param>
+		/// <returns></returns>
+		async ValueTask ExpireTasksAsync(CancellationToken cancellationToken)
+		{
+			using IScope scope = GlobalTracer.Instance.BuildSpan("ComputeService.ExpireTasksAsync").StartActive();
+			List<QueueKey> queueKeys = await _taskScheduler.GetInactiveQueuesAsync();
+			scope.Span.SetTag("numQueueKeys", queueKeys.Count);
+			
+			foreach (QueueKey queueKey in queueKeys)
+			{
+				_logger.LogInformation("Inactive queue: {QueueKey}", queueKey);
+				for (; ; )
+				{
+					ComputeTaskInfo? computeTask = await _taskScheduler.DequeueAsync(queueKey);
+					if (computeTask == null)
+					{
+						break;
+					}
+
+					ComputeTaskStatus status = CreateStatus(computeTask.TaskRefId, ComputeTaskState.Complete);
+					status.Outcome = ComputeTaskOutcome.Expired;
+					status.Detail = $"No agents monitoring queue {queueKey}";
+
+					_logger.LogInformation("Compute task expired (queue: {RequirementsHash}, task: {TaskHash}, channel: {ChannelId})", queueKey, computeTask.TaskRefId, computeTask.ChannelId);
+					await PostStatusMessageAsync(computeTask, status);
+				}
+			}
+		}
+
+		/// <inheritdoc/>
+		public void Dispose()
+		{
+			_messageQueue.Dispose();
+			_expireTasksTicker.Dispose();
+			_requirementsCache.Dispose();
+		}
+
+		/// <inheritdoc/>
+		public Task<IComputeClusterInfo> GetClusterInfoAsync(ClusterId clusterId)
+		{
+			ComputeClusterConfig? config = GetCluster(clusterId);
+			if (config == null)
+			{
+				throw new KeyNotFoundException();
+			}
+			return Task.FromResult<IComputeClusterInfo>(new ClusterInfo(config));
+		}
+
+		/// <inheritdoc/>
+		public async Task AddTasksAsync(ClusterId clusterId, ChannelId channelId, List<RefId> taskRefIds, CbObjectAttachment requirementsHash)
+		{
+			List<Task> tasks = new List<Task>();
+			foreach (RefId taskRefId in taskRefIds)
+			{
+				ComputeTaskInfo taskInfo = new ComputeTaskInfo(clusterId, taskRefId, channelId, DateTime.UtcNow);
+				_logger.LogDebug("Adding task {TaskHash} from channel {ChannelId} to queue {ClusterId}{RequirementsHash}", taskRefId.Hash, channelId, clusterId, requirementsHash);
+				tasks.Add(_taskScheduler.EnqueueAsync(new QueueKey(clusterId, requirementsHash), taskInfo, false));
+			}
+			await Task.WhenAll(tasks);
+		}
+
+		ComputeClusterConfig? GetCluster(ClusterId clusterId)
+		{
+			return _globalConfig.CurrentValue.Compute.FirstOrDefault(x => new ClusterId(x.Id) == clusterId);
+		}
+
+		/// <inheritdoc/>
+		public override async Task<Task<AgentLease?>> AssignLeaseAsync(IAgent agent, CancellationToken cancellationToken)
+		{
+			Task<(QueueKey, ComputeTaskInfo)?> task = await _taskScheduler.DequeueAsync(queueKey => CheckRequirements(agent, queueKey), cancellationToken);
+			return WaitForLeaseAsync(agent, task, cancellationToken);
+		}
+
+		private async Task<AgentLease?> WaitForLeaseAsync(IAgent agent, Task<(QueueKey, ComputeTaskInfo)?> task, CancellationToken cancellationToken)
+		{
+			for (; ; )
+			{
+				(QueueKey, ComputeTaskInfo)? entry = await task;
+				if (entry == null)
+				{
+					return null;
+				}
+
+				AgentLease? lease = await CreateLeaseForEntryAsync(agent, entry.Value);
+				if (lease != null)
+				{
+					return lease;
+				}
+
+				task = await _taskScheduler.DequeueAsync(queueKey => CheckRequirements(agent, queueKey), cancellationToken);
+			}
+		}
+
+		private async Task<AgentLease?> CreateLeaseForEntryAsync(IAgent agent, (QueueKey, ComputeTaskInfo) entry)
+		{
+			(QueueKey queueKey, ComputeTaskInfo taskInfo) = entry;
+
+			ComputeClusterConfig? cluster = GetCluster(taskInfo.ClusterId);
+			if (cluster == null)
+			{
+				_logger.LogWarning("Invalid cluster '{ClusterId}'; failing task {TaskRefId}", taskInfo.ClusterId, taskInfo.TaskRefId);
+
+				ComputeTaskStatus status = CreateStatus(taskInfo.TaskRefId, ComputeTaskState.Complete);
+				status.AgentId = agent.Id.ToString();
+				status.Detail = $"Invalid cluster '{taskInfo.ClusterId}'";
+				await PostStatusMessageAsync(taskInfo, status);
+
+				return null;
+			}
+
+			Requirements? requirements = await GetCachedRequirementsAsync(queueKey);
+			if (requirements == null)
+			{
+				_logger.LogWarning("Unable to fetch requirements {RequirementsHash}", queueKey);
+
+				ComputeTaskStatus status = CreateStatus(taskInfo.TaskRefId, ComputeTaskState.Complete);
+				status.AgentId = agent.Id.ToString();
+				status.Detail = $"Unable to retrieve requirements '{queueKey}'";
+				await PostStatusMessageAsync(taskInfo, status);
+
+				return null;
+			}
+
+			ComputeTaskMessage computeTask = new ComputeTaskMessage();
+			computeTask.ClusterId = taskInfo.ClusterId.ToString();
+			computeTask.ChannelId = taskInfo.ChannelId.ToString();
+			computeTask.NamespaceId = cluster.NamespaceId.ToString();
+			computeTask.InputBucketId = cluster.RequestBucketId.ToString();
+			computeTask.OutputBucketId = cluster.ResponseBucketId.ToString();
+			computeTask.RequirementsHash = queueKey.RequirementsHash;
+			computeTask.TaskRefId = taskInfo.TaskRefId;
+			computeTask.QueuedAt = Timestamp.FromDateTime(taskInfo.QueuedAt);
+			computeTask.DispatchedMs = (int)(DateTime.UtcNow - taskInfo.QueuedAt).TotalMilliseconds;
+				
+			string leaseName = $"Remote action ({taskInfo.TaskRefId})";
+			byte[] payload = Any.Pack(computeTask).ToByteArray();
+
+			AgentLease lease = new AgentLease(LeaseId.GenerateNewId(), leaseName, null, null, null, LeaseState.Pending, requirements.Resources, requirements.Exclusive, payload);
+			_logger.LogDebug("Created lease {LeaseId} for channel {ChannelId} task {TaskHash} req {RequirementsHash}", lease.Id, computeTask.ChannelId, computeTask.TaskRefId, computeTask.RequirementsHash);
+			return lease;
+		}
+
+		/// <inheritdoc/>
+		public override Task CancelLeaseAsync(IAgent agent, LeaseId leaseId, ComputeTaskMessage message)
+		{
+			ClusterId clusterId = new ClusterId(message.ClusterId);
+			ComputeTaskInfo taskInfo = new ComputeTaskInfo(clusterId, new RefId(new IoHash(message.TaskRefId.ToByteArray())), new ChannelId(message.ChannelId), message.QueuedAt.ToDateTime());
+			return _taskScheduler.EnqueueAsync(new QueueKey(clusterId, new IoHash(message.RequirementsHash.ToByteArray())), taskInfo, true);
+		}
+
+		/// <inheritdoc/>
+		public async Task<List<ComputeTaskStatus>> GetTaskUpdatesAsync(ClusterId clusterId, ChannelId channelId)
+		{
+			return await _messageQueue.ReadMessagesAsync(GetMessageQueueId(clusterId, channelId));
+		}
+		
+		/// <inheritdoc/>
+		public async Task<int> GetNumQueuedTasksForPoolAsync(ClusterId clusterId, IPool pool, CancellationToken cancellationToken = default)
+		{
+			return await _taskScheduler.GetNumQueuedTasksAsync(queueKey => CheckRequirements(pool, queueKey), cancellationToken);
+		}
+
+		/// <inheritdoc/>
+		public async Task<List<ComputeTaskStatus>> WaitForTaskUpdatesAsync(ClusterId clusterId, ChannelId channelId, CancellationToken cancellationToken)
+		{
+			return await _messageQueue.WaitForMessagesAsync(GetMessageQueueId(clusterId, channelId), cancellationToken);
+		}
+
+		/// <inheritdoc/>
+		public override async Task OnLeaseStartedAsync(IAgent agent, LeaseId leaseId, ComputeTaskMessage computeTask, ILogger logger)
+		{
+			await base.OnLeaseStartedAsync(agent, leaseId, computeTask, logger);
+
+			ComputeTaskStatus status = CreateStatus(computeTask.TaskRefId, ComputeTaskState.Executing);
+			status.AgentId = agent.Id.ToString();
+			status.LeaseId = leaseId.ToString();
+			await PostStatusMessageAsync(computeTask, status);
+		}
+
+		/// <inheritdoc/>
+		public override async Task OnLeaseFinishedAsync(IAgent agent, LeaseId leaseId, ComputeTaskMessage computeTask, LeaseOutcome outcome, ReadOnlyMemory<byte> output, ILogger logger)
+		{
+			await base.OnLeaseFinishedAsync(agent, leaseId, computeTask, outcome, output, logger);
+
+			DateTime queuedAt = computeTask.QueuedAt.ToDateTime();
+
+			ComputeTaskResultMessage message = ComputeTaskResultMessage.Parser.ParseFrom(output.ToArray());
+
+			ComputeTaskStatus status = CreateStatus(computeTask.TaskRefId, ComputeTaskState.Complete);
+			status.AgentId = agent.Id.ToString();
+			status.LeaseId = leaseId.ToString();
+			status.QueueStats = new ComputeTaskQueueStats(queuedAt, computeTask.DispatchedMs, (int)(DateTime.UtcNow - queuedAt).TotalMilliseconds);
+			status.ExecutionStats = message.ExecutionStats?.ToNative();
+
+			if (message.ResultRefId != null)
+			{
+				status.ResultRefId = message.ResultRefId;
+			}
+			else if ((ComputeTaskOutcome)message.Outcome != ComputeTaskOutcome.Success)
+			{
+				(status.Outcome, status.Detail) = ((ComputeTaskOutcome)message.Outcome, message.Detail);
+			}
+			else if (outcome == LeaseOutcome.Failed)
+			{
+				status.Outcome = ComputeTaskOutcome.Failed;
+			}
+			else if (outcome == LeaseOutcome.Cancelled)
+			{
+				status.Outcome = ComputeTaskOutcome.Cancelled;
+			}
+			else
+			{
+				status.Outcome = ComputeTaskOutcome.NoResult;
+			}
+
+			logger.LogInformation("Compute lease finished (lease: {LeaseId}, task: {TaskHash}, agent: {AgentId}, channel: {ChannelId}, result: {ResultHash}, outcome: {Outcome})", leaseId, computeTask.TaskRefId.AsRefId(), agent.Id, computeTask.ChannelId, status.ResultRefId?.ToString() ?? "(none)", status.Outcome);
+			await PostStatusMessageAsync(computeTask, status);
+		}
+
+		/// <summary>
+		/// Checks that an agent matches the necessary criteria to execute a task
+		/// </summary>
+		/// <param name="agent"></param>
+		/// <param name="queueKey"></param>
+		/// <returns></returns>
+		async ValueTask<bool> CheckRequirements(IAgent agent, QueueKey queueKey)
+		{
+			Requirements? requirements = await GetCachedRequirementsAsync(queueKey);
+			if (requirements == null)
+			{
+				return false;
+			}
+			return agent.MeetsRequirements(requirements);
+		}
+		
+		/// <summary>
+		/// Checks that a pool matches the necessary criteria to execute a task
+		/// </summary>
+		/// <param name="pool"></param>
+		/// <param name="queueKey"></param>
+		/// <returns></returns>
+		async ValueTask<bool> CheckRequirements(IPool pool, QueueKey queueKey)
+		{
+			Requirements? requirements = await GetCachedRequirementsAsync(queueKey);
+			if (requirements == null)
+			{
+				return false;
+			}
+			return pool.MeetsRequirements(requirements);
+		}
+
+		/// <summary>
+		/// Gets the requirements object from the CAS
+		/// </summary>
+		/// <param name="queueKey">Queue identifier</param>
+		/// <returns>Requirements object for the queue</returns>
+		async ValueTask<Requirements?> GetCachedRequirementsAsync(QueueKey queueKey)
+		{
+			Requirements? requirements;
+			if (!_requirementsCache.TryGetValue(queueKey.RequirementsHash, out requirements))
+			{
+				requirements = await GetRequirementsAsync(queueKey);
+				if (requirements != null)
+				{
+					using (ICacheEntry entry = _requirementsCache.CreateEntry(queueKey.RequirementsHash))
+					{
+						entry.SetSlidingExpiration(TimeSpan.FromMinutes(10.0));
+						entry.SetValue(requirements);
+					}
+				}
+			}
+			return requirements;
+		}
+
+		/// <summary>
+		/// Gets the requirements object for a given queue. Fails tasks in the queue if the requirements object is missing.
+		/// </summary>
+		/// <param name="queueKey">Queue identifier</param>
+		/// <returns>Requirements object for the queue</returns>
+		async ValueTask<Requirements?> GetRequirementsAsync(QueueKey queueKey)
+		{
+			Requirements? requirements = null;
+
+			ComputeClusterConfig? clusterConfig = GetCluster(queueKey.ClusterId);
+			if (clusterConfig != null)
+			{
+				NamespaceId namespaceId = new NamespaceId(clusterConfig.NamespaceId);
+				try
+				{
+					requirements = await _storageClient.ReadBlobAsync<Requirements>(namespaceId, queueKey.RequirementsHash);
+				}
+				catch (LegacyBlobNotFoundException)
+				{
+				}
+				catch (Exception ex)
+				{
+					_logger.LogError(ex, "Unable to read blob {NamespaceId}/{RequirementsHash} from storage service", clusterConfig.NamespaceId, queueKey.RequirementsHash);
+				}
+			}
+
+			if (requirements == null)
+			{
+				_logger.LogWarning("Unable to fetch requirements object for queue {QueueKey}; failing queued tasks.", queueKey);
+				for (; ; )
+				{
+					ComputeTaskInfo? computeTask = await _taskScheduler.DequeueAsync(queueKey);
+					if (computeTask == null)
+					{
+						break;
+					}
+
+					ComputeTaskStatus status = CreateStatus(computeTask.TaskRefId, ComputeTaskState.Complete);
+					status.Outcome = ComputeTaskOutcome.BlobNotFound;
+					status.Detail = $"Missing requirements object {queueKey.RequirementsHash}";
+					_logger.LogInformation("Compute task failed due to missing requirements (queue: {QueueKey}, task: {TaskHash}, channel: {ChannelId})", queueKey, computeTask.TaskRefId, computeTask.ChannelId);
+					await PostStatusMessageAsync(computeTask, status);
+				}
+			}
+
+			return requirements;
+		}
+
+		/// <summary>
+		/// Post a status message for a particular task
+		/// </summary>
+		/// <param name="computeTask">The compute task instance</param>
+		/// <param name="status">New status for the task</param>
+		async Task PostStatusMessageAsync(ComputeTaskInfo computeTask, ComputeTaskStatus status)
+		{
+			await _messageQueue.PostAsync(GetMessageQueueId(computeTask.ClusterId, computeTask.ChannelId), status);
+		}
+
+		/// <summary>
+		/// Post a status message for a particular task
+		/// </summary>
+		/// <param name="computeTaskMessage">The compute task lease</param>
+		/// <param name="status">New status for the task</param>
+		/// <returns></returns>
+		async Task PostStatusMessageAsync(ComputeTaskMessage computeTaskMessage, ComputeTaskStatus status)
+		{
+			await _messageQueue.PostAsync(GetMessageQueueId(new ClusterId(computeTaskMessage.ClusterId), new ChannelId(computeTaskMessage.ChannelId)), status);
+		}
+
+		/// <summary>
+		/// Gets the name of a particular message queue
+		/// </summary>
+		/// <param name="clusterId">The compute cluster</param>
+		/// <param name="channelId">Identifier for the message channel</param>
+		/// <returns>Name of the message queue</returns>
+		static string GetMessageQueueId(ClusterId clusterId, ChannelId channelId)
+		{
+			return $"{clusterId}/{channelId}";
+		}
+	}
+}
diff --git a/Engine/Source/Programs/Horde/Horde.Build/Compute/V1/IComputeService.cs b/Engine/Source/Programs/Horde/Horde.Build/Compute/V1/IComputeService.cs
new file mode 100644
--- /dev/null
+++ b/Engine/Source/Programs/Horde/Horde.Build/Compute/V1/IComputeService.cs
@@ -0,0 +1,59 @@
+// Copyright Epic Games, Inc. All Rights Reserved.
+
+using System.Collections.Generic;
+using System.Threading;
+using System.Threading.Tasks;
+using EpicGames.Horde.Compute;
+using EpicGames.Horde.Storage;
+using EpicGames.Serialization;
+using Horde.Build.Agents.Pools;
+
+namespace Horde.Build.Compute.V1;
+
+/// <summary>
+/// Interface for running and scheduling compute tasks (primarily for remote execution)
+/// </summary>
+public interface IComputeService
+{
+	/// <summary>
+	/// Post tasks to be executed to a channel
+	/// </summary>
+	/// <param name="clusterId">Cluster to use for execution</param>
+	/// <param name="channelId">Unique identifier of the client</param>
+	/// <param name="taskRefIds">List of task hashes</param>
+	/// <param name="requirementsHash">Requirements document for execution</param>
+	/// <returns>Async task</returns>
+	Task AddTasksAsync(ClusterId clusterId, ChannelId channelId, List<RefId> taskRefIds, CbObjectAttachment requirementsHash);
+
+	/// <summary>
+	/// Gets information about a compute cluster
+	/// </summary>
+	/// <param name="clusterId">Cluster to use for execution</param>
+	Task<IComputeClusterInfo> GetClusterInfoAsync(ClusterId clusterId);
+
+	/// <summary>
+	/// Dequeue completed items from a queue and return immediately
+	/// </summary>
+	/// <param name="clusterId">Cluster containing the channel</param>
+	/// <param name="channelId">Queue to remove items from</param>
+	/// <returns>List of status updates</returns>
+	Task<List<ComputeTaskStatus>> GetTaskUpdatesAsync(ClusterId clusterId, ChannelId channelId);
+
+	/// <summary>
+	/// Get number of queued tasks that are compatible for a given pool
+	/// </summary>
+	/// <param name="clusterId">Cluster to query</param>
+	/// <param name="pool">Pool to check for compatibility</param>
+	/// <param name="cancellationToken">Cancellation token for the operation</param> 
+	/// <returns>Number of tasks in queue</returns>
+	Task<int> GetNumQueuedTasksForPoolAsync(ClusterId clusterId, IPool pool, CancellationToken cancellationToken = default);
+
+	/// <summary>
+	/// Dequeue completed items from a queue
+	/// </summary>
+	/// <param name="clusterId">Cluster containing the channel</param>
+	/// <param name="channelId">Queue to remove items from</param>
+	/// <param name="cancellationToken">Cancellation token to stop waiting for items</param>
+	/// <returns>List of status updates</returns>
+	Task<List<ComputeTaskStatus>> WaitForTaskUpdatesAsync(ClusterId clusterId, ChannelId channelId, CancellationToken cancellationToken);
+}
\ No newline at end of file
diff --git a/Engine/Source/Programs/Horde/Horde.Build/Compute/V1/RedisMessageQueue.cs b/Engine/Source/Programs/Horde/Horde.Build/Compute/V1/RedisMessageQueue.cs
new file mode 100644
--- /dev/null
+++ b/Engine/Source/Programs/Horde/Horde.Build/Compute/V1/RedisMessageQueue.cs
@@ -0,0 +1,180 @@
+// Copyright Epic Games, Inc. All Rights Reserved.
+
+using System;
+using System.Collections.Generic;
+using System.Threading;
+using System.Threading.Tasks;
+using EpicGames.Redis;
+using StackExchange.Redis;
+
+namespace Horde.Build.Compute.V1
+{
+	/// <summary>
+	/// Interface for a distributed message queue
+	/// </summary>
+	/// <typeparam name="T">The message type</typeparam>
+	interface IMessageQueue<T> where T : class
+	{
+		/// <summary>
+		/// Adds messages to the queue
+		/// </summary>
+		/// <param name="channelId">The channel to post to</param>
+		/// <param name="message">Message to post</param>
+		Task PostAsync(string channelId, T message);
+
+		/// <summary>
+		/// Waits for a message to be available on the given channel
+		/// </summary>
+		/// <param name="channelId">The channel to read from</param>
+		/// <returns>True if a message was available, false otherwise</returns>
+		Task<List<T>> ReadMessagesAsync(string channelId);
+
+		/// <summary>
+		/// Waits for a message to be available on the given channel
+		/// </summary>
+		/// <param name="channelId">The channel to read from</param>
+		/// <param name="cancellationToken">May be signalled to stop the wait, without throwing an exception</param>
+		/// <returns>True if a message was available, false otherwise</returns>
+		Task<List<T>> WaitForMessagesAsync(string channelId, CancellationToken cancellationToken);
+	}
+
+	/// <summary>
+	/// Implementation of a message queue using Redis
+	/// </summary>
+	/// <typeparam name="T">The message type. Must be serailizable by RedisSerializer.</typeparam>
+	class RedisMessageQueue<T> : IMessageQueue<T>, IDisposable where T : class
+	{
+		readonly IDatabase _redis;
+		readonly RedisKey _keyPrefix;
+		readonly RedisChannel _updateChannel;
+		readonly Dictionary<string, TaskCompletionSource<bool>> _channelWakeEvents = new Dictionary<string, TaskCompletionSource<bool>>();
+
+		// Time after which entries should be removed
+		TimeSpan ExpireTime { get; set; } = TimeSpan.FromSeconds(30);
+
+		/// <summary>
+		/// Constructor
+		/// </summary>
+		/// <param name="redis">The Redis database instance</param>
+		/// <param name="keyPrefix">Prefix for keys to use in the database</param>
+		public RedisMessageQueue(IDatabase redis, RedisKey keyPrefix)
+		{
+			_redis = redis;
+			_keyPrefix = keyPrefix;
+			_updateChannel = keyPrefix.Append("updates").ToString();
+
+			redis.Multiplexer.GetSubscriber().Subscribe(_updateChannel, OnChannelUpdate);
+		}
+
+		/// <inheritdoc/>
+		public void Dispose()
+		{
+			_redis.Multiplexer.GetSubscriber().Unsubscribe(_updateChannel, OnChannelUpdate);
+		}
+
+		/// <summary>
+		/// Callback for a message being posted to a channel
+		/// </summary>
+		/// <param name="channel"></param>
+		/// <param name="value"></param>
+		void OnChannelUpdate(RedisChannel channel, RedisValue value)
+		{
+			TaskCompletionSource<bool>? completionSource;
+			lock (_channelWakeEvents)
+			{
+				_channelWakeEvents.TryGetValue(value.ToString(), out completionSource);
+			}
+			if (completionSource != null)
+			{
+				completionSource.TrySetResult(true);
+			}
+		}
+
+		/// <summary>
+		/// Gets the set of messages for a channel
+		/// </summary>
+		/// <param name="channelId"></param>
+		/// <returns></returns>
+		RedisListKey<T> GetChannel(string channelId)
+		{
+			return new RedisListKey<T>(_keyPrefix.Append(channelId));
+		}
+
+		/// <inheritdoc/>
+		public async Task PostAsync(string channelId, T message)
+		{
+			RedisListKey<T> channel = GetChannel(channelId);
+
+			long length = await _redis.ListRightPushAsync(channel, message);
+			if (length == 1)
+			{
+				await _redis.PublishAsync(_updateChannel, channelId, CommandFlags.FireAndForget);
+			}
+			await _redis.KeyExpireAsync(channel, ExpireTime, flags: CommandFlags.FireAndForget);
+		}
+
+		async Task<bool> ReadMessagesAsync(RedisListKey<T> list, List<T> messages)
+		{
+			T? message = await _redis.ListLeftPopAsync(list);
+			while(message != null)
+			{
+				messages.Add(message);
+				message = await _redis.ListLeftPopAsync(list);
+			}
+			return messages.Count > 0;
+		}
+
+		/// <inheritdoc/>
+		public async Task<List<T>> ReadMessagesAsync(string channelId)
+		{
+			List<T> messages = new List<T>();
+			await ReadMessagesAsync(GetChannel(channelId), messages);
+			return messages;
+		}
+
+		/// <inheritdoc/>
+		public async Task<List<T>> WaitForMessagesAsync(string channelId, CancellationToken cancellationToken)
+		{
+			List<T> messages = new List<T>();
+
+			RedisListKey<T> channel = GetChannel(channelId);
+			while (!cancellationToken.IsCancellationRequested && !await ReadMessagesAsync(channel, messages))
+			{
+				// Register for notifications on this channel
+				TaskCompletionSource<bool>? completionSource;
+				lock (_channelWakeEvents)
+				{
+					if (!_channelWakeEvents.TryGetValue(channelId, out completionSource))
+					{
+						completionSource = new TaskCompletionSource<bool>();
+						_channelWakeEvents.Add(channelId, completionSource);
+					}
+				}
+
+				try
+				{
+					// Read the current queue state again, in case it was modified since we registered.
+					if(await ReadMessagesAsync(channel, messages))
+					{
+						break;
+					}
+
+					// Wait for messages to be available
+					using (cancellationToken.Register(() => completionSource.TrySetResult(false)))
+					{
+						await completionSource.Task;
+					}
+				}
+				finally
+				{
+					lock (_channelWakeEvents)
+					{
+						_channelWakeEvents.Remove(channelId);
+					}
+				}
+			}
+
+			return messages;
+		}
+	}
+}
diff --git a/Engine/Source/Programs/Horde/Horde.Build/Compute/V1/RedisTaskScheduler.cs b/Engine/Source/Programs/Horde/Horde.Build/Compute/V1/RedisTaskScheduler.cs
new file mode 100644
--- /dev/null
+++ b/Engine/Source/Programs/Horde/Horde.Build/Compute/V1/RedisTaskScheduler.cs
@@ -0,0 +1,485 @@
+// Copyright Epic Games, Inc. All Rights Reserved.
+
+using System;
+using System.Collections.Generic;
+using System.Diagnostics;
+using System.Linq;
+using System.Threading;
+using System.Threading.Channels;
+using System.Threading.Tasks;
+using EpicGames.Core;
+using EpicGames.Redis;
+using Microsoft.Extensions.Logging;
+using StackExchange.Redis;
+
+namespace Horde.Build.Compute.V1
+{
+	using Condition = StackExchange.Redis.Condition;
+
+	/// <summary>
+	/// Interface for a queue of compute operations, each of which may have different requirements for the executing machines.
+	/// </summary>
+	/// <typeparam name="TQueueId">Type used to identify a particular queue</typeparam>
+	/// <typeparam name="TTask">Type used to describe a task to be performed</typeparam>
+	interface ITaskScheduler<TQueueId, TTask> 
+		where TQueueId : notnull 
+		where TTask : class
+	{
+		/// <summary>
+		/// Adds a task to a queue
+		/// </summary>
+		/// <param name="queueId">The queue identifier</param>
+		/// <param name="taskId">The task to add</param>
+		/// <param name="atFront">Whether to insert at the front of the queue</param>
+		Task EnqueueAsync(TQueueId queueId, TTask taskId, bool atFront);
+
+		/// <summary>
+		/// Dequeue any task from a particular queue
+		/// </summary>
+		/// <param name="queueId">The queue to remove a task from</param>
+		/// <returns>Information about the task to be executed</returns>
+		Task<TTask?> DequeueAsync(TQueueId queueId);
+
+		/// <summary>
+		/// Dequeues a task that the given agent can execute
+		/// </summary>
+		/// <param name="predicate">Predicate for determining which queues can be removed from</param>
+		/// <param name="token">Cancellation token for the operation. Will return a null entry rather than throwing an exception.</param>
+		/// <returns>Information about the task to be executed</returns>
+		Task<Task<(TQueueId, TTask)?>> DequeueAsync(Func<TQueueId, ValueTask<bool>> predicate, CancellationToken token = default);
+
+		/// <summary>
+		/// Gets hashes of all the inactive task queues
+		/// </summary>
+		/// <returns></returns>
+		Task<List<TQueueId>> GetInactiveQueuesAsync();
+		
+		/// <summary>
+		/// Get number of tasks that a given pool/agent can execute
+		/// A read-only operation and will not affect any queue.
+		/// </summary>
+		/// <param name="predicate">Predicate for determining which queues to read</param>
+		/// <param name="token">Cancellation token for the operation</param>
+		/// <returns>Number of tasks matching the supplied predicate</returns>
+		Task<int> GetNumQueuedTasksAsync(Func<TQueueId, ValueTask<bool>> predicate, CancellationToken token = default);
+	}
+
+	/// <summary>
+	/// Implementation of <see cref="ITaskScheduler{TQueue, TTask}"/> using Redis for storage
+	/// </summary>
+	/// <typeparam name="TQueueId">Type used to identify a particular queue</typeparam>
+	/// <typeparam name="TTask">Type used to describe a task to be performed</typeparam>
+	class RedisTaskScheduler<TQueueId, TTask> : ITaskScheduler<TQueueId, TTask>, IDisposable 
+		where TQueueId : notnull
+		where TTask : class
+	{
+		class Listener
+		{
+			public Func<TQueueId, ValueTask<bool>> _predicate;
+			public TaskCompletionSource<(TQueueId, TTask)?> CompletionSource { get; }
+
+			public Listener(Func<TQueueId, ValueTask<bool>> predicate)
+			{
+				_predicate = predicate;
+				CompletionSource = new TaskCompletionSource<(TQueueId, TTask)?>();
+			}
+		}
+
+		readonly RedisConnectionPool _redisConnectionPool;
+		readonly RedisKey _baseKey;
+		readonly RedisSetKey<TQueueId> _queueIndex;
+		readonly RedisHashKey<TQueueId, DateTime> _activeQueues; // Queues which are actively being dequeued from
+		ReadOnlyHashSet<TQueueId> _localActiveQueues = new HashSet<TQueueId>();
+		readonly Stopwatch _resetActiveQueuesTimer = Stopwatch.StartNew();
+		readonly List<Listener> _listeners = new List<Listener>();
+		readonly RedisChannel<TQueueId> _newQueueChannel;
+		readonly Task _queueUpdateTask;
+		readonly CancellationTokenSource _cancellationSource = new CancellationTokenSource();
+		readonly ILogger _logger;
+
+		/// <summary>
+		/// Constructor
+		/// </summary>
+		/// <param name="redisConnectionPool">The Redis connection pool</param>
+		/// <param name="baseKey">Base key for all keys used by this scheduler</param>
+		/// <param name="logger"></param>
+		public RedisTaskScheduler(RedisConnectionPool redisConnectionPool, RedisKey baseKey, ILogger logger)
+		{
+			_redisConnectionPool = redisConnectionPool;
+			_baseKey = baseKey;
+			_queueIndex = new RedisSetKey<TQueueId>(baseKey.Append("index"));
+			_activeQueues = new RedisHashKey<TQueueId, DateTime>(baseKey.Append("active"));
+			_newQueueChannel = new RedisChannel<TQueueId>(baseKey.Append("new_queues").ToString());
+			_logger = logger;
+
+			_queueUpdateTask = Task.Run(() => UpdateQueuesAsync(_cancellationSource.Token));
+		}
+
+		/// <inheritdoc/>
+		public async ValueTask DisposeAsync()
+		{
+			if (!_queueUpdateTask.IsCompleted)
+			{
+				_cancellationSource.Cancel();
+				try
+				{
+					await _queueUpdateTask;
+				}
+				catch (OperationCanceledException)
+				{
+				}
+			}
+			_cancellationSource.Dispose();
+		}
+
+		/// <inheritdoc/>
+		public void Dispose()
+		{
+			DisposeAsync().AsTask().Wait();
+		}
+
+		/// <summary>
+		/// Gets the key for a list of tasks for a particular queue
+		/// </summary>
+		/// <param name="queueId">The queue identifier</param>
+		/// <returns></returns>
+		RedisListKey<TTask> GetQueueKey(TQueueId queueId)
+		{
+			return new RedisListKey<TTask>(_baseKey.Append(RedisSerializer.Serialize(queueId).AsKey()));
+		}
+
+		/// <summary>
+		/// Pushes a task onto either end of a queue
+		/// </summary>
+		static Task<long> PushTaskAsync(IDatabaseAsync database, RedisListKey<TTask> list, TTask task, When when, CommandFlags flags, bool atFront)
+		{
+			if (atFront)
+			{
+				return database.ListLeftPushAsync(list, task, when, flags);
+			}
+			else
+			{
+				return database.ListRightPushAsync(list, task, when, flags);
+			}
+		}
+
+		/// <summary>
+		/// Adds a task to a particular queue, creating and adding that queue to the index if necessary
+		/// </summary>
+		/// <param name="queueId">The queue to add the task to</param>
+		/// <param name="task">Task to be scheduled</param>
+		/// <param name="atFront">Whether to add to the front of the queue</param>
+		public async Task EnqueueAsync(TQueueId queueId, TTask task, bool atFront)
+		{
+			RedisListKey<TTask> list = GetQueueKey(queueId);
+			for (; ; )
+			{
+				IDatabase connection = _redisConnectionPool.GetDatabase();
+
+				long newLength = await PushTaskAsync(connection, list, task, When.Exists, CommandFlags.None, atFront);
+				if (newLength > 0)
+				{
+					_logger.LogInformation("Length of queue {QueueId} is {Length}", queueId, newLength);
+					break;
+				}
+
+				ITransaction transaction = connection.CreateTransaction();
+				_ = transaction.SetAddAsync(_queueIndex, queueId, CommandFlags.FireAndForget);
+				_ = PushTaskAsync(transaction, list, task, When.Always, CommandFlags.FireAndForget, atFront);
+
+				if (await transaction.ExecuteAsync())
+				{
+					_logger.LogInformation("Created queue {QueueId}", queueId);
+					await connection.PublishAsync(_newQueueChannel, queueId);
+					break;
+				}
+
+				_logger.LogDebug("EnqueueAsync() retrying...");
+			}
+		}
+
+		/// <summary>
+		/// Dequeues an item for execution by the given agent
+		/// </summary>
+		/// <param name="predicate">Predicate for queues that tasks can be removed from</param>
+		/// <param name="token">Cancellation token for waiting for an item</param>
+		/// <returns>The dequeued item, or null if no item is available</returns>
+		public async Task<Task<(TQueueId, TTask)?>> DequeueAsync(Func<TQueueId, ValueTask<bool>> predicate, CancellationToken token = default)
+		{
+			// Compare against all the list of cached queues to see if we can dequeue something from any of them
+			TQueueId[] queues = await _redisConnectionPool.GetDatabase().SetMembersAsync(_queueIndex);
+
+			// Try to dequeue an item from the list
+			(TQueueId, TTask)? entry = await TryAssignToLocalAgentAsync(queues, predicate);
+			if (entry != null)
+			{
+				return Task.FromResult(entry);
+			}
+
+			// Otherwise create a new task to do the wait
+			return WaitToDequeueAsync(queues, predicate, token);
+		}
+
+		private async Task<(TQueueId, TTask)?> WaitToDequeueAsync(TQueueId[] queues, Func<TQueueId, ValueTask<bool>> predicate, CancellationToken token)
+		{
+			Listener listener = new Listener(predicate);
+			try
+			{
+				// Add the listener to the global list
+				lock (_listeners)
+				{
+					_listeners.Add(listener);
+				}
+
+				// Try to dequeue an item from the list again. An item may have been made available 
+				(TQueueId, TTask)? entry = await TryAssignToLocalAgentAsync(queues, predicate);
+				if (entry != null)
+				{
+					if (!listener.CompletionSource.TrySetResult(entry.Value))
+					{
+						(TQueueId queue, TTask task) = entry.Value;
+						await EnqueueAsync(queue, task, true);
+					}
+				}
+
+				// Wait for an item to be available
+				using (IDisposable registration = token.Register(() => listener.CompletionSource.TrySetResult(null)))
+				{
+					return await listener.CompletionSource.Task;
+				}
+			}
+			finally
+			{
+				lock (_listeners)
+				{
+					_listeners.Remove(listener);
+				}
+			}
+		}
+
+		/// <summary>
+		/// Attempts to dequeue a task from a set of queue
+		/// </summary>
+		/// <param name="queueIds">The current array of queues</param>
+		/// <param name="predicate">Predicate for queues that tasks can be removed from</param>
+		/// <returns>The dequeued item, or null if no item is available</returns>
+		async Task<(TQueueId, TTask)?> TryAssignToLocalAgentAsync(TQueueId[] queueIds, Func<TQueueId, ValueTask<bool>> predicate)
+		{
+			foreach (TQueueId queueId in queueIds)
+			{
+				if (await predicate(queueId))
+				{
+					TTask? task = await DequeueAsync(queueId);
+					if (task != null)
+					{
+						return (queueId, task);
+					}
+				}
+			}
+			return null;
+		}
+
+		/// <summary>
+		/// Dequeues the item at the front of a queue
+		/// </summary>
+		/// <param name="queueId">The queue to dequeue from</param>
+		/// <returns>The dequeued item, or null if the queue is empty</returns>
+		public async Task<TTask?> DequeueAsync(TQueueId queueId)
+		{
+			await AddActiveQueue(queueId);
+
+			IDatabase database = _redisConnectionPool.GetDatabase();
+
+			TTask? item = await database.ListLeftPopAsync(GetQueueKey(queueId));
+			if (item == null)
+			{
+				ITransaction transaction = database.CreateTransaction();
+				transaction.AddCondition(Condition.KeyNotExists(GetQueueKey(queueId).Inner));
+				Task<bool> wasRemoved = transaction.SetRemoveAsync(_queueIndex, queueId);
+
+				if (await transaction.ExecuteAsync() && await wasRemoved)
+				{
+					_logger.LogInformation("Removed queue {QueueId} from index", queueId);
+				}
+			}
+
+			return item;
+		}
+
+		/// <summary>
+		/// Marks a queue key as being actively monitored, preventing it being returned by <see cref="GetInactiveQueuesAsync"/>
+		/// </summary>
+		/// <param name="queueId">The queue key</param>
+		/// <returns></returns>
+		async ValueTask AddActiveQueue(TQueueId queueId)
+		{
+			// Periodically clear out the set of active keys
+			TimeSpan resetTime = TimeSpan.FromSeconds(10.0);
+			if (_resetActiveQueuesTimer.Elapsed > resetTime)
+			{
+				lock (_resetActiveQueuesTimer)
+				{
+					if (_resetActiveQueuesTimer.Elapsed > resetTime)
+					{
+						_localActiveQueues = new HashSet<TQueueId>();
+						_resetActiveQueuesTimer.Restart();
+					}
+				}
+			}
+
+			// Check if the set of active keys already contains the key we're adding. In order to optimize the 
+			// common case under heavy load where the key is in the set, updating it creates a full copy of it. Any
+			// readers can thus access it without the need for any locking.
+			for(; ;)
+			{
+				ReadOnlyHashSet<TQueueId> localActiveQueuesCopy = _localActiveQueues;
+				if (localActiveQueuesCopy.Contains(queueId))
+				{
+					break;
+				}
+
+				HashSet<TQueueId> newLocalActiveQueues = new HashSet<TQueueId>(localActiveQueuesCopy);
+				newLocalActiveQueues.Add(queueId);
+
+				if (Interlocked.CompareExchange(ref _localActiveQueues, newLocalActiveQueues, localActiveQueuesCopy) == localActiveQueuesCopy)
+				{
+					_logger.LogInformation("Refreshing active queue {QueueId}", queueId);
+					await _redisConnectionPool.GetDatabase().HashSetAsync(_activeQueues, queueId, DateTime.UtcNow);
+					break;
+				}
+			}
+		}
+
+		/// <summary>
+		/// Find any inactive keys
+		/// </summary>
+		/// <returns></returns>
+		public async Task<List<TQueueId>> GetInactiveQueuesAsync()
+		{
+			HashSet<TQueueId> keys = new HashSet<TQueueId>(await _redisConnectionPool.GetDatabase().SetMembersAsync(_queueIndex));
+			HashSet<TQueueId> invalidKeys = new HashSet<TQueueId>();
+
+			DateTime minTime = DateTime.UtcNow - TimeSpan.FromMinutes(10.0);
+
+			HashEntry<TQueueId, DateTime>[] entries = await _redisConnectionPool.GetDatabase().HashGetAllAsync(_activeQueues);
+			foreach (HashEntry<TQueueId, DateTime> entry in entries)
+			{
+				if (entry.Value < minTime)
+				{
+					invalidKeys.Add(entry.Name);
+				}
+				else
+				{
+					keys.Remove(entry.Name);
+				}
+			}
+
+			if (invalidKeys.Count > 0)
+			{
+				await _redisConnectionPool.GetDatabase().HashDeleteAsync(_activeQueues, invalidKeys.ToArray());
+			}
+
+			return keys.ToList();
+		}
+
+		public async Task<int> GetNumQueuedTasksAsync(Func<TQueueId, ValueTask<bool>> predicate, CancellationToken token = default)
+		{
+			HashSet<TQueueId> queueIds = new (await _redisConnectionPool.GetDatabase().SetMembersAsync(_queueIndex));
+			long totalTaskCount = 0;
+			foreach (TQueueId queueId in queueIds)
+			{
+				if (await predicate(queueId))
+				{
+					totalTaskCount += await _redisConnectionPool.GetDatabase().ListLengthAsync(GetQueueKey(queueId));
+				}
+			}
+
+			return (int)totalTaskCount;
+		}
+
+		async Task UpdateQueuesAsync(CancellationToken cancellationToken)
+		{
+			Channel<TQueueId> newQueues = Channel.CreateUnbounded<TQueueId>();
+
+			await using RedisSubscription _ = await _redisConnectionPool.GetDatabase().Multiplexer.SubscribeAsync(_newQueueChannel, v => newQueues.Writer.TryWrite(v));
+
+			while (await newQueues.Reader.WaitToReadAsync(cancellationToken))
+			{
+				HashSet<TQueueId> newQueueIds = new HashSet<TQueueId>();
+				while (newQueues.Reader.TryRead(out TQueueId? queueId))
+				{
+					newQueueIds.Add(queueId);
+				}
+				foreach (TQueueId newQueueId in newQueueIds)
+				{
+					await TryDispatchToNewQueueAsync(newQueueId);
+				}
+			}
+		}
+
+		async Task<bool> TryDispatchToNewQueueAsync(TQueueId queueId)
+		{
+			RedisListKey<TTask> queue = GetQueueKey(queueId);
+
+			// Find a local listener that can execute the work
+			(TQueueId QueueId, TTask TaskId)? entry = null;
+			try
+			{
+				for (; ; )
+				{
+					Listener? listener = null;
+
+					// Look for a listener that can execute the task
+					HashSet<Listener> checkedListeners = new HashSet<Listener>();
+					while(listener == null)
+					{
+						// Find up to 10 listeners we haven't seen before
+						List<Listener> newListeners = new List<Listener>();
+						lock (_listeners)
+						{
+							newListeners.AddRange(_listeners.Where(x => !x.CompletionSource.Task.IsCompleted && checkedListeners.Add(x)).Take(10));
+						}
+						if (newListeners.Count == 0)
+						{
+							return false;
+						}
+
+						// Check predicates for each one against the new queue
+						foreach (Listener newListener in newListeners)
+						{
+							if (await newListener._predicate(queueId))
+							{
+								listener = newListener;
+								break;
+							}
+						}
+					}
+
+					// Pop an entry from the queue
+					if (entry == null)
+					{
+						TTask? task = await DequeueAsync(queueId);
+						if (task == null)
+						{
+							return false;
+						}
+						entry = (queueId, task);
+					}
+
+					// Assign it to the listener
+					if (listener.CompletionSource.TrySetResult(entry.Value))
+					{
+						entry = null;
+					}
+				}
+			}
+			finally
+			{
+				if (entry != null)
+				{
+					await EnqueueAsync(entry.Value.QueueId, entry.Value.TaskId, true);
+				}
+			}
+		}
+	}
+}
diff --git a/Engine/Source/Programs/Horde/Horde.Build/Startup.cs b/Engine/Source/Programs/Horde/Horde.Build/Startup.cs
--- a/Engine/Source/Programs/Horde/Horde.Build/Startup.cs
+++ b/Engine/Source/Programs/Horde/Horde.Build/Startup.cs
@@ -31,7 +31,6 @@ using Horde.Build.Agents.Pools;
 using Horde.Build.Agents.Sessions;
 using Horde.Build.Agents.Software;
 using Horde.Build.Agents.Telemetry;
-using Horde.Build.Compute;
 using Horde.Build.Configuration;
 using Horde.Build.Devices;
 using Horde.Build.Issues;
@@ -676,16 +675,16 @@ namespace Horde.Build
 			services.AddHostedService<JobTaskSource>(provider => provider.GetRequiredService<JobTaskSource>());
 			services.AddSingleton<ConformTaskSource>();
 			services.AddHostedService<ConformTaskSource>(provider => provider.GetRequiredService<ConformTaskSource>());
-			services.AddSingleton<ComputeService>();
-			services.AddSingleton<IComputeService, ComputeService>();
-			services.AddHostedService(provider => provider.GetRequiredService<ComputeService>());
+			services.AddSingleton<Compute.V1.ComputeService>();
+			services.AddSingleton<Compute.V1.IComputeService, Compute.V1.ComputeService>();
+			services.AddHostedService(provider => provider.GetRequiredService<Compute.V1.ComputeService>());
 
 			services.AddSingleton<ITaskSource, UpgradeTaskSource>();
 			services.AddSingleton<ITaskSource, ShutdownTaskSource>();
 			services.AddSingleton<ITaskSource, RestartTaskSource>();
 			services.AddSingleton<ITaskSource, ConformTaskSource>(provider => provider.GetRequiredService<ConformTaskSource>());
 			services.AddSingleton<ITaskSource, JobTaskSource>(provider => provider.GetRequiredService<JobTaskSource>());
-			services.AddSingleton<ITaskSource, ComputeService>();
+			services.AddSingleton<ITaskSource, Compute.V1.ComputeService>();
 
 			services.AddHostedService(provider => provider.GetRequiredService<ConformTaskSource>());
 
