Horde: Updating dashboard dependencies to latest versions and cleanup
#rnx
#skipci

diff --git a/Engine/Source/Programs/Horde/HordeDashboard/package.json b/Engine/Source/Programs/Horde/HordeDashboard/package.json
--- a/Engine/Source/Programs/Horde/HordeDashboard/package.json
+++ b/Engine/Source/Programs/Horde/HordeDashboard/package.json
@@ -6,26 +6,25 @@
   "proxy": "https://localhost:5001",
   "dependencies": {
     "@blueprintjs/icons": "^3.13.0",
-    "@datadog/browser-logs": "^1.12.4",
-    "@fluentui/react": "^8.36.5",
-    "@fluentui/react-charting": "^5.2.0",
-    "crypto-js": "^4.0.0",
+    "@datadog/browser-logs": "^4.30.0",
+    "@fluentui/react": "^8.104.5",
+    "@fluentui/react-charting": "^5.14.36",
+    "crypto-js": "^4.1.1",
     "d3": "^7.6.1",
-    "formik": "^2.1.1",
-    "idb": "^7.0.2",
-    "markdown-to-jsx": "^6.10.3",
-    "mobx": "^5.15.1",
-    "mobx-react-lite": "^1.5.2",
-    "moment-timezone": "^0.5.28",
-    "react": "^16.12.0",
+    "idb": "^7.1.1",
+    "markdown-to-jsx": "^7.1.8",
+    "mobx": "^6.7.0",
+    "mobx-react-lite": "^3.4.0",
+    "moment-timezone": "^0.5.40",
+    "react": "^18.2.0",
     "react-custom-scrollbars": "^4.2.1",
     "react-darkreader": "^1.5.6",
-    "react-dom": "^17.0.1",
+    "react-dom": "^18.2.0",
     "react-flame-graph": "^1.4.0",
     "react-highlighter": "^0.4.3",
-    "react-router-dom": "5.2.0",
-    "react-router-hash-link": "^2.3.1",
-    "react-scripts": "^4.0.1",
+    "react-router-dom": "6.6.2",
+    "react-router-hash-link": "^2.4.3",
+    "react-scripts": "^5.0.1",
     "react-text-marquee": "^1.0.1"
   },
   "scripts": {
@@ -49,14 +48,14 @@
     ]
   },
   "devDependencies": {
-    "@types/crypto-js": "^4.0.1",
+    "@types/crypto-js": "^4.1.1",
     "@types/d3": "^7.4.0",
-    "@types/markdown-to-jsx": "6.9.1",
-    "@types/react": "^17.0.30",
-    "@types/react-custom-scrollbars": "^4.0.6",
-    "@types/react-dom": "^17.0.0",
-    "@types/react-highlighter": "^0.3.4",
-    "@types/react-router-dom": "^5.1.6",
+    "@types/markdown-to-jsx": "7.0.1",
+    "@types/react": "^18.0.26",
+    "@types/react-custom-scrollbars": "^4.0.10",
+    "@types/react-dom": "^18.0.10",
+    "@types/react-highlighter": "^0.3.5",
+    "@types/react-router-dom": "^5.3.3",
     "rimraf": "^3.0.0",
     "typescript": "^4.7.4"
   }
diff --git a/Engine/Source/Programs/Horde/HordeDashboard/src/App.tsx b/Engine/Source/Programs/Horde/HordeDashboard/src/App.tsx
--- a/Engine/Source/Programs/Horde/HordeDashboard/src/App.tsx
+++ b/Engine/Source/Programs/Horde/HordeDashboard/src/App.tsx
@@ -3,7 +3,7 @@
 import { Image, Spinner, SpinnerSize, Stack, Text } from '@fluentui/react';
 import React, { useState } from 'react';
 import { useDarkreader } from 'react-darkreader';
-import { BrowserRouter, Redirect, Route, Switch } from 'react-router-dom';
+import { BrowserRouter, Navigate, Route, Routes } from 'react-router-dom';
 import backend from './backend';
 import { DashboardPreference } from './backend/Api';
 import { getSiteConfig } from './backend/Config';
@@ -54,7 +54,7 @@ const Main: React.FC = () => {
 
             let local = localStorage.getItem("horde_darktheme");
 
-            if (!local) {               
+            if (!local) {
 
                console.log("Setting local theme to ", darktheme);
 
@@ -67,13 +67,13 @@ const Main: React.FC = () => {
                   setInit(true);
                   return null;
                }
-            
+
             } else if (local !== darktheme) {
 
                console.log(`Setting local theme to ${darktheme} and reloading for change`);
                localStorage.setItem("horde_darktheme", darktheme);
                window.location.reload();
-               
+
             } else {
                setInit(true);
                return null;
@@ -116,34 +116,33 @@ const Main: React.FC = () => {
    }
 
    return (
-      <Switch>
-         <Route path="/index" component={UserHomeView} />
-         <Route path="/project/:projectId" component={ProjectHome} />
-         <Route path="/pool/:poolId" component={PoolView} />
-         <Route path="/job/:jobId" component={JobDetailViewV2} />         
-         <Route path="/log/:logId" component={LogView} />
-         <Route path="/testreport/:testdataId" component={TestReportView} />
-         <Route path="/stream/:streamId" component={StreamView} />
-         <Route path="/agents/:agentId?" component={AgentView} />
-         <Route path="/admin/token" component={AdminToken} />
-         <Route path="/reports/utilization" component={UtilizationReportView} />
-         <Route path="/preflight" component={PreflightRedirector} />
-         <Route path="/dashboard" component={DashboardView} />
-         <Route path="/perforce/servers" component={PerforceServerView} />
-         <Route path="/notices" component={NoticeView} />
-         <Route path="/devices" component={DeviceView} />
-         <Route path="/server/settings" component={ServerSettingsView} />
-         <Route path="/audit/agent/:agentId" component={AuditLogView} />
-         <Route path="/audit/issue/:issueId" component={AuditLogView} />
-         <Route path="/automation" component={AutomationView} />
-         <Route path={["/debug/lease/:leaseId"]} component={DebugView} />
+      <Routes>
+         <Route path="/index" element={<UserHomeView />} />
+         <Route path="/project/:projectId" element={<ProjectHome />} />
+         <Route path="/pool/:poolId" element={<PoolView />} />
+         <Route path="/job/:jobId" element={<JobDetailViewV2 />} />
+         <Route path="/log/:logId" element={<LogView />} />
+         <Route path="/testreport/:testdataId" element={<TestReportView />} />
+         <Route path="/stream/:streamId" element={<StreamView />} />
+         <Route path="/agents/:agentId?" element={<AgentView />} />
+         <Route path="/admin/token" element={<AdminToken />} />
+         <Route path="/reports/utilization" element={<UtilizationReportView />} />
+         <Route path="/preflight" element={<PreflightRedirector />} />
+         <Route path="/dashboard" element={<DashboardView />} />
+         <Route path="/perforce/servers" element={<PerforceServerView />} />
+         <Route path="/notices" element={<NoticeView />} />
+         <Route path="/devices" element={<DeviceView />} />
+         <Route path="/server/settings" element={<ServerSettingsView />} />
+         <Route path="/audit/agent/:agentId" element={<AuditLogView />} />
+         <Route path="/audit/issue/:issueId" element={<AuditLogView />} />
+         <Route path="/automation" element={<AutomationView />} />
+         <Route path="/debug/lease/:leaseId" element={<DebugView />} />
          {hordePlugins.routes.map((route, index) => {
-            return <Route key={`key_plugin_route_${index}`} path={route.path} component={route.component} />;
+            return <Route key={`key_plugin_route_${index}`} path={route.path} element={<route.component />} />;
          })}
-         <Redirect from='/' to='/index' />
-      </Switch>
+         <Route path='*' element={<Navigate to='/index' replace={true} />} />
+      </Routes>
    );
-
 };
 
 const Darkmode: React.FC = () => {
@@ -155,7 +154,7 @@ const Darkmode: React.FC = () => {
    `;
 
    // NOTE: if an Stack child isn't respecting className="horde-no-darktheme", check that you are using style: {} instead of styles:{root:{}}!
-   /*const [isDark, { toggle }] = */useDarkreader(dashboard.darktheme, { brightness: 100, contrast: 100, sepia: 0, grayscale: 0, darkSchemeTextColor: "#FFFFFFFF" }, { invert: [], ignoreInlineStyle: ['.horde-no-darktheme *'], css: additionalCSS, ignoreImageAnalysis: [] });
+   /*const [isDark, { toggle }] = */useDarkreader(dashboard.darktheme, { brightness: 100, contrast: 100, sepia: 0, grayscale: 0, darkSchemeTextColor: "#FFFFFFFF" }, { disableStyleSheetsProxy: false, invert: [], ignoreInlineStyle: ['.horde-no-darktheme *'], css: additionalCSS, ignoreImageAnalysis: [] });
 
    return null;
 };
