Horde: Updating dashboard dependencies to latest versions and cleanup

diff --git a/Engine/Source/Programs/Horde/HordeDashboard/src/Plugins/AutomatedTestSession/index.tsx b/Engine/Source/Programs/Horde/HordeDashboard/src/Plugins/AutomatedTestSession/index.tsx
--- a/Engine/Source/Programs/Horde/HordeDashboard/src/Plugins/AutomatedTestSession/index.tsx
+++ b/Engine/Source/Programs/Horde/HordeDashboard/src/Plugins/AutomatedTestSession/index.tsx
@@ -25,5 +25,6 @@ class AutomatedTestSessionPlugin implements Plugin {
     ];
 }
 
-export default new AutomatedTestSessionPlugin();
+const automatedTestSessionPlugin = new AutomatedTestSessionPlugin();
+export default automatedTestSessionPlugin;
 
diff --git a/Engine/Source/Programs/Horde/HordeDashboard/src/Plugins/AutomatedTestSession/models/AutomatedTestSession.ts b/Engine/Source/Programs/Horde/HordeDashboard/src/Plugins/AutomatedTestSession/models/AutomatedTestSession.ts
--- a/Engine/Source/Programs/Horde/HordeDashboard/src/Plugins/AutomatedTestSession/models/AutomatedTestSession.ts
+++ b/Engine/Source/Programs/Horde/HordeDashboard/src/Plugins/AutomatedTestSession/models/AutomatedTestSession.ts
@@ -1,6 +1,6 @@
 // Copyright Epic Games, Inc. All Rights Reserved.
 
-import { observable, action, when, reaction } from "mobx";
+import { observable, action, when, reaction, makeObservable } from "mobx";
 import { DataWrapper, TestDataWrapper, TestDataCollection } from '../../../backend/TestDataHandler'
 
 export enum TestState {
@@ -13,7 +13,7 @@ export enum TestState {
     Unknown = "Unknown",
 }
 
-const TestStateRank: {[Key in string]: number}  = {
+const TestStateRank: { [Key in string]: number } = {
     "Fail": 0,
     "InProcess": 1,
     "NotRun": 2,
@@ -71,7 +71,7 @@ export type TestResultDetails = {
     Events: TestEvent[];
 }
 
-export type TestDevice = MetadataHolder& {
+export type TestDevice = MetadataHolder & {
     Name: string;
     AppInstanceName: string;
     AppInstanceLog: string;
@@ -90,7 +90,7 @@ export type IndexedError = {
     TestUIDs: string[];
 }
 
-export type Metadata = {[Key in string]: string|undefined}
+export type Metadata = { [Key in string]: string | undefined }
 
 type MetadataHolder = {
     Metadata?: Metadata
@@ -98,7 +98,7 @@ type MetadataHolder = {
 
 export type Loader = {
     loadingProgress: number;
-    whenLoadUpdate: () => Promise<void> & {cancel: () => void};
+    whenLoadUpdate: () => Promise<void> & { cancel: () => void };
 }
 
 const generatetHashFromString = (aggregate: string) => {
@@ -113,7 +113,7 @@ const generatetHashFromString = (aggregate: string) => {
 
 export const generateHashFromMetadata = (obj: Metadata) => {
     const keys = Object.keys(obj).sort();
-    const aggregate = keys.map((key) => key + (obj[key]??'')).join();
+    const aggregate = keys.map((key) => key + (obj[key] ?? '')).join();
     return generatetHashFromString(aggregate);
 }
 
@@ -126,15 +126,15 @@ export class MetaWrapper {
         this._hash = "";
     }
 
-    get(key: string) : string | undefined {
-        if(this.collection) {
+    get(key: string): string | undefined {
+        if (this.collection) {
             return this.collection[key];
         }
 
         return undefined;
     }
 
-    map(mapFunc: (key: string, value?: string) => any) : any[] {
+    map(mapFunc: (key: string, value?: string) => any): any[] {
         if (this.collection) {
             return Object.keys(this.collection).map((key) => mapFunc(key, this.collection![key]));
         }
@@ -148,31 +148,31 @@ export class MetaWrapper {
     }
 
     maskByKeys(keyMask: string[]) {
-        return this.map((key, value) => keyMask.includes(key)?undefined:value).filter((info) => info)
+        return this.map((key, value) => keyMask.includes(key) ? undefined : value).filter((info) => info)
     }
 
     private generateHashCode(): string {
-        if(this.collection) {
+        if (this.collection) {
             return generateHashFromMetadata(this.collection);
         }
         return "";
     }
 
     get hash(): string {
-        if(!this._hash) {
+        if (!this._hash) {
             this._hash = this.generateHashCode();
         }
         return this._hash;
     }
 }
 
-export type IndexedErrorMap = {[Key in string]: IndexedError}
+export type IndexedErrorMap = { [Key in string]: IndexedError }
 
-export type TestResultDetailsMap = {[Key in string]: TestResultDetails}
+export type TestResultDetailsMap = { [Key in string]: TestResultDetails }
 
-export type TestResultsMap = {[Key in string]: TestResult}
+export type TestResultsMap = { [Key in string]: TestResult }
 
-export type TestSession = MetadataHolder& {
+export type TestSession = MetadataHolder & {
     Name: string;
     Type: string;
     PreFlightChange: string;
@@ -181,7 +181,7 @@ export type TestSession = MetadataHolder& {
     Devices: TestDevice[];
 }
 
-type CheckList = {[Key in string]: boolean}
+type CheckList = { [Key in string]: boolean }
 
 export enum SessionFields {
     initiate = "data.Metadata,data.Name,data.PreFlightChange",
@@ -226,6 +226,7 @@ export class TestSessionWrapper implements TestSession, DataWrapper {
     private _testsession: TestSession;
 
     constructor(testdata: TestDataWrapper) {
+        makeObservable(this);
         this._testsession = testdata.data;
         this.Testdata = testdata;
         // first initialization, we should call the post Properties Update
@@ -237,17 +238,17 @@ export class TestSessionWrapper implements TestSession, DataWrapper {
     }
 
     getUID(): string {
-        return this.Name + (this.MetaHandler?.hash??'');
+        return this.Name + (this.MetaHandler?.hash ?? '');
     }
 
     postPropertiesUpdate() {
         this.capitalizeKeys()
         const data = this._testsession;
-        if(!this.MetaHandler && data.Metadata) {
+        if (!this.MetaHandler && data.Metadata) {
             this.MetaHandler = new MetaWrapper(data.Metadata);
         }
-        if(data.TestSessionInfo && !this.isPropertyLoaded('TestSessionInfo')) {
-            for(const key in data.TestSessionInfo.Tests) {
+        if (data.TestSessionInfo && !this.isPropertyLoaded('TestSessionInfo')) {
+            for (const key in data.TestSessionInfo.Tests) {
                 const test = data.TestSessionInfo.Tests[key];
                 test.Session = this;
             }
@@ -256,7 +257,7 @@ export class TestSessionWrapper implements TestSession, DataWrapper {
     }
 
     getTestResultByUID(uid: string) {
-        if(this.TestSessionInfo) {
+        if (this.TestSessionInfo) {
             return this.TestSessionInfo.Tests[uid];
         }
 
@@ -264,11 +265,11 @@ export class TestSessionWrapper implements TestSession, DataWrapper {
     }
 
     getDevices(AppInstanceName: string | string[]) {
-        if(this.Devices) {
-            if (typeof(AppInstanceName) === 'string') {
+        if (this.Devices) {
+            if (typeof (AppInstanceName) === 'string') {
                 const asString = AppInstanceName as string;
                 const device = this.Devices.find((item) => item.AppInstanceName === asString);
-                if(device) {
+                if (device) {
                     return [device];
                 }
             }
@@ -287,7 +288,7 @@ export class TestSessionWrapper implements TestSession, DataWrapper {
         for (var key in data) {
             const first = key.charAt(0);
             const firstUpper = first.toUpperCase();
-            if(first !== firstUpper) {
+            if (first !== firstUpper) {
                 const value = data[key];
                 delete data[key];
                 data[firstUpper + key.slice(1)] = value;
@@ -305,26 +306,26 @@ export class TestStats {
     TotalRun: number = 0;
 
     update(state: string, increment: boolean = true) {
-        switch(state) {
-        case TestState.Failed:
-            increment? this.Failed++ : this.Failed--;
-            increment? this.TotalRun++ : this.TotalRun--
-            break;
-        case TestState.NotRun:
-            increment? this.Unexecuted++ : this.Unexecuted--;
-            break;
-        case TestState.Success:
-        case TestState.SuccessWithWarnings:
-            increment? this.Passed++ : this.Passed--;
-            increment? this.TotalRun++ : this.TotalRun--
-            break;
-        case TestState.InProcess:
-            increment? this.Incomplete++ : this.Incomplete--;
-            increment? this.TotalRun++ : this.TotalRun--
-            break;
-        case TestState.Skipped:
-            increment? this.Skipped++ : this.Skipped--;
-            break;
+        switch (state) {
+            case TestState.Failed:
+                increment ? this.Failed++ : this.Failed--;
+                increment ? this.TotalRun++ : this.TotalRun--
+                break;
+            case TestState.NotRun:
+                increment ? this.Unexecuted++ : this.Unexecuted--;
+                break;
+            case TestState.Success:
+            case TestState.SuccessWithWarnings:
+                increment ? this.Passed++ : this.Passed--;
+                increment ? this.TotalRun++ : this.TotalRun--
+                break;
+            case TestState.InProcess:
+                increment ? this.Incomplete++ : this.Incomplete--;
+                increment ? this.TotalRun++ : this.TotalRun--
+                break;
+            case TestState.Skipped:
+                increment ? this.Skipped++ : this.Skipped--;
+                break;
         }
     }
 
@@ -379,15 +380,16 @@ export class TestCase implements Loader {
     incrementHistoryTarget() {
         this.historyTarget++;
     }
-    get loadingProgress() {return this.history}
+    get loadingProgress() { return this.history }
     whenLoadUpdate() {
         const previousVersion = this.history;
         return when(() => this.version !== previousVersion);
     }
 
     constructor(name: string, testuid: string, suite: string, sessionName: string, displayName?: string) {
+        makeObservable(this);
         this.Name = name;
-        this.DisplayName = displayName??this.toDislayName(name);
+        this.DisplayName = displayName ?? this.toDislayName(name);
         this.TestUID = testuid;
         this.State = TestState.Unknown;
         this.StateMetaUID = undefined;
@@ -415,7 +417,7 @@ export class TestCase implements Loader {
 
     addMeta(item: TestResult) {
         const metauid = item.Session.MetaHandler!.hash;
-        if(this.Metas.has(metauid)) {
+        if (this.Metas.has(metauid)) {
             return false;
         }
         this.Metas.set(metauid, item);
@@ -424,7 +426,7 @@ export class TestCase implements Loader {
             state = TestState.SuccessWithWarnings;
             item.State = state;
         }
-        if(this.setState(state)) {
+        if (this.setState(state)) {
             this.StateMetaUID = metauid;
         }
         this.Stats.update(state);
@@ -451,7 +453,7 @@ export class TestCase implements Loader {
                 return sessionKey && metafilter[key].includes(sessionKey);
             })
         );
-        if(metas.length === 0) {
+        if (metas.length === 0) {
             return undefined;
         }
         const testcase = new TestCase(this.Name, this.TestUID, this.Suite, this.SessionName, this.DisplayName);
@@ -464,7 +466,7 @@ export class TestCase implements Loader {
         const metas = Array.from(this.Metas.values()).filter(
             (test) => targetRank <= TestStateRank[test.State]
         );
-        if(metas.length === 0) {
+        if (metas.length === 0) {
             return undefined;
         }
         const testcase = new TestCase(this.Name, this.TestUID, this.Suite, this.SessionName, this.DisplayName);
@@ -474,17 +476,17 @@ export class TestCase implements Loader {
 
     private toDislayName(name: string): string {
         // It is necessary to sanitize the test names that contains / instead of . syntax
-        if(name.indexOf('/') > -1) {
+        if (name.indexOf('/') > -1) {
             return name.replaceAll(/\.*\//g, ".");
         }
         return name;
     }
 
     private setState(state: string) {
-        if(this.State !== state) {
+        if (this.State !== state) {
             const currentRank = TestStateRank[this.State];
             const newRank = TestStateRank[state];
-            if(newRank < currentRank) {
+            if (newRank < currentRank) {
                 this.State = state;
                 return true;
             }
@@ -495,9 +497,9 @@ export class TestCase implements Loader {
     private sortResults(a: TestResult, b: TestResult) {
         const aRank = TestStateRank[a.State];
         const bRank = TestStateRank[b.State];
-        if(aRank > bRank) {
+        if (aRank > bRank) {
             return 1;
-        } else if(aRank < bRank) {
+        } else if (aRank < bRank) {
             return -1;
         }
         return 0;
@@ -523,7 +525,7 @@ export class Section {
 
     addTestCase(testcase: TestCase, testName?: string) {
         this.TestCases.push(testcase);
-        this.TestNames.add(testName??testcase.DisplayName);
+        this.TestNames.add(testName ?? testcase.DisplayName);
         this.Stats.updateFromStats(testcase.Stats);
     }
 
@@ -544,17 +546,17 @@ export enum FilterType {
     meta = 'meta',
 }
 
-export type MetadataFilter = {[Key in string]: string[]}
+export type MetadataFilter = { [Key in string]: string[] }
 export class MetaFilterTools {
-    static Match = (key: string, target: string) : boolean => target.startsWith(`${key}.`)
-    static Encode = (key: string, items: string[]) : string => `${key}.${items.join('~')}`
-    static Decode = (items: string[]) : MetadataFilter | undefined => {
+    static Match = (key: string, target: string): boolean => target.startsWith(`${key}.`)
+    static Encode = (key: string, items: string[]): string => `${key}.${items.join('~')}`
+    static Decode = (items: string[]): MetadataFilter | undefined => {
         return items.length === 0 ? undefined :
             Object.fromEntries(items.map((item) => {
                 const [key, value] = item.split('.', 2);
                 return [key, value.split('~')]
             }))
-    }    
+    }
 }
 
 export type SectionFilter = {
@@ -573,15 +575,15 @@ export type ReportFilter = {
 }
 
 export const generateHashFromSectionFilter = (obj: SectionFilter) => {
-    const baseFilter : Metadata = Object.fromEntries(
+    const baseFilter: Metadata = Object.fromEntries(
         Object.keys(obj).filter((key) => key !== 'meta').map((key) => [
             key, Object.getOwnPropertyDescriptor(obj, key)?.value
         ])
     );
-    let aggregate = Object.keys(baseFilter).sort().map((key) => key + (baseFilter[key]??'')).join();
+    let aggregate = Object.keys(baseFilter).sort().map((key) => key + (baseFilter[key] ?? '')).join();
     const subObj = obj.meta;
-    if(subObj !== undefined) {
-        aggregate += 'meta'+Object.keys(subObj).sort().map((key) => key + (subObj[key]?.join('~')??'')).join();
+    if (subObj !== undefined) {
+        aggregate += 'meta' + Object.keys(subObj).sort().map((key) => key + (subObj[key]?.join('~') ?? '')).join();
     }
     return generatetHashFromString(aggregate);
 }
@@ -591,8 +593,8 @@ const regexNameFilter = /"([^"]+)"|([^\s"]+)/g;
 
 export class SectionCollection {
 
-    private filterCache: Map<string, {hash: string, sections: Section[]}>;
-    private onTheFlyCache: {filter: string, hash: string, sections: Section[]} | undefined;
+    private filterCache: Map<string, { hash: string, sections: Section[] }>;
+    private onTheFlyCache: { filter: string, hash: string, sections: Section[] } | undefined;
 
     constructor() {
         this.filterCache = new Map();
@@ -603,24 +605,24 @@ export class SectionCollection {
         const filterCacheKey = generateHashFromSectionFilter(filter);
         const cache = this.filterCache.get(filterCacheKey);
         const isSameHash = cache && cache.hash === sessions.hash;
-        const sections = cache && isSameHash? cache.sections : this.filterSections(filter, sessions.getTestCases());
+        const sections = cache && isSameHash ? cache.sections : this.filterSections(filter, sessions.getTestCases());
         if (!this.filterCache.has(filterCacheKey) || !isSameHash) {
-            this.filterCache.set(filterCacheKey, {hash: sessions.hash, sections:sections})
+            this.filterCache.set(filterCacheKey, { hash: sessions.hash, sections: sections })
         }
         return sections;
     }
 
     reExtractSectionOnTheFly(filter: SectionFilter, testNameFilter: string, sessions: TestSessionCollection) {
         // to use only when iterating through cached data
-        let sections = this.extractSections({...filter, namecontains: undefined}, sessions);
-        const onTheFlyFilter = {...filter, namecontains: testNameFilter.toLowerCase()}
+        let sections = this.extractSections({ ...filter, namecontains: undefined }, sessions);
+        const onTheFlyFilter = { ...filter, namecontains: testNameFilter.toLowerCase() }
         sections = this.filterSections(onTheFlyFilter, sections.flatMap((section) => section.TestCases));
-        this.onTheFlyCache = {filter: generateHashFromSectionFilter(onTheFlyFilter), hash: sessions.hash, sections};
+        this.onTheFlyCache = { filter: generateHashFromSectionFilter(onTheFlyFilter), hash: sessions.hash, sections };
         return sections;
     }
 
     promoteOnTheFlyCache() {
-        if(this.onTheFlyCache) {
+        if (this.onTheFlyCache) {
             this.filterCache.set(this.onTheFlyCache.filter, this.onTheFlyCache);
             this.onTheFlyCache = undefined;
         }
@@ -628,28 +630,28 @@ export class SectionCollection {
 
     private filterSections(filter: SectionFilter, tests: TestCase[]) {
         const filterFunc = (test: TestCase) => {
-            if(filter.suite) {
-                if(test.Suite !== filter.suite) {
+            if (filter.suite) {
+                if (test.Suite !== filter.suite) {
                     return false;
                 }
             }
-            if(filter.namecontains) {
+            if (filter.namecontains) {
                 const substrings = Array.from(filter.namecontains.matchAll(regexNameFilter), match => match[1] || match[2]);
                 const testName = test.Name.toLowerCase();
-                if(!substrings.some((value) => testName.indexOf(value) >= 0)) {
+                if (!substrings.some((value) => testName.indexOf(value) >= 0)) {
                     return false;
                 }
             }
-            if(filter.name) {
-                if(test.DisplayName.indexOf(filter.name + '.') !== 0) {
+            if (filter.name) {
+                if (test.DisplayName.indexOf(filter.name + '.') !== 0) {
                     return false;
                 }
             }
             return true;
         }
         const filteredTests = tests.reduce((filtered: TestCase[], item: TestCase) => {
-            if(filterFunc(item)) {
-                if(filter.meta === undefined) {
+            if (filterFunc(item)) {
+                if (filter.meta === undefined) {
                     filtered.push(item);
                 } else {
                     const filteredCase = item.getFilteredByMeta(filter.meta);
@@ -661,9 +663,9 @@ export class SectionCollection {
             return filtered
         }, []);
 
-        const commonSection = filter.name? this.getCommonSection(filteredTests.map((item) => item.DisplayName)) : "";
+        const commonSection = filter.name ? this.getCommonSection(filteredTests.map((item) => item.DisplayName)) : "";
         const sectionKeyFunc = (key: FilterType): GetKeyFunction | undefined => {
-            switch(key) {
+            switch (key) {
                 case FilterType.suite:
                     return (testcase: TestCase) => {
                         return [testcase.Suite, testcase.Suite];
@@ -671,11 +673,11 @@ export class SectionCollection {
                 case FilterType.name:
                     return (testcase: TestCase) => {
                         const suite = testcase.DisplayName.slice(commonSection.length).split('.', 1)[0];
-                        return [suite, commonSection+suite];
+                        return [suite, commonSection + suite];
                     }
             }
         };
-        const sectionFocus = filter.suite || filter.name? FilterType.name : FilterType.suite;
+        const sectionFocus = filter.suite || filter.name ? FilterType.name : FilterType.suite;
         const sections = this.buildSections(filteredTests, sectionFocus, sectionKeyFunc(sectionFocus)!);
         // sorting
         sections.forEach((section) => this.sortSectionResults(section));
@@ -684,13 +686,13 @@ export class SectionCollection {
     }
 
     private buildSections(testcases: TestCase[], type: FilterType, keyFunc: GetKeyFunction): Section[] {
-        const sections : Section[] = [];
-        const sectionsMap = new Map<string, Section>(); 
+        const sections: Section[] = [];
+        const sectionsMap = new Map<string, Section>();
 
         testcases.forEach((test) => {
-            const [ key, fullKey ] = keyFunc(test);
+            const [key, fullKey] = keyFunc(test);
             let targetSection = sectionsMap.get(key);
-            if(!targetSection) {
+            if (!targetSection) {
                 targetSection = new Section(key, fullKey, type);
                 sections.push(targetSection);
                 sectionsMap.set(key, targetSection);
@@ -702,22 +704,21 @@ export class SectionCollection {
     }
 
     private getCommonString(list: string[]): string {
-        if(list.length === 0) {
+        if (list.length === 0) {
             return "";
         }
         let shortest: string = list[0];
         let longest: string = "";
         list.forEach((item) => {
-            if(item > longest) {
+            if (item > longest) {
                 longest = item;
                 return;
             }
-            if(item < shortest)
-            {
+            if (item < shortest) {
                 shortest = item;
             }
         });
-        while(longest.indexOf(shortest) !== 0) {
+        while (longest.indexOf(shortest) !== 0) {
             shortest = shortest.slice(0, -1)
         }
 
@@ -726,7 +727,7 @@ export class SectionCollection {
 
     private getCommonSection(list: string[]): string {
         let commonKey = this.getCommonString(list);
-        if(commonKey[commonKey.length - 1] !== '.' && commonKey.indexOf('.') >= 0) {
+        if (commonKey[commonKey.length - 1] !== '.' && commonKey.indexOf('.') >= 0) {
             commonKey = commonKey.slice(0, commonKey.lastIndexOf('.') + 1);
         }
 
@@ -736,8 +737,8 @@ export class SectionCollection {
     // sort comparing the number of Failed, Incomplete, Unexecuted and Passed
     private sortSections(sections: Section[]) {
         sections.sort((a, b) => {
-            const aStats = a.isOneTest? a.TestCases[0].Stats : a.Stats;
-            const bStats = b.isOneTest? b.TestCases[0].Stats : b.Stats;
+            const aStats = a.isOneTest ? a.TestCases[0].Stats : a.Stats;
+            const bStats = b.isOneTest ? b.TestCases[0].Stats : b.Stats;
             if (aStats.Failed > bStats.Failed) {
                 return -1;
             } else if (aStats.Failed < bStats.Failed) {
@@ -772,9 +773,9 @@ export class SectionCollection {
         testcases.sort((a, b) => {
             const aRank = TestStateRank[a.State];
             const bRank = TestStateRank[b.State];
-            if(aRank > bRank) {
+            if (aRank > bRank) {
                 return 1;
-            } else if(aRank < bRank) {
+            } else if (aRank < bRank) {
                 return -1;
             }
 
@@ -823,13 +824,13 @@ export class TestSessionCollection implements Loader {
         const previousCount = this.loadingCount;
         return when(() => this.loadingCount !== previousCount)
     }
-    get loadingProgress(): number { return this._datahandler.items?.size??0}
+    get loadingProgress(): number { return this._datahandler.items?.size ?? 0 }
 
     get size(): number { return this._collection.size; }
     get hash(): string {
         let aggregate = "";
         this.forEachSession( // only include fully loaded session
-            (session) => aggregate += session.isAllPropertiesLoaded()?session.getUID():""
+            (session) => aggregate += session.isAllPropertiesLoaded() ? session.getUID() : ""
         );
         return generatetHashFromString(aggregate);
     }
@@ -838,9 +839,10 @@ export class TestSessionCollection implements Loader {
     private _testcases: Map<string, TestCase>;
     private _datahandler: TestDataCollection;
     private _includePreflight: boolean = false;
-    private _cachedMetaValueByKey: {hash: string, value: Map<string, Set<string>>} | undefined;
+    private _cachedMetaValueByKey: { hash: string, value: Map<string, Set<string>> } | undefined;
 
     constructor(handler: TestDataCollection, includePreflight: boolean = false) {
+        makeObservable(this);
         this._collection = new Map();
         this._testcases = new Map();
         this._datahandler = handler;
@@ -849,7 +851,7 @@ export class TestSessionCollection implements Loader {
 
     set(uid: string, item: TestSessionWrapper) {
         this._collection.set(uid, item);
-        if(item.isAllPropertiesLoaded()) {
+        if (item.isAllPropertiesLoaded()) {
             this.indexTestsFromSession(item);
             this.updateVersion();
         }
@@ -867,10 +869,10 @@ export class TestSessionCollection implements Loader {
         const session = this._collection.get(uid);
         this._collection.delete(uid);
         // clean up testuids map
-        if(session && session.TestSessionInfo) {
+        if (session && session.TestSessionInfo) {
             Object.keys(session.TestSessionInfo.Tests).forEach((key) => {
                 const testcase = this._testcases.get(key)!;
-                if(testcase.removeMeta(session.MetaHandler!)) {
+                if (testcase.removeMeta(session.MetaHandler!)) {
                     if (testcase.isEmpty()) {
                         this._testcases.delete(key)
                     }
@@ -891,15 +893,15 @@ export class TestSessionCollection implements Loader {
     }
 
     getMetaValueByKey() {
-        if(this._cachedMetaValueByKey !== undefined) {
-            if(this._cachedMetaValueByKey.hash === this.hash) {
+        if (this._cachedMetaValueByKey !== undefined) {
+            if (this._cachedMetaValueByKey.hash === this.hash) {
                 return this._cachedMetaValueByKey.value;
             }
         }
         const keys: Map<string, Set<string>> = new Map();
         this.forEachSession((item) => {
             const meta = item.MetaHandler;
-            if(meta) {
+            if (meta) {
                 meta.forEach((key, value) => {
                     if (value) {
                         let valueSet = keys.get(key);
@@ -912,14 +914,14 @@ export class TestSessionCollection implements Loader {
                 });
             }
         });
-        this._cachedMetaValueByKey = {hash:this.hash, value:keys};
+        this._cachedMetaValueByKey = { hash: this.hash, value: keys };
         return keys;
     }
 
     *iterSessions() {
         const iterator = this._collection.values();
         let result = iterator.next();
-        while(!result.done) {
+        while (!result.done) {
             yield result.value;
             result = iterator.next();
         }
@@ -928,7 +930,7 @@ export class TestSessionCollection implements Loader {
     *iterTests() {
         const iterSession = this._collection.values();
         let resultSession = iterSession.next();
-        while(!resultSession.done) {
+        while (!resultSession.done) {
             const sessionInfo = resultSession.value.TestSessionInfo;
             if (sessionInfo) {
                 for (const key in sessionInfo.Tests) {
@@ -943,12 +945,12 @@ export class TestSessionCollection implements Loader {
         this.incrementLoading();
         const onFetch = (items: (TestDataWrapper | undefined)[]) => this.onInitialFetch(items);
         try {
-            if(streamId === this._datahandler.streamId
+            if (streamId === this._datahandler.streamId
                 && this._datahandler.items
                 && maxCount > this._datahandler.items.size) {
-                    // Max count resize
-                    this._datahandler.setActive(true);
-                    await this._datahandler.queueFetchItems(onFetch, SessionFields.initiate, maxCount);
+                // Max count resize
+                this._datahandler.setActive(true);
+                await this._datahandler.queueFetchItems(onFetch, SessionFields.initiate, maxCount);
             } else {
                 // Initial fetch or reload
                 await this._datahandler.setFromStream(
@@ -959,12 +961,12 @@ export class TestSessionCollection implements Loader {
         await when(() => this.loadingCount === 0);
     }
 
-    private onInitialFetch(items: (TestDataWrapper|undefined)[]) {
+    private onInitialFetch(items: (TestDataWrapper | undefined)[]) {
         const sessionsToFetch: string[] = [];
         items.forEach((item) => {
             if (item) {
                 let data = item.getDataHandler() as TestSessionWrapper;
-                if(!data) {
+                if (!data) {
                     data = new TestSessionWrapper(item);
                     item.setDataHandler(data);
                 }
@@ -972,7 +974,7 @@ export class TestSessionCollection implements Loader {
                     return;
                 }
                 const session_uid = data.getUID();
-                if(!this.has(session_uid)) {
+                if (!this.has(session_uid)) {
                     this.set(session_uid, data);
                     if (!data.isAllPropertiesLoaded()) {
                         sessionsToFetch.push(item.id);
@@ -980,7 +982,7 @@ export class TestSessionCollection implements Loader {
                 }
             }
         });
-        if(sessionsToFetch.length > 0) {
+        if (sessionsToFetch.length > 0) {
             this.incrementLoading();
             this._datahandler.fetchUpdateItems(
                 sessionsToFetch, (items) => this.onRemainingFetch(items), SessionFields.remaining
@@ -988,23 +990,23 @@ export class TestSessionCollection implements Loader {
         }
     }
 
-    private onRemainingFetch(items: (TestDataWrapper|undefined)[]) {
+    private onRemainingFetch(items: (TestDataWrapper | undefined)[]) {
         let hasUpdated = false;
         items.forEach((item) => {
             if (item) {
                 const data = item.getDataHandler() as TestSessionWrapper;
-                if(data && data.isAllPropertiesLoaded()) {
+                if (data && data.isAllPropertiesLoaded()) {
                     this.indexTestsFromSession(data);
                     hasUpdated = true;
                 }
             }
         });
-        if(hasUpdated) { this.updateVersion() }
+        if (hasUpdated) { this.updateVersion() }
     }
 
     private indexTestsFromSession(session: TestSessionWrapper) {
         const tests = Object.values(session.TestSessionInfo!.Tests);
-        if(tests.length === 0) {
+        if (tests.length === 0) {
             return;
         }
         const first = this._testcases.get(tests[0].TestUID);
@@ -1014,11 +1016,11 @@ export class TestSessionCollection implements Loader {
         tests.forEach((item) => {
             const testuid = item.TestUID;
             let testcase = this._testcases.get(testuid);
-            if(!testcase) {
+            if (!testcase) {
                 testcase = new TestCase(item.Name, testuid, item.Suite, session.Name);
                 this._testcases.set(testuid, testcase);
             }
-            else if(testcase.SessionName !== session.Name) {
+            else if (testcase.SessionName !== session.Name) {
                 console.warn(`'${testcase.DisplayName}' is in another session. Appears in '${testcase.SessionName}' and '${session.Name}'.`);
             }
             testcase.addMeta(item);
@@ -1042,7 +1044,7 @@ export class TestSessionCollection implements Loader {
         }
         getTestCase();
 
-        while(!testcase && toFetch.length === 0 && this.loadingCount !== 0) {
+        while (!testcase && toFetch.length === 0 && this.loadingCount !== 0) {
             // Give a chance to the collection to load a bit more if nothing was found
             // It is possible the test is run infrequently and is deep in the result collection
             const previousVersion = this.version;
@@ -1051,11 +1053,11 @@ export class TestSessionCollection implements Loader {
             getTestCase();
         }
 
-        if(toFetch.length > 0) {
+        if (toFetch.length > 0) {
             // Fetch missing data from session test results
             this.fetchSessionResults(
                 toFetch.map((item) => {
-                    return {SessionId: item.Testdata.id, ResultsId: item.TestSessionInfo.TestResultsTestDataUID}
+                    return { SessionId: item.Testdata.id, ResultsId: item.TestSessionInfo.TestResultsTestDataUID }
                 }), 'data.' + uid
             );
         }
@@ -1067,23 +1069,23 @@ export class TestSessionCollection implements Loader {
         const testcase = await this.getTestCaseByUID(uid);
         let test: TestResult | undefined = undefined;
 
-        if(testcase) {
+        if (testcase) {
             test = testcase.Metas.get(metauid);
-            while(!test && this.loadingCount !== 0) {
+            while (!test && this.loadingCount !== 0) {
                 const previousVersion = testcase.version;
                 await when(() => testcase.version !== previousVersion || this.loadingCount === 0);
                 test = testcase.Metas.get(metauid);
             }
-            if(change && test?.Session.Testdata.change !== change) {
+            if (change && test?.Session.Testdata.change !== change) {
                 // Need to get the target change from the history
                 test = this.getLoadedTestHistoryByMeta(uid)?.get(metauid)?.find(
                     (item) => item.Session.Testdata.change === change
                 );
-                if(!test) {
+                if (!test) {
                     // trigger fetch history in background
                     this.fetchSessionHistory(testcase.SessionName, uid);
                 }
-                while(!test && this.loadingCount !== 0) {
+                while (!test && this.loadingCount !== 0) {
                     const previousVersion = this.version;
                     await when(() => this.version !== previousVersion || this.loadingCount === 0);
                     test = this.getLoadedTestHistoryByMeta(uid)?.get(metauid)?.find(
@@ -1094,13 +1096,13 @@ export class TestSessionCollection implements Loader {
             }
         }
 
-        if(test && test.Details === undefined) {
+        if (test && test.Details === undefined) {
             // Fetch missing data from the session test result
             await this.fetchSessionResults(
-                [{SessionId: test.Session.Testdata.id, ResultsId: test.Session.TestSessionInfo.TestResultsTestDataUID}]
+                [{ SessionId: test.Session.Testdata.id, ResultsId: test.Session.TestSessionInfo.TestResultsTestDataUID }]
                 , 'data.' + uid
             );
-            if(test.Details === undefined) {
+            if (test.Details === undefined) {
                 // if nothing was fetched or failed to fetch, return nothing
                 return undefined;
             }
@@ -1126,11 +1128,11 @@ export class TestSessionCollection implements Loader {
     private async fetchResultItemFromSession(item: SessionReferences, filter?: string): Promise<number> {
         const update = await this._datahandler.fetchItemsFromAlternateKey(`${this._datahandler.key} Result Details::${item.ResultsId}`, undefined, 1, 0, filter);
         // We only expect one returned item
-        if(update.length > 0) {
+        if (update.length > 0) {
             const results = update[0].data as TestResultDetailsMap;
             const session = this._datahandler.items?.get(item.SessionId)?.getDataHandler() as TestSessionWrapper;
-            if(session) {
-                for(const uid in results ) {
+            if (session) {
+                for (const uid in results) {
                     const test = session.getTestResultByUID(uid);
                     if (test) {
                         test.Details = results[uid];
@@ -1147,9 +1149,9 @@ export class TestSessionCollection implements Loader {
         const tests: TestResult[] = [];
         const iterator = this._collection.values();
         let result = iterator.next();
-        while(!result.done) {
+        while (!result.done) {
             const sessionInfo = result.value.TestSessionInfo;
-            if(sessionInfo) {
+            if (sessionInfo) {
                 tests.push(...Object.values(result.value.TestSessionInfo.Tests));
             }
             result = iterator.next();
@@ -1163,21 +1165,21 @@ export class TestSessionCollection implements Loader {
 
     getLoadedTestHistoryByMeta(uid: string) {
         const testcase = this._testcases.get(uid);
-        if(!testcase) {
+        if (!testcase) {
             return undefined;
         }
 
         const testresultsByMeta: Map<string, TestResult[]> = new Map();
         this.forEachSessionHistoryByChange(testcase.SessionName,
             (session) => {
-                if(session.isAllPropertiesLoaded()) {
+                if (session.isAllPropertiesLoaded()) {
                     const test = session.getTestResultByUID(uid);
                     if (test) {
                         const metauid = session.MetaHandler!.hash;
                         const meta = testresultsByMeta.get(metauid);
                         if (meta) {
                             const change = session.Testdata.change;
-                            if(!meta.some((item) => item.Session.Testdata.change === change)) {
+                            if (!meta.some((item) => item.Session.Testdata.change === change)) {
                                 // Keep only one result per change
                                 meta.push(test);
                             }
@@ -1193,15 +1195,15 @@ export class TestSessionCollection implements Loader {
     }
 
     async fetchSessionHistory(sessionName: string, testuid: string, maxCount: number = 1000) {
-        if(this._datahandler.cursor) {
+        if (this._datahandler.cursor) {
             // for single view mode we need to fetch the whole pool for context
-            const onFetch = (items: (TestDataWrapper|undefined)[]) => this.onInitialHistoryFetch(items, sessionName, testuid)
+            const onFetch = (items: (TestDataWrapper | undefined)[]) => this.onInitialHistoryFetch(items, sessionName, testuid)
             await this._datahandler.fetchCursorHistory(onFetch, SessionFields.initiate, maxCount);
             return;
         }
 
         await this.fetchIncompleteSessionsForHistory(sessionName, testuid);
-        while(this.loadingCount !== 0) {
+        while (this.loadingCount !== 0) {
             // wait for the collection to be fully loaded, in the meanwhile load the missing session history
             const previousVersion = this.version;
             await when(() => this.version !== previousVersion || this.loadingCount === 0);
@@ -1213,14 +1215,14 @@ export class TestSessionCollection implements Loader {
         const missingSessions: string[] = [];
         this.forEachSessionHistoryByChange(sessionName,
             (session) => {
-                if(!session.isAllPropertiesLoaded()) {
+                if (!session.isAllPropertiesLoaded()) {
                     missingSessions.push(session.Testdata.id);
                     this._testcases.get(testuid)?.incrementHistoryTarget();
                 }
             }
         );
 
-        if(missingSessions.length > 0) {
+        if (missingSessions.length > 0) {
             await this._datahandler.fetchUpdateItems(
                 missingSessions, (item) => this.onFetchTestCaseHistory(item, testuid), SessionFields.remaining
             )
@@ -1234,7 +1236,7 @@ export class TestSessionCollection implements Loader {
         while (!change.done) {
             change.value.forEach((item) => {
                 const session = item.getDataHandler() as TestSessionWrapper;
-                if(session.Name === sessionName) {
+                if (session.Name === sessionName) {
                     if (!this._includePreflight && session.PreFlightChange) {
                         return;
                     }
@@ -1245,31 +1247,31 @@ export class TestSessionCollection implements Loader {
         }
     }
 
-    private onInitialHistoryFetch(items: (TestDataWrapper|undefined)[], sessionName: string, testuid: string) {
+    private onInitialHistoryFetch(items: (TestDataWrapper | undefined)[], sessionName: string, testuid: string) {
         const toFetch: TestDataWrapper[] = [];
         items.forEach((item) => {
             if (item) {
                 let data = item.getDataHandler() as TestSessionWrapper;
-                if(!data) {
+                if (!data) {
                     data = new TestSessionWrapper(item);
                     item.setDataHandler(data);
-                    if(data.Name === sessionName) {
+                    if (data.Name === sessionName) {
                         toFetch.push(item);
                         this._testcases.get(testuid)?.incrementHistoryTarget();
                     }
                 }
             }
         });
-        if(toFetch.length > 0) {
+        if (toFetch.length > 0) {
             this._datahandler.fetchUpdateItems(
                 toFetch.map((item) => item.id), (items) => this.onFetchTestCaseHistory(items, testuid), SessionFields.remaining
             )
         }
     }
 
-    private onFetchTestCaseHistory(items: (TestDataWrapper|undefined)[], testuid: string) {
+    private onFetchTestCaseHistory(items: (TestDataWrapper | undefined)[], testuid: string) {
         items.forEach((item) => {
-            if(item) {
+            if (item) {
                 this._testcases.get(testuid)?.updateHistory();
             }
         });
diff --git a/Engine/Source/Programs/Horde/HordeDashboard/src/Plugins/SimpleTestReport/index.tsx b/Engine/Source/Programs/Horde/HordeDashboard/src/Plugins/SimpleTestReport/index.tsx
--- a/Engine/Source/Programs/Horde/HordeDashboard/src/Plugins/SimpleTestReport/index.tsx
+++ b/Engine/Source/Programs/Horde/HordeDashboard/src/Plugins/SimpleTestReport/index.tsx
@@ -20,4 +20,5 @@ class SimpleTestReportPlugin implements Plugin {
     ];
 }
 
-export default new SimpleTestReportPlugin();
+const simpleTestReportPlugin = new SimpleTestReportPlugin();
+export default simpleTestReportPlugin; 
diff --git a/Engine/Source/Programs/Horde/HordeDashboard/src/Plugins/UnrealAutomatedTests/components/UnrealAutomatedTests.tsx b/Engine/Source/Programs/Horde/HordeDashboard/src/Plugins/UnrealAutomatedTests/components/UnrealAutomatedTests.tsx
--- a/Engine/Source/Programs/Horde/HordeDashboard/src/Plugins/UnrealAutomatedTests/components/UnrealAutomatedTests.tsx
+++ b/Engine/Source/Programs/Horde/HordeDashboard/src/Plugins/UnrealAutomatedTests/components/UnrealAutomatedTests.tsx
@@ -3,7 +3,7 @@
 import { Checkbox, Icon, Image, Spinner, SpinnerSize, Stack, Text } from '@fluentui/react';
 import { getTheme, mergeStyles, mergeStyleSets } from '@fluentui/react/lib/Styling';
 import React, { useState, useEffect, useRef } from 'react';
-import { useHistory } from 'react-router';
+import { useNavigate } from 'react-router';
 import { Link } from 'react-router-dom';
 import backend from '../../../backend';
 import { ArtifactData } from '../../../backend/Api';
@@ -15,542 +15,542 @@ import { EventType, Metadata, TestDetails, TestEntry, TestEntryArtifact, TestPas
 
 const theme = getTheme();
 const gutterClass = mergeStyles({
-    borderLeftWidth: 6,
-    padding: 0,
-    margin: 0,
-    paddingTop: 0,
-    paddingBottom: 0,
-    paddingRight: 8,
-    marginTop: 0,
-    marginBottom: 0,
-    height: 20
+   borderLeftWidth: 6,
+   padding: 0,
+   margin: 0,
+   paddingTop: 0,
+   paddingBottom: 0,
+   paddingRight: 8,
+   marginTop: 0,
+   marginBottom: 0,
+   height: 20
 });
 const styles = mergeStyleSets({
-    container: {
-        overflow: 'auto',
-        height: 'calc(100vh - 165px)',
-        marginTop: 8
-    },
-    item: [
-        {
-            fontSize: "11px",
-            fontFamily: "Horde Cousine Regular"
-        }
-    ],
-    gutter: [
-        {
-            padding: 0,
-            margin: 0,
-            paddingTop: 0,
-            paddingBottom: 0,
-            paddingRight: 14,
-            marginTop: 0,
-            marginBottom: 0,
-            height: 20
-        }
-    ],
-    gutterError: [
-        {
-            background: "#FEF6F6",
-            borderLeftStyle: 'solid',
-            borderLeftColor: "#EC4C47"
-        }, gutterClass
-    ],
-    gutterWarning: [
-        {
-            background: "#FEF8E7",
-            borderLeftStyle: 'solid',
-            borderLeftColor: "#F7D154"
-        }, gutterClass
-    ],
-    gutterSuccess: [
-        {
-            borderLeftStyle: 'solid',
-            borderLeftColor: theme.palette.green
-        }, gutterClass
-    ],
-    itemWarning: [
-        {
-            background: "#FEF8E7"
-        }
-    ],
-    itemError: [
-        {
-            background: "#FEF6F6"
-        }
-    ],
-    itemHover: {
-        cursor: "pointer",
-        selectors:{
-            ':hover': {
-                background: theme.palette.neutralLight
-            }
-        }
-    },
-    itemHighlighted: {
-        background: theme.palette.neutralLighter
-    },
-    historyList: {
-        zIndex: 1,
-        borderWidth: "1px",
-        borderColor: "#888",
-        borderStyle: "solid",
-        backgroundColor: "#FFF",
-        padding: "3px",
-        overflow: "auto"
-    }
+   container: {
+      overflow: 'auto',
+      height: 'calc(100vh - 165px)',
+      marginTop: 8
+   },
+   item: [
+      {
+         fontSize: "11px",
+         fontFamily: "Horde Cousine Regular"
+      }
+   ],
+   gutter: [
+      {
+         padding: 0,
+         margin: 0,
+         paddingTop: 0,
+         paddingBottom: 0,
+         paddingRight: 14,
+         marginTop: 0,
+         marginBottom: 0,
+         height: 20
+      }
+   ],
+   gutterError: [
+      {
+         background: "#FEF6F6",
+         borderLeftStyle: 'solid',
+         borderLeftColor: "#EC4C47"
+      }, gutterClass
+   ],
+   gutterWarning: [
+      {
+         background: "#FEF8E7",
+         borderLeftStyle: 'solid',
+         borderLeftColor: "#F7D154"
+      }, gutterClass
+   ],
+   gutterSuccess: [
+      {
+         borderLeftStyle: 'solid',
+         borderLeftColor: theme.palette.green
+      }, gutterClass
+   ],
+   itemWarning: [
+      {
+         background: "#FEF8E7"
+      }
+   ],
+   itemError: [
+      {
+         background: "#FEF6F6"
+      }
+   ],
+   itemHover: {
+      cursor: "pointer",
+      selectors: {
+         ':hover': {
+            background: theme.palette.neutralLight
+         }
+      }
+   },
+   itemHighlighted: {
+      background: theme.palette.neutralLighter
+   },
+   historyList: {
+      zIndex: 1,
+      borderWidth: "1px",
+      borderColor: "#888",
+      borderStyle: "solid",
+      backgroundColor: "#FFF",
+      padding: "3px",
+      overflow: "auto"
+   }
 });
 
 const gutterHystoryStyles = new Map<string, string>([
-    [TestState.Success, styles.gutterSuccess],
-    [TestState.InProcess, mergeStyles({
-            background: "#E4F1F5",
-            borderLeftStyle: 'solid',
-            borderLeftColor: "#01BCF2"
-        }, gutterClass)],
-    [TestState.NotRun, mergeStyles({
-            borderLeftStyle: 'solid',
-            borderLeftColor: "#A19F9D"
-        }, gutterClass)],
-    [TestState.SuccessWithWarnings, styles.gutterWarning],
-    [TestState.Failed, styles.gutterError],
+   [TestState.Success, styles.gutterSuccess],
+   [TestState.InProcess, mergeStyles({
+      background: "#E4F1F5",
+      borderLeftStyle: 'solid',
+      borderLeftColor: "#01BCF2"
+   }, gutterClass)],
+   [TestState.NotRun, mergeStyles({
+      borderLeftStyle: 'solid',
+      borderLeftColor: "#A19F9D"
+   }, gutterClass)],
+   [TestState.SuccessWithWarnings, styles.gutterWarning],
+   [TestState.Failed, styles.gutterError],
 ]);
 
 const iconClass = mergeStyles({
-    fontSize: 12
+   fontSize: 12
 });
 
 const stateStyles = new Map<string, string>([
-    [TestState.Success, mergeStyles({ color: theme.palette.green, userSelect: "none" }, iconClass)],
-    [TestState.InProcess, mergeStyles({  color: "#01BCF2", userSelect: "none" }, iconClass)],
-    [TestState.NotRun, mergeStyles({ color: "#A19F9D", userSelect: "none" }, iconClass)],
-    [TestState.SuccessWithWarnings, mergeStyles({ color: "#F7D154", userSelect: "none" }, iconClass)],
-    [TestState.Failed, mergeStyles({ color: "#EC4C47", userSelect: "none" }, iconClass)],
-    [TestState.Unknown, mergeStyles({ color: "#000000", userSelect: "none" }, iconClass)],
+   [TestState.Success, mergeStyles({ color: theme.palette.green, userSelect: "none" }, iconClass)],
+   [TestState.InProcess, mergeStyles({ color: "#01BCF2", userSelect: "none" }, iconClass)],
+   [TestState.NotRun, mergeStyles({ color: "#A19F9D", userSelect: "none" }, iconClass)],
+   [TestState.SuccessWithWarnings, mergeStyles({ color: "#F7D154", userSelect: "none" }, iconClass)],
+   [TestState.Failed, mergeStyles({ color: "#EC4C47", userSelect: "none" }, iconClass)],
+   [TestState.Unknown, mergeStyles({ color: "#000000", userSelect: "none" }, iconClass)],
 ]);
-const getTestStateStyles = (test: TestResult) : string | undefined => {
-    if (!stateStyles.has(test.State)) {
-        return stateStyles.get(TestState.Unknown);
-    }
-    return stateStyles.get(test.Warnings > 0 && test.State === TestState.Success? TestState.SuccessWithWarnings : test.State)
+const getTestStateStyles = (test: TestResult): string | undefined => {
+   if (!stateStyles.has(test.State)) {
+      return stateStyles.get(TestState.Unknown);
+   }
+   return stateStyles.get(test.Warnings > 0 && test.State === TestState.Success ? TestState.SuccessWithWarnings : test.State)
 }
 
 const shadowBoxSlyle = {
-    boxShadow: '0 4px 8px 0 rgba(0, 0, 0, 0.2), 0 6px 20px 0 rgba(0, 0, 0, 0.19)',
-    margin: 5,
-    paddingRight: 5
+   boxShadow: '0 4px 8px 0 rgba(0, 0, 0, 0.2), 0 6px 20px 0 rgba(0, 0, 0, 0.19)',
+   margin: 5,
+   paddingRight: 5
 };
 
 const stateLabels = new Map<string, string>([
-    [TestState.Success, 'Passed'],
-    [TestState.InProcess, 'Incomplete'],
-    [TestState.NotRun, 'Unexecuted'],
-    [TestState.SuccessWithWarnings, 'Passed'],
-    [TestState.Failed, 'Failed'],
-    [TestState.Skipped, 'Skipped'],
-    [TestState.Unknown, 'Unknown'],
+   [TestState.Success, 'Passed'],
+   [TestState.InProcess, 'Incomplete'],
+   [TestState.NotRun, 'Unexecuted'],
+   [TestState.SuccessWithWarnings, 'Passed'],
+   [TestState.Failed, 'Failed'],
+   [TestState.Skipped, 'Skipped'],
+   [TestState.Unknown, 'Unknown'],
 ]);
-const getStateLabel = (state: string) : string | undefined => {return stateLabels.has(state)? stateLabels.get(state) : state}
+const getStateLabel = (state: string): string | undefined => { return stateLabels.has(state) ? stateLabels.get(state) : state }
 
-function getMetadata(collection: Metadata, key: string) : string | undefined {
-    if(collection) {
-        return collection[key];
-    }
+function getMetadata(collection: Metadata, key: string): string | undefined {
+   if (collection) {
+      return collection[key];
+   }
 
-    return undefined;
+   return undefined;
 }
 
 const copyToClipboard = (value: string | undefined) => {
 
-    if (!value) {
-        return;
-    }
+   if (!value) {
+      return;
+   }
 
-    const el = document.createElement('textarea');
-    el.value = value;
-    document.body.appendChild(el);
-    el.select();
-    document.execCommand('copy');
-    document.body.removeChild(el);
+   const el = document.createElement('textarea');
+   el.value = value;
+   document.body.appendChild(el);
+   el.select();
+   document.execCommand('copy');
+   document.body.removeChild(el);
 }
 
 const missingImage = "/images/missing-image.png";
-const MissingImageLabel = ():JSX.Element => {return <span style={{fontWeight:'bold'}}> [missing image]</span>}
-type ImageLinks = {approved?: string, unapproved?: string, difference?: string}
-const buildImageLink = (artifact?: ArtifactData) => artifact !== undefined?`${backend.serverUrl}/api/v1/artifacts/${artifact.id}/download?Code=${artifact.code}`:undefined;
-
-const EntryPane: React.FC<{entry: TestEntry, testArtifacts: TestEntryArtifact[]}> = (props) => {
-    const { entry, testArtifacts } = props;
-    const [imageLinks, setImageLinks] = useState<ImageLinks>({});
-    const eventType = entry.Event.Type;
-    const is_error = eventType === EventType.Error;
-    const is_warning = eventType === EventType.Warning;
-    const style = is_error?styles.itemError:(is_warning?styles.itemWarning:styles.item);
-    const gutterStyle = is_error?styles.gutterError:(is_warning?styles.gutterWarning:styles.gutter);
-
-    const artifact = testArtifacts.find((value) => value.Id === entry.Event.Artifact);
-    const artifactRef = useRef(artifact);
-
-    const need_image_comparison = artifact !== undefined && artifact.Type === "Comparison";
-    useEffect(() => {
-        const artifact = artifactRef.current;
-        if (artifact !== undefined) {
-            const findLinks = async () => {
-                const imageLinks: ImageLinks = {};
-                // Approved
-                let foundJobArtifact = await testDataHandler.cursor?.findArtifactData(artifact.Files.Approved);
-                imageLinks.approved = buildImageLink(foundJobArtifact);
-                // Unapproved
-                foundJobArtifact = await testDataHandler.cursor?.findArtifactData(artifact.Files.Unapproved);
-                imageLinks.unapproved = buildImageLink(foundJobArtifact);
-                // Difference
-                foundJobArtifact = await testDataHandler.cursor?.findArtifactData(artifact.Files.Difference);
-                imageLinks.difference = buildImageLink(foundJobArtifact);
-
-                setImageLinks(imageLinks);
-            }
-            findLinks();
-        }
-    }, []);
-
-    return (
-        <Stack className={styles.item} disableShrink={true} styles={{ root: { paddingLeft: 8, width: "100%" } }}>
-            <Stack horizontal>
-                <Stack className={gutterStyle}></Stack>
-                <Stack.Item className={style} align="center" styles={{ root: { paddingLeft: 8, width: "100%" } }}><pre style={{margin: 0, whiteSpace: "pre-wrap"}}>[{entry.Timestamp}] {eventType}: {entry.Event.Message}</pre></Stack.Item>
+const MissingImageLabel = (): JSX.Element => { return <span style={{ fontWeight: 'bold' }}> [missing image]</span> }
+type ImageLinks = { approved?: string, unapproved?: string, difference?: string }
+const buildImageLink = (artifact?: ArtifactData) => artifact !== undefined ? `${backend.serverUrl}/api/v1/artifacts/${artifact.id}/download?Code=${artifact.code}` : undefined;
+
+const EntryPane: React.FC<{ entry: TestEntry, testArtifacts: TestEntryArtifact[] }> = (props) => {
+   const { entry, testArtifacts } = props;
+   const [imageLinks, setImageLinks] = useState<ImageLinks>({});
+   const eventType = entry.Event.Type;
+   const is_error = eventType === EventType.Error;
+   const is_warning = eventType === EventType.Warning;
+   const style = is_error ? styles.itemError : (is_warning ? styles.itemWarning : styles.item);
+   const gutterStyle = is_error ? styles.gutterError : (is_warning ? styles.gutterWarning : styles.gutter);
+
+   const artifact = testArtifacts.find((value) => value.Id === entry.Event.Artifact);
+   const artifactRef = useRef(artifact);
+
+   const need_image_comparison = artifact !== undefined && artifact.Type === "Comparison";
+   useEffect(() => {
+      const artifact = artifactRef.current;
+      if (artifact !== undefined) {
+         const findLinks = async () => {
+            const imageLinks: ImageLinks = {};
+            // Approved
+            let foundJobArtifact = await testDataHandler.cursor?.findArtifactData(artifact.Files.Approved);
+            imageLinks.approved = buildImageLink(foundJobArtifact);
+            // Unapproved
+            foundJobArtifact = await testDataHandler.cursor?.findArtifactData(artifact.Files.Unapproved);
+            imageLinks.unapproved = buildImageLink(foundJobArtifact);
+            // Difference
+            foundJobArtifact = await testDataHandler.cursor?.findArtifactData(artifact.Files.Difference);
+            imageLinks.difference = buildImageLink(foundJobArtifact);
+
+            setImageLinks(imageLinks);
+         }
+         findLinks();
+      }
+   }, []);
+
+   return (
+      <Stack className={styles.item} disableShrink={true} styles={{ root: { paddingLeft: 8, width: "100%" } }}>
+         <Stack horizontal>
+            <Stack className={gutterStyle}></Stack>
+            <Stack.Item className={style} align="center" styles={{ root: { paddingLeft: 8, width: "100%" } }}><pre style={{ margin: 0, whiteSpace: "pre-wrap" }}>[{entry.Timestamp}] {eventType}: {entry.Event.Message}</pre></Stack.Item>
+         </Stack>
+         {need_image_comparison &&
+            <Stack styles={{ root: { paddingLeft: 16 } }}>
+               <Stack><Text variant="medium" styles={{ root: { fontWeight: "bold" } }}>Image comparison: {artifact?.Name}</Text></Stack>
+               <Stack horizontal>
+                  <Stack styles={{ root: { padding: 5 } }}>
+                     <a href={imageLinks.approved}><Image width={400} src={imageLinks.approved || missingImage} alt={artifact?.Files.Approved} /></a>
+                     <Stack.Item align="center">Approved{!imageLinks.approved && MissingImageLabel()}</Stack.Item>
+                  </Stack>
+                  <Stack styles={{ root: { padding: 5 } }}>
+                     <a href={imageLinks.difference}><Image width={400} src={imageLinks.difference || missingImage} alt={artifact?.Files.Difference} /></a>
+                     <Stack.Item align="center">Difference{!imageLinks.difference && MissingImageLabel()}</Stack.Item>
+                  </Stack>
+                  <Stack styles={{ root: { padding: 5 } }}>
+                     <a href={imageLinks.unapproved}><Image width={400} src={imageLinks.unapproved || missingImage} alt={artifact?.Files.Unapproved} /></a>
+                     <Stack.Item align="center">Unapproved{!imageLinks.unapproved && MissingImageLabel()}</Stack.Item>
+                  </Stack>
+               </Stack>
             </Stack>
-            {need_image_comparison &&
-                <Stack styles={{ root: { paddingLeft: 16} }}>
-                    <Stack><Text variant="medium" styles={{root:{fontWeight: "bold"}}}>Image comparison: {artifact?.Name}</Text></Stack>
-                    <Stack horizontal>
-                        <Stack styles={{root:{padding: 5}}}>
-                            <a href={imageLinks.approved}><Image width={400} src={imageLinks.approved||missingImage} alt={artifact?.Files.Approved}/></a>
-                            <Stack.Item align="center">Approved{!imageLinks.approved && MissingImageLabel()}</Stack.Item>
-                        </Stack>
-                        <Stack styles={{root:{padding: 5}}}>
-                            <a href={imageLinks.difference}><Image width={400} src={imageLinks.difference||missingImage} alt={artifact?.Files.Difference}/></a>
-                            <Stack.Item align="center">Difference{!imageLinks.difference && MissingImageLabel()}</Stack.Item>
-                        </Stack>
-                        <Stack styles={{root:{padding: 5}}}>
-                            <a href={imageLinks.unapproved}><Image width={400} src={imageLinks.unapproved||missingImage} alt={artifact?.Files.Unapproved}/></a>
-                            <Stack.Item align="center">Unapproved{!imageLinks.unapproved && MissingImageLabel()}</Stack.Item>
-                        </Stack>
-                    </Stack>
-                </Stack>
-            }
-        </Stack>
-    );
+         }
+      </Stack>
+   );
 }
 
-const HistoryItem: React.FC<{item: TestStateHistoryItem, testName: string, selected: boolean}> = (props) => {
-    const { item, testName, selected } = props;
-    const gutterStyle = gutterHystoryStyles.get(item.State);
-
-    return (
-        <Stack horizontal className={`${styles.itemHover} ${selected? styles.itemHighlighted : ""}`}>
-            <Stack className={gutterStyle}></Stack>
-            <Stack.Item>
-                <Link to={`/testreport/${item.TestdataId}?test=${testName}`}>
-                    <Text variant="smallPlus">{getStateLabel(item.State)} on {item.Change}</Text>
-                </Link>
-            </Stack.Item>
-        </Stack>
-    );
+const HistoryItem: React.FC<{ item: TestStateHistoryItem, testName: string, selected: boolean }> = (props) => {
+   const { item, testName, selected } = props;
+   const gutterStyle = gutterHystoryStyles.get(item.State);
+
+   return (
+      <Stack horizontal className={`${styles.itemHover} ${selected ? styles.itemHighlighted : ""}`}>
+         <Stack className={gutterStyle}></Stack>
+         <Stack.Item>
+            <Link to={`/testreport/${item.TestdataId}?test=${testName}`}>
+               <Text variant="smallPlus">{getStateLabel(item.State)} on {item.Change}</Text>
+            </Link>
+         </Stack.Item>
+      </Stack>
+   );
 }
 
-const IconItem: React.FC<{item: TestResult}> = (props) => {
-    const { item } = props;
-    const pageHistory = useHistory();
-    const style = getTestStateStyles(item);
-    function onClickTest() {
-        pageHistory.replace(`${window.location.pathname}?test=${item.FullTestPath}`);
-    }
-    return (
-        <Icon styles={{root: {cursor: "pointer"}}} className={style} iconName="Square" title={item.FullTestPath} onClick={(ev) => {onClickTest()}}/>
-    );
+const IconItem: React.FC<{ item: TestResult }> = (props) => {
+   const { item } = props;
+   const navigate = useNavigate();
+   const style = getTestStateStyles(item);
+   function onClickTest() {
+      navigate(`${window.location.pathname}?test=${item.FullTestPath}`, { replace: true });
+   }
+   return (
+      <Icon styles={{ root: { cursor: "pointer" } }} className={style} iconName="Square" title={item.FullTestPath} onClick={(ev) => { onClickTest() }} />
+   );
 }
 
-const TestResultPane: React.FC<{test: TestResult, selected: boolean}> = (props) => {
-    const { test, selected } = props;
-    const [testDetails, setTestDetails] = useState<TestDetails | undefined>(undefined);
-    const [loading, setLoading] = useState(false);
-    const [visible, setVisible] = useState(selected);
-    const [filterError, setFilterError] = useState(true);
-    const [filterWarning, setFilterWarning] = useState(true);
-    const [filterInfo, setFilterInfo] = useState(test.State !== TestState.Failed);
-    const [historyLoaded, setHistoryLoaded] = useState(false);
-    const [historyVisible, setHistoryVisible] = useState(false);
-    const historyRef = useRef<TestStateHistoryItem[]>([]);
-
-    const pageHistory = useHistory();
-
-    function onClickTest() {
-        if(selected) {
-            if (visible) {
-                // remove test name selection
-                pageHistory.replace(window.location.pathname);
-                return;
-            }
-        }
-        // add test name selection
-        pageHistory.replace(`${window.location.pathname}?test=${test.FullTestPath}`);
-    }
-
-    function onClickTestName() {
-        if (testDetails === undefined) {
-            getTestDetails();
-            setVisible(true);
+const TestResultPane: React.FC<{ test: TestResult, selected: boolean }> = (props) => {
+   const { test, selected } = props;
+   const [testDetails, setTestDetails] = useState<TestDetails | undefined>(undefined);
+   const [loading, setLoading] = useState(false);
+   const [visible, setVisible] = useState(selected);
+   const [filterError, setFilterError] = useState(true);
+   const [filterWarning, setFilterWarning] = useState(true);
+   const [filterInfo, setFilterInfo] = useState(test.State !== TestState.Failed);
+   const [historyLoaded, setHistoryLoaded] = useState(false);
+   const [historyVisible, setHistoryVisible] = useState(false);
+   const historyRef = useRef<TestStateHistoryItem[]>([]);
+
+   const navigate = useNavigate();
+
+   function onClickTest() {
+      if (selected) {
+         if (visible) {
+            // remove test name selection
+            navigate(window.location.pathname, {replace: true});
             return;
-        }
-
-        setVisible(!visible);
-    }
-
-    async function getTestDetails() {
-        if (testDetails !== undefined) {
-            return;
-        }
-
-        if (loading) {
-            return;
-        }
-        const testArtifact : ArtifactData | undefined = await testDataHandler.cursor?.findArtifactData(test.ArtifactName);
-        if (!testArtifact) {
-            console.error("Could not find Job Artifacts Data with name '"+ test.ArtifactName +"'!");
-            return;
-        }
-
-        backend.getArtifactDataById(testArtifact.id).then(
-            (value) => {setTestDetails(value as TestDetails)}
-        ).catch(
-            (reason) => {console.error(reason)}
-        ).finally(
-            () => {setLoading(false)}
-        );
-
-        setLoading(true);
-    }
-
-    function onChangeFilterError(ev?: React.FormEvent<HTMLElement | HTMLInputElement>, checked?: boolean) {
-        checked !== undefined && setFilterError(checked);
-    }
-    function onChangeFilterWarning(ev?: React.FormEvent<HTMLElement | HTMLInputElement>, checked?: boolean) {
-        checked !== undefined && setFilterWarning(checked);
-    }
-    function onChangeFilterInfo(ev?: React.FormEvent<HTMLElement | HTMLInputElement>, checked?: boolean) {
-        checked !== undefined && setFilterInfo(checked);
-    }
-
-    function isEntryNeedDisplay(entry: TestEntry, index?: number, array?: TestEntry[]): boolean {
-        const entryEventType = entry.Event.Type;
-        if (filterError && entryEventType === EventType.Error) {
-            return true;
-        }
-        if (filterWarning && entryEventType === EventType.Warning) {
-            return true;
-        }
-        if (filterInfo && entryEventType === EventType.Info) {
-            return true;
-        }
-
-        return false;
-    }
-
-    function getTestHistory(testFullName: string, testdataItems: TestDataWrapper[]): TestStateHistoryItem[] {
-        const testResults: TestStateHistoryItem[] = [];
-
-        testdataItems.forEach((item) => {
-            const testdata = item.data as TestPassSummary;
-            const foundTest = testdata.Tests.find(testItem => testItem.FullTestPath === testFullName);
-
-            if (foundTest !== undefined) {
-                const testHistoryItem: TestStateHistoryItem = {
-                    Change: item.change,
-                    TestdataId: item.id,
-                    State: foundTest.Warnings > 0 && foundTest.State === TestState.Success? TestState.SuccessWithWarnings : foundTest.State,
-                }; 
-                testResults.push(testHistoryItem);
-            }
-        });
-
-        return testResults;
-    }
-
-    function onClickHistory() {
-        const needHistoryVisible = !historyVisible;
-        setHistoryVisible(needHistoryVisible);
-        if (needHistoryVisible && !historyLoaded) {
-            testDataHandler.getCursorHistory(undefined, undefined, 100).then(
-                (items) => historyRef.current = getTestHistory(test.FullTestPath, items)
-            ).catch(
-                (reason) => console.error(reason)
-            ).finally(
-                () => {setHistoryLoaded(true)}
-            );
-        }
-    }
-    function onHistoryBlur(ev: any) {
-        // onBlur is used on the History widget in order to make the widget hide when user click outside of it (it loses the focus).
-        // However, item links clicked in the widget make the parent widget lose its focus, so we check if the focus is lost in favor
-        // of one of its children; in that case the widget is kept visible.
-        if (!ev.currentTarget.contains(ev.relatedTarget as Node)) setHistoryVisible(false);
-    }
-
-    const style = getTestStateStyles(test);
-    const history = historyRef.current;
-
-    if (selected) {
-        getTestDetails();
-        if(!visible) {
-            setVisible(true); // Set the visibility when using the top icon for focus
-        }
-    }
-
-    const scrollToSelected = (ref: any) => {
-        if (ref && selected && visible) {
-            ref.scrollIntoView({ behavior: 'smooth', block: 'start' });
-        }
-    }
-
-    return (
-        <Stack styles={{root:{paddingBottom: 6}}} horizontalAlign="start" onClick={(ev) => {onClickTest()}}>
-            <div onClick={(ev) => {onClickTestName()}} className={styles.itemHover} style={{width: "100%"}} ref={scrollToSelected}>
-                <Stack horizontal>
-                    <Icon styles={{ root: { paddingTop: 7, paddingRight: 8 } }} className={style} iconName="Square" />
-                    <Stack.Item disableShrink><Text variant="mediumPlus" >{test.TestDisplayName}</Text></Stack.Item>
-                </Stack>
+         }
+      }
+      // add test name selection
+      navigate(`${window.location.pathname}?test=${test.FullTestPath}`, {replace: true});
+   }
+
+   function onClickTestName() {
+      if (testDetails === undefined) {
+         getTestDetails();
+         setVisible(true);
+         return;
+      }
+
+      setVisible(!visible);
+   }
+
+   async function getTestDetails() {
+      if (testDetails !== undefined) {
+         return;
+      }
+
+      if (loading) {
+         return;
+      }
+      const testArtifact: ArtifactData | undefined = await testDataHandler.cursor?.findArtifactData(test.ArtifactName);
+      if (!testArtifact) {
+         console.error("Could not find Job Artifacts Data with name '" + test.ArtifactName + "'!");
+         return;
+      }
+
+      backend.getArtifactDataById(testArtifact.id).then(
+         (value) => { setTestDetails(value as TestDetails) }
+      ).catch(
+         (reason) => { console.error(reason) }
+      ).finally(
+         () => { setLoading(false) }
+      );
+
+      setLoading(true);
+   }
+
+   function onChangeFilterError(ev?: React.FormEvent<HTMLElement | HTMLInputElement>, checked?: boolean) {
+      checked !== undefined && setFilterError(checked);
+   }
+   function onChangeFilterWarning(ev?: React.FormEvent<HTMLElement | HTMLInputElement>, checked?: boolean) {
+      checked !== undefined && setFilterWarning(checked);
+   }
+   function onChangeFilterInfo(ev?: React.FormEvent<HTMLElement | HTMLInputElement>, checked?: boolean) {
+      checked !== undefined && setFilterInfo(checked);
+   }
+
+   function isEntryNeedDisplay(entry: TestEntry, index?: number, array?: TestEntry[]): boolean {
+      const entryEventType = entry.Event.Type;
+      if (filterError && entryEventType === EventType.Error) {
+         return true;
+      }
+      if (filterWarning && entryEventType === EventType.Warning) {
+         return true;
+      }
+      if (filterInfo && entryEventType === EventType.Info) {
+         return true;
+      }
+
+      return false;
+   }
+
+   function getTestHistory(testFullName: string, testdataItems: TestDataWrapper[]): TestStateHistoryItem[] {
+      const testResults: TestStateHistoryItem[] = [];
+
+      testdataItems.forEach((item) => {
+         const testdata = item.data as TestPassSummary;
+         const foundTest = testdata.Tests.find(testItem => testItem.FullTestPath === testFullName);
+
+         if (foundTest !== undefined) {
+            const testHistoryItem: TestStateHistoryItem = {
+               Change: item.change,
+               TestdataId: item.id,
+               State: foundTest.Warnings > 0 && foundTest.State === TestState.Success ? TestState.SuccessWithWarnings : foundTest.State,
+            };
+            testResults.push(testHistoryItem);
+         }
+      });
+
+      return testResults;
+   }
+
+   function onClickHistory() {
+      const needHistoryVisible = !historyVisible;
+      setHistoryVisible(needHistoryVisible);
+      if (needHistoryVisible && !historyLoaded) {
+         testDataHandler.getCursorHistory(undefined, undefined, 100).then(
+            (items) => historyRef.current = getTestHistory(test.FullTestPath, items)
+         ).catch(
+            (reason) => console.error(reason)
+         ).finally(
+            () => { setHistoryLoaded(true) }
+         );
+      }
+   }
+   function onHistoryBlur(ev: any) {
+      // onBlur is used on the History widget in order to make the widget hide when user click outside of it (it loses the focus).
+      // However, item links clicked in the widget make the parent widget lose its focus, so we check if the focus is lost in favor
+      // of one of its children; in that case the widget is kept visible.
+      if (!ev.currentTarget.contains(ev.relatedTarget as Node)) setHistoryVisible(false);
+   }
+
+   const style = getTestStateStyles(test);
+   const history = historyRef.current;
+
+   if (selected) {
+      getTestDetails();
+      if (!visible) {
+         setVisible(true); // Set the visibility when using the top icon for focus
+      }
+   }
+
+   const scrollToSelected = (ref: any) => {
+      if (ref && selected && visible) {
+         ref.scrollIntoView({ behavior: 'smooth', block: 'start' });
+      }
+   }
+
+   return (
+      <Stack styles={{ root: { paddingBottom: 6 } }} horizontalAlign="start" onClick={(ev) => { onClickTest() }}>
+         <div onClick={(ev) => { onClickTestName() }} className={styles.itemHover} style={{ width: "100%" }} ref={scrollToSelected}>
+            <Stack horizontal>
+               <Icon styles={{ root: { paddingTop: 7, paddingRight: 8 } }} className={style} iconName="Square" />
+               <Stack.Item disableShrink><Text variant="mediumPlus" >{test.TestDisplayName}</Text></Stack.Item>
+            </Stack>
+         </div>
+         {loading && <Spinner styles={{ root: { padding: 10 } }} size={SpinnerSize.medium}></Spinner>}
+         {testDetails && visible &&
+            <div style={shadowBoxSlyle} onClick={(ev) => { if (selected) ev.stopPropagation() }}>
+               <Stack horizontal tokens={{ childrenGap: 10 }} styles={{ root: { padding: 5 } }}>
+                  <Text styles={{ root: { fontWeight: 'bold' } }}>Filters:</Text>
+                  <Checkbox label="Error" defaultChecked={filterError} onChange={onChangeFilterError} />
+                  <Checkbox label="Warning" defaultChecked={filterWarning} onChange={onChangeFilterWarning} />
+                  <Checkbox label="Info" defaultChecked={filterInfo} onChange={onChangeFilterInfo} />
+                  <Stack onClick={(ev) => { onClickHistory() }} tabIndex={0} onBlur={onHistoryBlur} className={styles.itemHover} style={{ position: "relative" }}>
+                     <Text styles={{ root: { fontWeight: 'bold', paddingLeft: 4, paddingRight: 4 } }}>History</Text>
+                     {historyVisible &&
+                        <div style={{ position: "absolute", width: 200, top: "100%", maxHeight: 300 }} className={styles.historyList} onClick={(ev) => { ev.stopPropagation() }}>
+                           {!historyLoaded && <Spinner styles={{ root: { padding: 10 } }} size={SpinnerSize.small}></Spinner>}
+                           {historyLoaded && history.length > 0 &&
+                              history.map((historyItem) => <HistoryItem key={historyItem.Change} item={historyItem} testName={test.FullTestPath} selected={historyItem.Change === testDataHandler.cursor?.change} />)
+                           }
+                           {historyLoaded && history.length === 0 && <div><Text variant="smallPlus">no history found</Text></div>}
+                        </div>
+                     }
+                  </Stack>
+                  <Text onClick={(ev) => { copyToClipboard(test.FullTestPath); }} styles={{ root: { fontWeight: 'bold', paddingLeft: 4, paddingRight: 4 } }} className={styles.itemHover}>To Clipboard</Text>
+               </Stack>
+               {testDetails.Entries.filter(isEntryNeedDisplay).map((value, index) => <EntryPane key={index} entry={value} testArtifacts={testDetails.Artifacts} />)}
+               {testDetails.Entries.length === 0 && <Text styles={{ root: { padding: 8, fontWeight: 'bold' } }}>No event for this test.</Text>}
             </div>
-            {loading && <Spinner  styles={{ root: { padding: 10 }}} size={SpinnerSize.medium}></Spinner>}
-            {testDetails && visible &&
-                <div style={shadowBoxSlyle} onClick={(ev) => {if (selected) ev.stopPropagation()}}>
-                    <Stack horizontal tokens={{ childrenGap: 10 }} styles={{root: {padding: 5}}}>
-                        <Text styles={{root: {fontWeight: 'bold'}}}>Filters:</Text>
-                        <Checkbox label="Error" defaultChecked={filterError} onChange={onChangeFilterError}/>
-                        <Checkbox label="Warning" defaultChecked={filterWarning} onChange={onChangeFilterWarning}/>
-                        <Checkbox label="Info" defaultChecked={filterInfo} onChange={onChangeFilterInfo}/>
-                        <Stack onClick={(ev) => {onClickHistory()}} tabIndex={0} onBlur={onHistoryBlur} className={styles.itemHover} style={{position: "relative"}}>
-                            <Text styles={{root: {fontWeight: 'bold', paddingLeft: 4, paddingRight: 4}}}>History</Text>
-                            {historyVisible &&
-                                <div style={{position: "absolute", width: 200, top: "100%", maxHeight: 300}} className={styles.historyList} onClick={(ev) => {ev.stopPropagation()}}>
-                                    {!historyLoaded && <Spinner  styles={{ root: { padding: 10 }}} size={SpinnerSize.small}></Spinner>}
-                                    {historyLoaded && history.length > 0 &&
-                                        history.map((historyItem) => <HistoryItem key={historyItem.Change} item={historyItem} testName={test.FullTestPath} selected={historyItem.Change === testDataHandler.cursor?.change}/>)
-                                    }
-                                    {historyLoaded && history.length === 0 && <div><Text variant="smallPlus">no history found</Text></div> }
-                                </div>
-                            }
-                        </Stack>
-                        <Text onClick={(ev) => {copyToClipboard(test.FullTestPath);}} styles={{root: {fontWeight: 'bold', paddingLeft: 4, paddingRight: 4}}} className={styles.itemHover}>To Clipboard</Text>
-                    </Stack>
-                    {testDetails.Entries.filter(isEntryNeedDisplay).map((value, index) => <EntryPane key={index} entry={value} testArtifacts={testDetails.Artifacts}/>)}
-                    {testDetails.Entries.length === 0 && <Text styles={{ root: { padding: 8, fontWeight: 'bold' } }}>No event for this test.</Text>}
-                </div>
-            }
-        </Stack>
-    );
+         }
+      </Stack>
+   );
 }
 
-const TestResultPanel: React.FC<{tests: TestResult[], title?: string, selected?: string}> = (props) => {
-    const { tests, title, selected } = props;
-
-    const suites = new Map<string, TestResult[]>();
-    tests.forEach((test) => {
-        const testSuite = test.FullTestPath.replace('.'+test.TestDisplayName, '');
-        if (!suites.has(testSuite)) {
-            suites.set(testSuite, [test]);
-        } else {
-            suites.get(testSuite)?.push(test);
-        }
-    });
-    const suiteKeys = Array.from(suites.keys());
-
-    return (
-        <Stack styles={{ root: { paddingTop: 18, paddingRight: 0 } }}>
-            <Stack className={hordeClasses.raised}>
-                <Stack tokens={{ childrenGap: 9 }} grow>
-                    <Stack>
-                        <Text variant="mediumPlus" styles={{ root: { fontFamily: "Horde Open Sans SemiBold" } }}>{title} tests ({tests.length})</Text>
-                    </Stack>
-                    {suiteKeys.map((suite) => {
-                        return (<Stack styles={{root: { paddingLeft: 8 }}} tokens={{ childrenGap: 5 }} grow key={`${title}-${suite}`}>
-                            <Text variant="mediumPlus" styles={{root: { color: theme.palette.neutralDark } }}>{suite}</Text>
-                            <Stack styles={{root: { paddingLeft: 8 }}}>
-                                {suites.get(suite)?.map((test: TestResult) => <TestResultPane test={test} selected={selected === test.FullTestPath} key={test.ArtifactName}/>)}
-                            </Stack>
-                        </Stack>)
-                    })}
-                </Stack>
+const TestResultPanel: React.FC<{ tests: TestResult[], title?: string, selected?: string }> = (props) => {
+   const { tests, title, selected } = props;
+
+   const suites = new Map<string, TestResult[]>();
+   tests.forEach((test) => {
+      const testSuite = test.FullTestPath.replace('.' + test.TestDisplayName, '');
+      if (!suites.has(testSuite)) {
+         suites.set(testSuite, [test]);
+      } else {
+         suites.get(testSuite)?.push(test);
+      }
+   });
+   const suiteKeys = Array.from(suites.keys());
+
+   return (
+      <Stack styles={{ root: { paddingTop: 18, paddingRight: 0 } }}>
+         <Stack className={hordeClasses.raised}>
+            <Stack tokens={{ childrenGap: 9 }} grow>
+               <Stack>
+                  <Text variant="mediumPlus" styles={{ root: { fontFamily: "Horde Open Sans SemiBold" } }}>{title} tests ({tests.length})</Text>
+               </Stack>
+               {suiteKeys.map((suite) => {
+                  return (<Stack styles={{ root: { paddingLeft: 8 } }} tokens={{ childrenGap: 5 }} grow key={`${title}-${suite}`}>
+                     <Text variant="mediumPlus" styles={{ root: { color: theme.palette.neutralDark } }}>{suite}</Text>
+                     <Stack styles={{ root: { paddingLeft: 8 } }}>
+                        {suites.get(suite)?.map((test: TestResult) => <TestResultPane test={test} selected={selected === test.FullTestPath} key={test.ArtifactName} />)}
+                     </Stack>
+                  </Stack>)
+               })}
             </Stack>
-        </Stack>
-    );
+         </Stack>
+      </Stack>
+   );
 }
 
-export const TestPassSummaryView: React.FC<{data: TestPassSummary, query: URLSearchParams}> = (props) => {
-    const { data, query } = props;
-    const failedTests : TestResult[] = [];
-    const notrunTests : TestResult[] = [];
-    const inprocessTests : TestResult[] = [];
-    const passedTests : TestResult[] = [];
-    const skippedTests : TestResult[] = [];
-    const TestsMap = new Map<string, TestResult[]>([
-        [TestState.Failed, failedTests],
-        [TestState.NotRun, notrunTests],
-        [TestState.InProcess, inprocessTests],
-        [TestState.Success, passedTests],
-        [TestState.Skipped, skippedTests],
-    ]);
-    data.Tests.forEach((test) => TestsMap.get(test.State)?.push(test));
-
-    const selectedTest = query.get('test')? query.get('test')! : undefined; 
-
-    return (
-        <Stack className={styles.container} styles={{ root: { backgroundColor: "#faf9f9", paddingLeft: 24, paddingTop: 12, paddingRight: 12 } }}>
-
-            <Stack styles={{ root: { paddingTop: 18, paddingRight: 0 } }}>
-                <Stack className={hordeClasses.raised}>
-                    <Stack tokens={{ childrenGap: 12 }} grow>
-                        <Stack>
-                            <Text variant="mediumPlus" styles={{ root: { fontFamily: "Horde Open Sans SemiBold" } }}>Summary</Text>
-                        </Stack>
-                        <Stack styles={{root: { paddingLeft: 8 }}}>
-                            <Text>This test pass run on <span style={{fontWeight: 'bold'}}>{data.ReportCreatedOn}</span> for a duration of <span style={{fontWeight: 'bold'}}>{msecToElapsed(data.TotalDurationSeconds*1000)}</span> on <span style={{fontWeight: 'bold'}}>{getMetadata(data.Metadata, 'Platform')??"Unknown"}</span></Text>
-                            <Text>
-                                {data.FailedCount > 0 && <span><span style={{fontWeight: 'bold'}}>{data.FailedCount}</span> tests <span style={{fontWeight: 'bold'}} className={stateStyles.get('Fail')}>{stateLabels.get('Fail')?.toLowerCase()}</span>. </span>}
-                                {data.InProcessCount > 0 && <span><span style={{fontWeight: 'bold'}}>{data.InProcessCount}</span> tests <span style={{fontWeight: 'bold'}} className={stateStyles.get('InProcess')}>{stateLabels.get('InProcess')?.toLowerCase()}</span>. </span>}
-                                {data.NotRunCount > 0 && <span><span style={{fontWeight: 'bold'}}>{data.NotRunCount}</span> tests <span style={{fontWeight: 'bold'}} className={stateStyles.get('NotRun')}>{stateLabels.get('NotRun')?.toLowerCase()}</span>. </span>}
-                                {data.SucceededCount > 0 &&
-                                    <span>
-                                        <span style={{fontWeight: 'bold'}}>{(data.SucceededCount + data.SucceededWithWarningsCount)}</span> tests <span style={{fontWeight: 'bold'}} className={stateStyles.get('Success')}>{stateLabels.get('Success')?.toLowerCase()}</span>
-                                        {data.SucceededWithWarningsCount > 0 && 
-                                            <span>, including <span style={{fontWeight: 'bold'}}>{data.SucceededWithWarningsCount}</span> with <span style={{fontWeight: 'bold'}} className={stateStyles.get('SuccessWithWarnings')}>warnings</span></span>
-                                        }.
-                                    </span>
-                                }
-                            </Text>
-                            <Stack horizontal wrap>
-                                {data.Tests.map((item: TestResult) => <IconItem item={item} key={item.ArtifactName}/>) }
-                            </Stack>
-                            {data.ReportURL && <Text>Link to: <a href={data.ReportURL}>External report</a></Text>}
-                        </Stack>
-                    </Stack>
-                </Stack>
+export const TestPassSummaryView: React.FC<{ data: TestPassSummary, query: URLSearchParams }> = (props) => {
+   const { data, query } = props;
+   const failedTests: TestResult[] = [];
+   const notrunTests: TestResult[] = [];
+   const inprocessTests: TestResult[] = [];
+   const passedTests: TestResult[] = [];
+   const skippedTests: TestResult[] = [];
+   const TestsMap = new Map<string, TestResult[]>([
+      [TestState.Failed, failedTests],
+      [TestState.NotRun, notrunTests],
+      [TestState.InProcess, inprocessTests],
+      [TestState.Success, passedTests],
+      [TestState.Skipped, skippedTests],
+   ]);
+   data.Tests.forEach((test) => TestsMap.get(test.State)?.push(test));
+
+   const selectedTest = query.get('test') ? query.get('test')! : undefined;
+
+   return (
+      <Stack className={styles.container} styles={{ root: { backgroundColor: "#faf9f9", paddingLeft: 24, paddingTop: 12, paddingRight: 12 } }}>
+
+         <Stack styles={{ root: { paddingTop: 18, paddingRight: 0 } }}>
+            <Stack className={hordeClasses.raised}>
+               <Stack tokens={{ childrenGap: 12 }} grow>
+                  <Stack>
+                     <Text variant="mediumPlus" styles={{ root: { fontFamily: "Horde Open Sans SemiBold" } }}>Summary</Text>
+                  </Stack>
+                  <Stack styles={{ root: { paddingLeft: 8 } }}>
+                     <Text>This test pass run on <span style={{ fontWeight: 'bold' }}>{data.ReportCreatedOn}</span> for a duration of <span style={{ fontWeight: 'bold' }}>{msecToElapsed(data.TotalDurationSeconds * 1000)}</span> on <span style={{ fontWeight: 'bold' }}>{getMetadata(data.Metadata, 'Platform') ?? "Unknown"}</span></Text>
+                     <Text>
+                        {data.FailedCount > 0 && <span><span style={{ fontWeight: 'bold' }}>{data.FailedCount}</span> tests <span style={{ fontWeight: 'bold' }} className={stateStyles.get('Fail')}>{stateLabels.get('Fail')?.toLowerCase()}</span>. </span>}
+                        {data.InProcessCount > 0 && <span><span style={{ fontWeight: 'bold' }}>{data.InProcessCount}</span> tests <span style={{ fontWeight: 'bold' }} className={stateStyles.get('InProcess')}>{stateLabels.get('InProcess')?.toLowerCase()}</span>. </span>}
+                        {data.NotRunCount > 0 && <span><span style={{ fontWeight: 'bold' }}>{data.NotRunCount}</span> tests <span style={{ fontWeight: 'bold' }} className={stateStyles.get('NotRun')}>{stateLabels.get('NotRun')?.toLowerCase()}</span>. </span>}
+                        {data.SucceededCount > 0 &&
+                           <span>
+                              <span style={{ fontWeight: 'bold' }}>{(data.SucceededCount + data.SucceededWithWarningsCount)}</span> tests <span style={{ fontWeight: 'bold' }} className={stateStyles.get('Success')}>{stateLabels.get('Success')?.toLowerCase()}</span>
+                              {data.SucceededWithWarningsCount > 0 &&
+                                 <span>, including <span style={{ fontWeight: 'bold' }}>{data.SucceededWithWarningsCount}</span> with <span style={{ fontWeight: 'bold' }} className={stateStyles.get('SuccessWithWarnings')}>warnings</span></span>
+                              }.
+                           </span>
+                        }
+                     </Text>
+                     <Stack horizontal wrap>
+                        {data.Tests.map((item: TestResult) => <IconItem item={item} key={item.ArtifactName} />)}
+                     </Stack>
+                     {data.ReportURL && <Text>Link to: <a href={data.ReportURL}>External report</a></Text>}
+                  </Stack>
+               </Stack>
             </Stack>
-
-            {failedTests.length > 0 &&
-                <TestResultPanel tests={failedTests} title={stateLabels.get('Fail')} selected={selectedTest}/>
-            }
-            {inprocessTests.length > 0 &&
-                <TestResultPanel tests={inprocessTests} title={stateLabels.get('InProcess')} selected={selectedTest}/>
-            }
-            {notrunTests.length > 0 &&
-                <TestResultPanel tests={notrunTests} title={stateLabels.get('NotRun')} selected={selectedTest}/>
-            }
-            {skippedTests.length > 0 &&
-                <TestResultPanel tests={skippedTests} title={stateLabels.get('Skipped')} selected={selectedTest}/>
-            }
-            {passedTests.length > 0 && 
-                <TestResultPanel tests={passedTests} title={stateLabels.get('Success')} selected={selectedTest}/>
-            }
-        </Stack>
-    );
+         </Stack>
+
+         {failedTests.length > 0 &&
+            <TestResultPanel tests={failedTests} title={stateLabels.get('Fail')} selected={selectedTest} />
+         }
+         {inprocessTests.length > 0 &&
+            <TestResultPanel tests={inprocessTests} title={stateLabels.get('InProcess')} selected={selectedTest} />
+         }
+         {notrunTests.length > 0 &&
+            <TestResultPanel tests={notrunTests} title={stateLabels.get('NotRun')} selected={selectedTest} />
+         }
+         {skippedTests.length > 0 &&
+            <TestResultPanel tests={skippedTests} title={stateLabels.get('Skipped')} selected={selectedTest} />
+         }
+         {passedTests.length > 0 &&
+            <TestResultPanel tests={passedTests} title={stateLabels.get('Success')} selected={selectedTest} />
+         }
+      </Stack>
+   );
 }
diff --git a/Engine/Source/Programs/Horde/HordeDashboard/src/Plugins/UnrealAutomatedTests/index.tsx b/Engine/Source/Programs/Horde/HordeDashboard/src/Plugins/UnrealAutomatedTests/index.tsx
--- a/Engine/Source/Programs/Horde/HordeDashboard/src/Plugins/UnrealAutomatedTests/index.tsx
+++ b/Engine/Source/Programs/Horde/HordeDashboard/src/Plugins/UnrealAutomatedTests/index.tsx
@@ -14,5 +14,6 @@ class UnrealAutomatedTestsPlugin implements Plugin {
 
 }
 
-export default new UnrealAutomatedTestsPlugin();
+const unrealAutomatedTestsPlugin = new UnrealAutomatedTestsPlugin();
+export default unrealAutomatedTestsPlugin;
 
diff --git a/Engine/Source/Programs/Horde/HordeDashboard/src/Plugins/example/index.tsx b/Engine/Source/Programs/Horde/HordeDashboard/src/Plugins/example/index.tsx
--- a/Engine/Source/Programs/Horde/HordeDashboard/src/Plugins/example/index.tsx
+++ b/Engine/Source/Programs/Horde/HordeDashboard/src/Plugins/example/index.tsx
@@ -17,4 +17,5 @@ class ExamplePlugin implements Plugin {
 
 }
 
-export default new ExamplePlugin();
\ No newline at end of file
+const examplePlugin = new ExamplePlugin();
+export default examplePlugin;
\ No newline at end of file
diff --git a/Engine/Source/Programs/Horde/HordeDashboard/src/Plugins/index.tsx b/Engine/Source/Programs/Horde/HordeDashboard/src/Plugins/index.tsx
--- a/Engine/Source/Programs/Horde/HordeDashboard/src/Plugins/index.tsx
+++ b/Engine/Source/Programs/Horde/HordeDashboard/src/Plugins/index.tsx
@@ -68,4 +68,6 @@ export class Plugins {
 
 }
 
-export default new Plugins();
\ No newline at end of file
+const plugins = new Plugins();
+
+export default plugins;
\ No newline at end of file
diff --git a/Engine/Source/Programs/Horde/HordeDashboard/src/backend/AgentStore.ts b/Engine/Source/Programs/Horde/HordeDashboard/src/backend/AgentStore.ts
--- a/Engine/Source/Programs/Horde/HordeDashboard/src/backend/AgentStore.ts
+++ b/Engine/Source/Programs/Horde/HordeDashboard/src/backend/AgentStore.ts
@@ -1,11 +1,15 @@
 // Copyright Epic Games, Inc. All Rights Reserved.
 
-import { action, observable } from 'mobx';
+import { action, makeObservable, observable } from 'mobx';
 import { AgentData, PoolData } from './Api';
 import backend from '.';
 
 export class AgentStore {
 
+    constructor() {
+        makeObservable(this);
+    }
+
     @observable.ref
 	agents: AgentData[] = [];
 	
diff --git a/Engine/Source/Programs/Horde/HordeDashboard/src/backend/Backend.ts b/Engine/Source/Programs/Horde/HordeDashboard/src/backend/Backend.ts
--- a/Engine/Source/Programs/Horde/HordeDashboard/src/backend/Backend.ts
+++ b/Engine/Source/Programs/Horde/HordeDashboard/src/backend/Backend.ts
@@ -1,7 +1,5 @@
 // Copyright Epic Games, Inc. All Rights Reserved.
 
-import { configure } from "mobx";
-import { isNumber } from 'util';
 import templateCache from '../backend/TemplateCache';
 import { AgentData, AgentQuery, ArtifactData, AuditLogEntry, AuditLogQuery, BatchUpdatePoolRequest, ChangeSummaryData, CreateDeviceRequest, CreateDeviceResponse, CreateExternalIssueRequest, CreateExternalIssueResponse, CreateJobRequest, CreateJobResponse, CreateNoticeRequest, CreatePoolRequest, CreateSoftwareResponse, CreateSubscriptionRequest, CreateSubscriptionResponse, DashboardPreference, DevicePoolTelemetryQuery, DeviceTelemetryQuery, EventData, FindIssueResponse, FindJobTimingsResponse, GetAgentSoftwareChannelResponse, GetArtifactZipRequest, GetDashboardConfigResponse, GetDevicePlatformResponse, GetDevicePoolResponse, GetDevicePoolTelemetryResponse, GetDeviceReservationResponse, GetDeviceResponse, GetDeviceTelemetryResponse, GetExternalIssueProjectResponse, GetExternalIssueResponse, GetGraphResponse, GetIssueStreamResponse, GetJobsTabResponse, GetJobStepRefResponse, GetJobStepTraceResponse, GetJobTimingResponse, GetLogEventResponse, GetNoticeResponse, GetNotificationResponse, GetPerforceServerStatusResponse, GetPoolResponse, GetServerInfoResponse, GetServerSettingsResponse, GetSoftwareResponse, GetSubscriptionResponse, GetTestDataDetailsResponse, GetTestDataRefResponse, GetTestMetaResponse, GetTestResponse, GetTestsRequest, GetTestStreamResponse, GetUserResponse, GetUtilizationTelemetryResponse, GlobalConfig, IssueData, IssueQuery, IssueQueryV2, JobData, JobQuery, JobsTabColumnType, JobStepOutcome, JobStreamQuery, JobTimingsQuery, LeaseData, LogData, LogLineData, PoolData, ProjectData, ScheduleData, ScheduleQuery, SearchLogFileResponse, ServerUpdateResponse, SessionData, StreamData, TabType, TestData, UpdateAgentRequest, UpdateDeviceRequest, UpdateGlobalConfigRequest, UpdateIssueRequest, UpdateJobRequest, UpdateLeaseRequest, UpdateNoticeRequest, UpdateNotificationsRequest, UpdatePoolRequest, UpdateServerSettingsRequest, UpdateStepRequest, UpdateStepResponse, UpdateTemplateRefRequest, UpdateUserRequest, UsersQuery } from './Api';
 import dashboard from './Dashboard';
@@ -10,12 +8,6 @@ import graphCache, { GraphQuery } from './GraphCache';
 import { projectStore } from './ProjectStore';
 
 
-// configure mobx
-configure({
-    reactionRequiresObservable: true,
-    enforceActions: "observed"
-});
-
 // Update interval for relatively static global data such as projects, schedules, templates
 const updateInterval = 120 * 1000;
 
@@ -605,11 +597,11 @@ export class Backend {
 
         const params: any = {};
 
-        if (isNumber(index)) {
+        if (Number.isInteger(index)) {
             params.index = index;
         }
 
-        if (isNumber(count)) {
+        if (Number.isInteger(count)) {
             params.Count = count;
         }
 
diff --git a/Engine/Source/Programs/Horde/HordeDashboard/src/backend/CommitCache.ts b/Engine/Source/Programs/Horde/HordeDashboard/src/backend/CommitCache.ts
--- a/Engine/Source/Programs/Horde/HordeDashboard/src/backend/CommitCache.ts
+++ b/Engine/Source/Programs/Horde/HordeDashboard/src/backend/CommitCache.ts
@@ -1,10 +1,14 @@
 // Copyright Epic Games, Inc. All Rights Reserved.
-import { action, observable } from 'mobx';
+import { action, makeObservable, observable } from 'mobx';
 import backend from '.';
 import { GetChangeSummaryResponse } from "./Api";
 
 export class CommitCache {
 
+    constructor() {
+        makeObservable(this);
+    }
+
     @action
     setUpdated() {
         this.updated++;
diff --git a/Engine/Source/Programs/Horde/HordeDashboard/src/backend/Datadog.ts b/Engine/Source/Programs/Horde/HordeDashboard/src/backend/Datadog.ts
--- a/Engine/Source/Programs/Horde/HordeDashboard/src/backend/Datadog.ts
+++ b/Engine/Source/Programs/Horde/HordeDashboard/src/backend/Datadog.ts
@@ -1,6 +1,6 @@
 // Copyright Epic Games, Inc. All Rights Reserved.
 
-import { datadogLogs, StatusType, Datacenter } from '@datadog/browser-logs';
+import { datadogLogs, StatusType } from '@datadog/browser-logs';
 import { getSiteConfig } from '../backend/Config';
 
 export const initDatadog = () => {
@@ -13,7 +13,7 @@ export const initDatadog = () => {
 
 	datadogLogs.init({
 		clientToken: siteConfig.datadogClient,
-		datacenter: Datacenter.US,
+		//datacenter: Datacenter.US,
 		env: siteConfig.environment === "production" ? "prod" : "dev",
 		forwardErrorsToLogs: true,
 		sampleRate: 100
