Horde: Updating dashboard dependencies to latest versions and cleanup

diff --git a/Engine/Source/Programs/Horde/HordeDashboard/src/backend/GraphCache.ts b/Engine/Source/Programs/Horde/HordeDashboard/src/backend/GraphCache.ts
--- a/Engine/Source/Programs/Horde/HordeDashboard/src/backend/GraphCache.ts
+++ b/Engine/Source/Programs/Horde/HordeDashboard/src/backend/GraphCache.ts
@@ -84,5 +84,6 @@ class GraphDatabase extends LocalCache<GetGraphResponse> {
 }
 
 const graphDatabase = new GraphDatabase();
+const graphCache = new GraphCache();
 
-export default new GraphCache();
\ No newline at end of file
+export default graphCache;
\ No newline at end of file
diff --git a/Engine/Source/Programs/Horde/HordeDashboard/src/backend/JobDetails.ts b/Engine/Source/Programs/Horde/HordeDashboard/src/backend/JobDetails.ts
--- a/Engine/Source/Programs/Horde/HordeDashboard/src/backend/JobDetails.ts
+++ b/Engine/Source/Programs/Horde/HordeDashboard/src/backend/JobDetails.ts
@@ -1,6 +1,6 @@
 // Copyright Epic Games, Inc. All Rights Reserved.
 
-import { action, observable } from 'mobx';
+import { action, makeObservable, observable } from 'mobx';
 import { getTheme, mergeStyles, mergeStyleSets } from '@fluentui/react/lib/Styling';
 import backend from '.';
 import { getBatchInitElapsed, getNiceTime, getStepElapsed, getStepETA, getStepFinishTime, getStepTimingDelta } from '../base/utilities/timeUtils';
@@ -26,6 +26,8 @@ export class JobDetails {
 
     constructor(id?: string, logId?: string, stepId?: string, isLogView?: boolean) {
 
+        makeObservable(this);
+
         if (id) {
             this.set(id, logId, stepId, undefined);
         }
diff --git a/Engine/Source/Programs/Horde/HordeDashboard/src/backend/JobHandler.ts b/Engine/Source/Programs/Horde/HordeDashboard/src/backend/JobHandler.ts
--- a/Engine/Source/Programs/Horde/HordeDashboard/src/backend/JobHandler.ts
+++ b/Engine/Source/Programs/Horde/HordeDashboard/src/backend/JobHandler.ts
@@ -1,6 +1,6 @@
 // Copyright Epic Games, Inc. All Rights Reserved.
 
-import { action, observable } from 'mobx';
+import { action, makeObservable, observable } from 'mobx';
 import moment from 'moment';
 import backend from '../backend';
 import { JobData, JobState, JobStepOutcome, JobStreamQuery, StreamData } from '../backend/Api';
@@ -15,6 +15,7 @@ const jobsRefreshTime = 10000;
 export class JobHandler {
 
     constructor(includeBatches: boolean = false, jobLimit?: number) {
+        makeObservable(this);
         this.includeBatches = includeBatches;
         this.jobLimit = jobLimit;
     }
diff --git a/Engine/Source/Programs/Horde/HordeDashboard/src/backend/PollBase.ts b/Engine/Source/Programs/Horde/HordeDashboard/src/backend/PollBase.ts
--- a/Engine/Source/Programs/Horde/HordeDashboard/src/backend/PollBase.ts
+++ b/Engine/Source/Programs/Horde/HordeDashboard/src/backend/PollBase.ts
@@ -1,12 +1,12 @@
 // Copyright Epic Games, Inc. All Rights Reserved.
 
-import { action, observable } from 'mobx';
+import { action, makeObservable, observable } from 'mobx';
 
 // base class for polling behavior
 export abstract class PollBase {
 
     constructor(pollTime = 5000) {
-
+        makeObservable(this);
         this.pollTime = pollTime;
     }
 
diff --git a/Engine/Source/Programs/Horde/HordeDashboard/src/backend/TemplateCache.ts b/Engine/Source/Programs/Horde/HordeDashboard/src/backend/TemplateCache.ts
--- a/Engine/Source/Programs/Horde/HordeDashboard/src/backend/TemplateCache.ts
+++ b/Engine/Source/Programs/Horde/HordeDashboard/src/backend/TemplateCache.ts
@@ -16,4 +16,5 @@ export class TemplateCache {
 	}
 }
 
-export default new TemplateCache();
\ No newline at end of file
+const templateCache = new TemplateCache();
+export default templateCache;
\ No newline at end of file
diff --git a/Engine/Source/Programs/Horde/HordeDashboard/src/backend/TestDataHandler.ts b/Engine/Source/Programs/Horde/HordeDashboard/src/backend/TestDataHandler.ts
--- a/Engine/Source/Programs/Horde/HordeDashboard/src/backend/TestDataHandler.ts
+++ b/Engine/Source/Programs/Horde/HordeDashboard/src/backend/TestDataHandler.ts
@@ -1,6 +1,6 @@
 // Copyright Epic Games, Inc. All Rights Reserved.
 
-import { observable, action, when } from "mobx";
+import { observable, action, when, makeObservable } from "mobx";
 import backend from '../backend';
 import { JobData, ArtifactData, TestData } from './Api';
 
@@ -13,7 +13,7 @@ export type DataWrapper = {
 const updateObject = (target: any, version: any, skipKeys?: string[]) => {
     const properties = Object.getOwnPropertyDescriptors(version);
     Object.getOwnPropertyNames(version).forEach((key) => {
-        if(!properties[key].value || skipKeys?.includes(key)) {
+        if (!properties[key].value || skipKeys?.includes(key)) {
             delete properties[key];
         }
     });
@@ -45,13 +45,13 @@ export class TestDataWrapper implements TestData {
     }
 
     setDataHandler(handler: DataWrapper) {
-        if(this.data) {
+        if (this.data) {
             this._datahandler = handler;
         }
     }
 
     getDataHandler(): DataWrapper | undefined {
-        return this._datahandler ;
+        return this._datahandler;
     }
 
     getTestDataProperty<T>(key: string): T {
@@ -78,7 +78,7 @@ export class TestDataWrapper implements TestData {
             while (tailKey.indexOf('.') > 0) {
                 [headKey, tailKey] = tailKey.split('.', 2);
                 const nextProp = Object.getOwnPropertyDescriptor(obj, headKey);
-                if(!nextProp) {
+                if (!nextProp) {
                     return false;
                 }
                 obj = nextProp?.value;
@@ -103,7 +103,7 @@ export class TestDataWrapper implements TestData {
 
         const jobArtifacts = await this.getArtifacts();
 
-        artifactName = artifactName.replace(/\\/g,  '/');
+        artifactName = artifactName.replace(/\\/g, '/');
 
         if (!this.artifactMap) {
             this.artifactMap = new Map();
@@ -120,10 +120,10 @@ export class TestDataWrapper implements TestData {
     }
 
     async getJobStepName() {
-        if(this.stepName) {
+        if (this.stepName) {
             return this.stepName;
         }
-        if(!this.jobdata) {
+        if (!this.jobdata) {
             this.jobdata = await backend.getJob(this.jobId);
         }
 
@@ -144,7 +144,7 @@ export class TestDataWrapper implements TestData {
     }
 
     async getJobName() {
-        if(!this.jobdata) {
+        if (!this.jobdata) {
             this.jobdata = await backend.getJob(this.jobId);
         }
 
@@ -156,7 +156,7 @@ export class TestDataWrapper implements TestData {
     }
 
     async preflighChange() {
-        if(!this.jobdata) {
+        if (!this.jobdata) {
             this.jobdata = await backend.getJob(this.jobId);
         }
 
@@ -164,7 +164,7 @@ export class TestDataWrapper implements TestData {
     }
 
     async getChangeName() {
-        if(!this.jobdata) {
+        if (!this.jobdata) {
             this.jobdata = await backend.getJob(this.jobId);
         }
 
@@ -178,16 +178,20 @@ export class TestDataWrapper implements TestData {
         } else {
             changeName = `${data.change ? "CL " + data.change : "Latest CL"}`;
         }
-        
+
         return changeName;
     }
 }
 
-export type OnFetchTestData = (newChunk: (TestDataWrapper|undefined)[]) => void;
+export type OnFetchTestData = (newChunk: (TestDataWrapper | undefined)[]) => void;
 export type SelectorTestData = (item: TestDataWrapper) => boolean;
 
 export class TestDataCollection {
 
+    constructor() {
+        makeObservable(this);
+    }
+
     activated = false;
     setActive(activated: boolean) {
         this.activated = activated;
@@ -228,14 +232,14 @@ export class TestDataCollection {
     }
 
     findById(id: string) {
-        if(!this.items) {
+        if (!this.items) {
             return undefined;
         }
         return this.items.get(id);
     }
 
     forEach(func: (item: TestDataWrapper) => void) {
-        if(!this.items) {
+        if (!this.items) {
             return;
         }
         this.items.forEach(func);
@@ -243,10 +247,10 @@ export class TestDataCollection {
 
     *iterItems() {
         const items = this.items;
-        if(items) {
+        if (items) {
             const iterator = items.values();
             let result = iterator.next();
-            while(!result.done) {
+            while (!result.done) {
                 yield result.value;
                 result = iterator.next();
             }
@@ -254,7 +258,7 @@ export class TestDataCollection {
     }
 
     *iterItemsByChange() {
-        if(this.itemsByChange) {
+        if (this.itemsByChange) {
             const changes = this.getChanges();
             for (const change of changes) {
                 const items = this.itemsByChange.get(change);
@@ -266,24 +270,24 @@ export class TestDataCollection {
     }
 
     getChanges() {
-        if(this.itemsByChange) {
+        if (this.itemsByChange) {
             const changes = Array.from(this.itemsByChange.keys());
-            changes.sort((a, b) => a > b? -1 : (a < b? 1 : 0)); // in reversed order
+            changes.sort((a, b) => a > b ? -1 : (a < b ? 1 : 0)); // in reversed order
             return changes;
         }
         return [];
     }
 
-    async setFromStream(streamId: string, key: string, onFetch?: OnFetchTestData, filter?: string, maxCount? : number) {
+    async setFromStream(streamId: string, key: string, onFetch?: OnFetchTestData, filter?: string, maxCount?: number) {
         this.setActive(true);
 
         if (key === this.key && streamId === this.streamId && !this.cursor) {
             // reprocessed what was loaded already, and then continue loading if max count was not reached
-            if(this.items) {
-                if(onFetch) {
+            if (this.items) {
+                if (onFetch) {
                     onFetch(Array.from(this.items.values()));
                 }
-                if(this.items.size%this.fetchCount !== 0 || this.items.size >= this.fetchMaxCount) {
+                if (this.items.size % this.fetchCount !== 0 || this.items.size >= this.fetchMaxCount) {
                     return;
                 }
             }
@@ -301,10 +305,10 @@ export class TestDataCollection {
 
         if (this.items) {
             const item = this.items.get(id);
-            if(item) {
+            if (item) {
                 const handler = item.getDataHandler();
                 if (handler) {
-                    if(!handler.isAllPropertiesLoaded()){
+                    if (!handler.isAllPropertiesLoaded()) {
                         const missing = handler.getMissingProperties();
                         const filter = missing.map((item) => `data.${item}`).join(',');
                         const count = await this.fetchUpdateItem(id, undefined, filter);
@@ -320,10 +324,10 @@ export class TestDataCollection {
 
         // pre fetch minimal data to get the context and see if we need to clean the cache
         const item = await backend.getTestData(id, "nodata");
-        if(!item) {
+        if (!item) {
             return undefined;
         }
-        if(item.key !== this.key || item.streamId !== this.streamId) {
+        if (item.key !== this.key || item.streamId !== this.streamId) {
             this.cleanCache();
             this.key = item.key;
             this.streamId = item.streamId;
@@ -346,7 +350,7 @@ export class TestDataCollection {
 
         await this.fetchCursorHistory(onFetch, filter, maxCount, selector);
 
-        const mapper = selector? (items: TestDataWrapper[]) => items.find(selector) as TestDataWrapper : (items: TestDataWrapper[]) => items[0];
+        const mapper = selector ? (items: TestDataWrapper[]) => items.find(selector) as TestDataWrapper : (items: TestDataWrapper[]) => items[0];
 
         return Array.from(this.iterItemsByChange()).map(mapper);
     }
@@ -358,24 +362,24 @@ export class TestDataCollection {
 
         if (!this.itemsByChange || this.itemsByChange.size === 1) {
             const key = `cursorHistory-${this.key}`;
-            if(this.isLoading(key)) {
+            if (this.isLoading(key)) {
                 await when(() => !this.loadingNotice.has(key));
                 return;
             }
             this.acquireLoadingNotice(key);
-            await this.queueFetchItems(onFetch, filter, maxCount??this.fetchCount);
+            await this.queueFetchItems(onFetch, filter, maxCount ?? this.fetchCount);
             this.releaseLoadingNotice(key);
         }
     }
 
-    async queueFetchItems(onFetch?: OnFetchTestData, filter?: string, maxCount? : number) {
+    async queueFetchItems(onFetch?: OnFetchTestData, filter?: string, maxCount?: number) {
         if (maxCount) {
             this.fetchMaxCount = maxCount;
         }
 
         let totalCount = await this.fetchItems(onFetch, filter);
         let count = totalCount;
-        while (this.activated && count > 0 && count%this.fetchCount === 0 && this.items && this.items.size < this.fetchMaxCount) {
+        while (this.activated && count > 0 && count % this.fetchCount === 0 && this.items && this.items.size < this.fetchMaxCount) {
             count = await this.fetchItems(onFetch, filter);
             totalCount += count;
         }
@@ -395,7 +399,7 @@ export class TestDataCollection {
     }
 
     private initiateCache() {
-        if(!this.items) {
+        if (!this.items) {
             this.items = new Map();
             this.itemsByChange = new Map();
         }
@@ -407,7 +411,7 @@ export class TestDataCollection {
         }
         const change = item.change;
         const changeCollection = this.itemsByChange.get(change);
-        if(changeCollection) {
+        if (changeCollection) {
             changeCollection.push(item);
         } else {
             this.itemsByChange.set(change, [item])
@@ -423,15 +427,15 @@ export class TestDataCollection {
         let count = 0;
         try {
             const pageData = await backend.getTestDataHistory(streamId, key, undefined, this.fetchCount, this.nextIndex, filter);
-            this.nextIndex += this.fetchCount; 
+            this.nextIndex += this.fetchCount;
 
             // Are we still fetching for the same collection?
-            if(key === this.key && streamId === this.streamId) {
+            if (key === this.key && streamId === this.streamId) {
 
                 this.initiateCache();
 
                 const wrappedData = pageData.map((item) => {
-                    if(!this.items?.has(item.id)) {
+                    if (!this.items?.has(item.id)) {
                         const itemWrapped = new TestDataWrapper(item);
                         this.items?.set(item.id, itemWrapped);
                         this.updateItemsByChange(itemWrapped);
@@ -448,7 +452,7 @@ export class TestDataCollection {
                 count = pageData.length;
             }
 
-        } catch(reason) {
+        } catch (reason) {
             console.error(`Failed to fetch TestData history from ${streamId} on key '${key}' !`);
             console.error(reason);
         }
@@ -474,9 +478,9 @@ export class TestDataCollection {
         if (!this.key || !this.streamId) {
             return 0;
         }
-        if(this.loadingNotice.has(id)) {
+        if (this.loadingNotice.has(id)) {
             await when(() => !this.loadingNotice.has(id));
-            return this.activated? 1 : 0;
+            return this.activated ? 1 : 0;
         }
         const key = this.key;
         const streamId = this.streamId;
@@ -486,12 +490,12 @@ export class TestDataCollection {
             const update = await backend.getTestData(id, filter);
 
             // Are we still fetching for the same collection?
-            if(key === this.key && streamId === this.streamId) {
+            if (key === this.key && streamId === this.streamId) {
 
                 this.initiateCache();
 
                 let item = this.items?.get(id);
-                if(item) {
+                if (item) {
                     item.updateProperties(update);
                 } else {
                     item = new TestDataWrapper(update);
@@ -506,7 +510,7 @@ export class TestDataCollection {
 
                 count = 1;
             }
-        } catch(reason) {
+        } catch (reason) {
             console.error(`Failed to fetch updated TestData item ${id} from ${streamId} on key '${key}' !`);
         } finally { this.releaseLoadingNotice(id) }
 
@@ -517,18 +521,18 @@ export class TestDataCollection {
         if (!this.key || !this.streamId) {
             return [];
         }
-        if(this.loadingNotice.has(altKey)) {
+        if (this.loadingNotice.has(altKey)) {
             await when(() => !this.loadingNotice.has(altKey));
             return []; // We have no way to know what is the results... :/
         }
 
         const streamId = this.streamId;
-        let results: TestData[]= [];
+        let results: TestData[] = [];
 
         this.acquireLoadingNotice(altKey);
         try {
-                results = await backend.getTestDataHistory(streamId, altKey, maxChange, count, page, filter);
-        } catch(reason) {
+            results = await backend.getTestDataHistory(streamId, altKey, maxChange, count, page, filter);
+        } catch (reason) {
             console.error(`Failed to fetch Test Data for ${altKey} from ${streamId} !`);
         } finally { this.releaseLoadingNotice(altKey) }
 
diff --git a/Engine/Source/Programs/Horde/HordeDashboard/src/backend/UserCache.ts b/Engine/Source/Programs/Horde/HordeDashboard/src/backend/UserCache.ts
--- a/Engine/Source/Programs/Horde/HordeDashboard/src/backend/UserCache.ts
+++ b/Engine/Source/Programs/Horde/HordeDashboard/src/backend/UserCache.ts
@@ -1,11 +1,15 @@
 
-import { action, observable } from 'mobx';
+import { action, makeObservable, observable } from 'mobx';
 import backend from '.';
 import { GetUserResponse } from './Api';
 
 
 class UserCache {
 
+    constructor() {
+        makeObservable(this);
+    }
+
     @action
     setUpdated() {
         this.updated++;
@@ -52,4 +56,5 @@ class UserCache {
 
 }
 
-export default new UserCache();
\ No newline at end of file
+const userCache = new UserCache()
+export default userCache;
\ No newline at end of file
diff --git a/Engine/Source/Programs/Horde/HordeDashboard/src/backend/index.ts b/Engine/Source/Programs/Horde/HordeDashboard/src/backend/index.ts
--- a/Engine/Source/Programs/Horde/HordeDashboard/src/backend/index.ts
+++ b/Engine/Source/Programs/Horde/HordeDashboard/src/backend/index.ts
@@ -5,7 +5,8 @@ import { Backend } from './Backend';
 import { projectStore } from './ProjectStore';
 import { issueStore } from './IssueStore';
 
-export default new Backend();
+const backend = new Backend();
+export default backend;
 
 const backendContext = React.createContext({
   projectStore: projectStore,
diff --git a/Engine/Source/Programs/Horde/HordeDashboard/src/base/components/Markdown/Markdown.tsx b/Engine/Source/Programs/Horde/HordeDashboard/src/base/components/Markdown/Markdown.tsx
--- a/Engine/Source/Programs/Horde/HordeDashboard/src/base/components/Markdown/Markdown.tsx
+++ b/Engine/Source/Programs/Horde/HordeDashboard/src/base/components/Markdown/Markdown.tsx
@@ -1,115 +1,130 @@
 import * as React from 'react';
-import MarkdownToJsx, { MarkdownProps } from 'markdown-to-jsx';
-import { DefaultButton, Image, IImageStyles, classNamesFunction, IStyleFunction, styled } from '@fluentui/react';
+import type MarkdownComponentType from 'markdown-to-jsx';
+import type { MarkdownToJSX } from 'markdown-to-jsx';
+import * as MarkdownModule from 'markdown-to-jsx';
+import { Image, IImageStyles, classNamesFunction, IStyleFunction, styled } from '@fluentui/react';
+import { DefaultButton } from '@fluentui/react/lib/Button';
 import * as MDTable from '../MarkdownTable/index';
 import { MarkdownHeader } from './MarkdownHeader';
 import { MarkdownParagraph } from './MarkdownParagraph';
 import { IMarkdownProps, IMarkdownSubComponentStyles, IMarkdownStyleProps, IMarkdownStyles } from './Markdown.types';
 import { MarkdownLink } from './MarkdownLink';
+import { MarkdownPre } from './MarkdownPre';
+import { PropsWithChildren } from 'react';
 
-const getStyles: IStyleFunction<IMarkdownStyleProps, IMarkdownStyles> = props => {
+// This is to work around inconsistency between the way markdown-to-jsx declares its types
+// (as having a default export) and the way it actually builds its files (for its cjs `main` file,
+// assigning to module.exports and not setting a default export)
+const MarkdownComponent: typeof MarkdownComponentType =
+  // eslint-disable-next-line @typescript-eslint/no-explicit-any
+  (MarkdownModule as any).default || (MarkdownModule as any);
+
+const getStyles: IStyleFunction<IMarkdownStyleProps, IMarkdownStyles> = () => {
   const imageStyles: Partial<IImageStyles> = {
     root: {
       maxWidth: '100%',
-      margin: '8px 0'
-    }
+      margin: '8px 0',
+    },
   };
 
   return {
     root: 'ms-Markdown',
-    subComponentStyles: { image: imageStyles } as IMarkdownSubComponentStyles
+    subComponentStyles: { image: imageStyles } as IMarkdownSubComponentStyles,
   };
 };
 
 const getClassNames = classNamesFunction<IMarkdownStyleProps, IMarkdownStyles>();
 
-const MarkdownBase: React.StatelessComponent<IMarkdownProps> & { displayName?: string } = props => {
+const MarkdownBase: React.FunctionComponent<PropsWithChildren<IMarkdownProps>> = props => {
   const { styles, theme, children } = props;
   const classNames = getClassNames(styles, { theme: theme! });
+
   return (
     <div className={classNames.root}>
-      <MarkdownToJsx {...getMarkdownProps(classNames.subComponentStyles)}>{children}</MarkdownToJsx>
+      <MarkdownComponent options={{ overrides: getOverrides(classNames.subComponentStyles, props) }}>
+        {children as string}
+      </MarkdownComponent>
     </div>
   );
 };
 MarkdownBase.displayName = 'Markdown';
 
-function getMarkdownProps(subComponentStyles: IMarkdownSubComponentStyles): MarkdownProps {
+function getOverrides(subComponentStyles: IMarkdownSubComponentStyles, props: IMarkdownProps): MarkdownToJSX.Overrides {
   return {
-    options: {
-      overrides: {
-        h1: {
-          component: MarkdownHeader,
-          props: { styles: subComponentStyles.header }
-        },
-        h2: {
-          component: MarkdownHeader,
-          props: { as: 'h2', styles: subComponentStyles.header }
-        },
-        h3: {
-          component: MarkdownHeader,
-          props: { as: 'h3', styles: subComponentStyles.header }
-        },
-        h4: {
-          component: MarkdownHeader,
-          props: { as: 'h4', styles: subComponentStyles.header }
-        },
-        h5: {
-          component: MarkdownHeader,
-          props: { as: 'h5', styles: subComponentStyles.header }
-        },
-        h6: {
-          component: MarkdownHeader,
-          props: { as: 'h6', styles: subComponentStyles.header }
-        },
-        p: {
-          component: MarkdownParagraph,
-          props: { styles: subComponentStyles.paragraph }
-        },
-        a: {
-          component: MarkdownLink,
-          props: { className: 'ms-mdLink', styles: subComponentStyles.link }
-        },
-        img: {
-          component: Image,
-          props: { className: 'ms-mdImage', styles: subComponentStyles.image }
-        },
-        button: {
-          component: DefaultButton,
-          props: { className: 'ms-mdButton', styles: subComponentStyles.button }
-        },
-        table: {
-          component: MDTable.MarkdownTable,
-          props: { styles: subComponentStyles.table }
-        },
-        thead: {
-          component: MDTable.MarkdownTHead,
-          props: { styles: subComponentStyles.table }
-        },
-        tbody: {
-          component: MDTable.MarkdownTBody,
-          props: { styles: subComponentStyles.table }
-        },
-        tr: {
-          component: MDTable.MarkdownTr,
-          props: { styles: subComponentStyles.table }
-        },
-        th: {
-          component: MDTable.MarkdownCell,
-          props: { as: 'th', styles: subComponentStyles.table }
-        },
-        td: {
-          component: MDTable.MarkdownCell,
-          props: { as: 'td', styles: subComponentStyles.table }
-        }
-      }
-    }
+    h1: {
+      component: MarkdownHeader,
+      props: { styles: subComponentStyles.header },
+    },
+    h2: {
+      component: MarkdownHeader,
+      props: { as: 'h2', styles: subComponentStyles.header },
+    },
+    h3: {
+      component: MarkdownHeader,
+      props: { as: 'h3', styles: subComponentStyles.header },
+    },
+    h4: {
+      component: MarkdownHeader,
+      props: { as: 'h4', styles: subComponentStyles.header },
+    },
+    h5: {
+      component: MarkdownHeader,
+      props: { as: 'h5', styles: subComponentStyles.header },
+    },
+    h6: {
+      component: MarkdownHeader,
+      props: { as: 'h6', styles: subComponentStyles.header },
+    },
+    p: {
+      component: MarkdownParagraph,
+      props: { styles: subComponentStyles.paragraph },
+    },
+    pre: {
+      component: MarkdownPre,
+      props: { enableRenderHtmlBlock: props.enableRenderHtmlBlock },
+    },
+    a: {
+      component: MarkdownLink,
+      props: { className: 'ms-mdLink', styles: subComponentStyles.link },
+    },
+    img: {
+      component: Image,
+      props: { className: 'ms-mdImage', styles: subComponentStyles.image },
+    },
+    button: {
+      component: DefaultButton,
+      props: { className: 'ms-mdButton', styles: subComponentStyles.button },
+    },
+    table: {
+      component: MDTable.MarkdownTable,
+      props: { styles: subComponentStyles.table },
+    },
+    thead: {
+      component: MDTable.MarkdownTHead,
+      props: { styles: subComponentStyles.table },
+    },
+    tbody: {
+      component: MDTable.MarkdownTBody,
+      props: { styles: subComponentStyles.table },
+    },
+    tr: {
+      component: MDTable.MarkdownTr,
+      props: { styles: subComponentStyles.table },
+    },
+    th: {
+      component: MDTable.MarkdownCell,
+      props: { as: 'th', styles: subComponentStyles.table },
+    },
+    td: {
+      component: MDTable.MarkdownCell,
+      props: { as: 'td', styles: subComponentStyles.table },
+    },
+    ...props.overrides,
   };
 }
 
-export const Markdown: React.StatelessComponent<IMarkdownProps> = styled<IMarkdownProps, IMarkdownStyleProps, IMarkdownStyles>(
-  MarkdownBase,
-  getStyles,
-  undefined,
-  { scope: 'Markdown' }
-);
+export const Markdown: React.FunctionComponent<PropsWithChildren<IMarkdownProps>> = styled<
+  IMarkdownProps,
+  IMarkdownStyleProps,
+  IMarkdownStyles
+>(MarkdownBase, getStyles, undefined, { scope: 'Markdown' });
diff --git a/Engine/Source/Programs/Horde/HordeDashboard/src/base/components/Markdown/Markdown.types.ts b/Engine/Source/Programs/Horde/HordeDashboard/src/base/components/Markdown/Markdown.types.ts
--- a/Engine/Source/Programs/Horde/HordeDashboard/src/base/components/Markdown/Markdown.types.ts
+++ b/Engine/Source/Programs/Horde/HordeDashboard/src/base/components/Markdown/Markdown.types.ts
@@ -1,11 +1,21 @@
-import { ITheme, IStyleFunctionOrObject, IStyle, ILinkStyleProps, IImageStyleProps } from '@fluentui/react';
-import { IMarkdownHeaderStyleProps } from './MarkdownHeader';
-import { IMarkdownParagraphStyleProps } from './MarkdownParagraph';
-import { IMarkdownTableStyleProps } from '../MarkdownTable/index';
+import type { ITheme, IStyleFunctionOrObject, IStyle, ILinkStyleProps, IImageStyleProps } from '@fluentui/react';
+import type { IMarkdownHeaderStyleProps } from './MarkdownHeader';
+import type { IMarkdownParagraphStyleProps } from './MarkdownParagraph';
+import type { IMarkdownTableStyleProps } from '../MarkdownTable/index';
+import type { MarkdownToJSX } from 'markdown-to-jsx';
 
 export interface IMarkdownProps {
+  /** CSS class to apply to the component root */
   className?: string;
-  children?: React.ReactNode;
+
+  /**
+   * If true, using a code block with language name `renderhtml` will render the contents as HTML.
+   * This is to work around markdown-to-jsx's limited support for nested HTML elements.
+   */
+  enableRenderHtmlBlock?: boolean;
+
+  /** Additional component overrides for markdown rendering */
+  overrides?: MarkdownToJSX.Overrides;
 
   /** Theme provided by higher-order component. */
   theme?: ITheme;
@@ -23,7 +33,7 @@ export interface IMarkdownStyles {
 
 export interface IMarkdownSubComponentStyles {
   // TODO: remove anys
-  // tslint:disable:no-any
+  /* eslint-disable @typescript-eslint/no-explicit-any */
   header: IStyleFunctionOrObject<IMarkdownHeaderStyleProps, any>;
   paragraph: IStyleFunctionOrObject<IMarkdownParagraphStyleProps, any>;  
   link: IStyleFunctionOrObject<ILinkStyleProps, any>;
diff --git a/Engine/Source/Programs/Horde/HordeDashboard/src/base/components/Markdown/MarkdownHeader.tsx b/Engine/Source/Programs/Horde/HordeDashboard/src/base/components/Markdown/MarkdownHeader.tsx
--- a/Engine/Source/Programs/Horde/HordeDashboard/src/base/components/Markdown/MarkdownHeader.tsx
+++ b/Engine/Source/Programs/Horde/HordeDashboard/src/base/components/Markdown/MarkdownHeader.tsx
@@ -1,6 +1,7 @@
 import * as React from 'react';
 import { classNamesFunction, IStyleFunctionOrObject, styled, IStyleFunction } from '@fluentui/react/lib/Utilities';
 import { ITheme, IStyle } from '@fluentui/react/lib/Styling';
+import { FontSizes, NeutralColors } from '@fluentui/theme';
 
 /**
  * The component props.
@@ -35,7 +36,7 @@ const getStyles: IStyleFunction<IMarkdownHeaderStyleProps, IMarkdownHeaderStyles
   return {
     root: [
       {
-        fontSize: "16",
+        fontSize: FontSizes.size16,
         fontWeight: 600,
         padding: 0,
         margin: 0,
@@ -44,39 +45,49 @@ const getStyles: IStyleFunction<IMarkdownHeaderStyleProps, IMarkdownHeaderStyles
 
         selectors: {
           '&:first-child': {
-            marginTop: 0
-          }
-        }
+            marginTop: 0,
+          },
+        },
       },
 
       props.as === 'h1' && [
         {
-          fontSize: "32",
-          marginBottom: '28px'
-        }
+          fontSize: FontSizes.size32,
+          marginBottom: '28px',
+          color: NeutralColors.gray160,
+        },
       ],
 
       props.as === 'h2' && [
         {
-          fontSize: "24",
-          marginBottom: '20px'
-        }
+          fontSize: FontSizes.size24,
+          marginBottom: '20px',
+          color: NeutralColors.gray160,
+        },
       ],
 
       props.as === 'h3' && [
         {
-          fontSize: "20",
-          marginBottom: '8px'
-        }
+          fontSize: FontSizes.size20,
+          marginBottom: '8px',
+          color: NeutralColors.gray130,
+        },
       ],
-      className
-    ]
+      props.as === 'h4' && [
+        {
+          fontSize: FontSizes.size16,
+          marginBottom: '6px',
+          color: NeutralColors.gray160,
+        },
+      ],
+      className,
+    ],
   };
 };
 
 const getClassNames = classNamesFunction<IMarkdownHeaderStyleProps, IMarkdownHeaderStyles>();
 
-const MarkdownHeaderBase: React.StatelessComponent<IMarkdownHeaderProps> = props => {
+const MarkdownHeaderBase: React.FunctionComponent<IMarkdownHeaderProps> = props => {
   const { as: RootType = 'h1', children, id, styles, className } = props;
 
   const classNames = getClassNames(styles, { as: RootType, className });
@@ -87,7 +98,7 @@ const MarkdownHeaderBase: React.StatelessComponent<IMarkdownHeaderProps> = props
   );
 };
 
-export const MarkdownHeader: React.StatelessComponent<IMarkdownHeaderProps> = styled<
+export const MarkdownHeader: React.FunctionComponent<IMarkdownHeaderProps> = styled<
   IMarkdownHeaderProps,
   IMarkdownHeaderStyleProps,
   IMarkdownHeaderStyles
diff --git a/Engine/Source/Programs/Horde/HordeDashboard/src/base/components/Markdown/MarkdownLink.tsx b/Engine/Source/Programs/Horde/HordeDashboard/src/base/components/Markdown/MarkdownLink.tsx
--- a/Engine/Source/Programs/Horde/HordeDashboard/src/base/components/Markdown/MarkdownLink.tsx
+++ b/Engine/Source/Programs/Horde/HordeDashboard/src/base/components/Markdown/MarkdownLink.tsx
@@ -3,7 +3,7 @@ import { Link } from 'react-router-dom';
 import { ILinkProps } from '@fluentui/react/lib/Link';
 import { removeAnchorLink } from '../../utilities/index';
 
-export const MarkdownLink: React.StatelessComponent<ILinkProps> = props => {
+export const MarkdownLink: React.FunctionComponent<ILinkProps> = props => {
    let href = props.href;
    if (href && href[0] === '#' && href.indexOf('/') === -1) {
       // This is an anchor link within this page. We need to prepend the current route.
@@ -22,7 +22,7 @@ export const MarkdownLink: React.StatelessComponent<ILinkProps> = props => {
          if (url.hostname === window.location.hostname) {
             href = url.pathname + url.search
          } else {
-            return <a children={props.children} href={href} target="_blank" rel="noreferrer"/>;
+            return <a children={props.children} href={href} target="_blank" rel="noreferrer" />;
          }
       } catch {
 
diff --git a/Engine/Source/Programs/Horde/HordeDashboard/src/base/components/Markdown/MarkdownParagraph.tsx b/Engine/Source/Programs/Horde/HordeDashboard/src/base/components/Markdown/MarkdownParagraph.tsx
--- a/Engine/Source/Programs/Horde/HordeDashboard/src/base/components/Markdown/MarkdownParagraph.tsx
+++ b/Engine/Source/Programs/Horde/HordeDashboard/src/base/components/Markdown/MarkdownParagraph.tsx
@@ -1,6 +1,13 @@
 import * as React from 'react';
-import { IClassNames, IStyleFunction, classNamesFunction, styled, IStyleFunctionOrObject } from '@fluentui/react/lib/Utilities';
+import {
+  IClassNames,
+  IStyleFunction,
+  classNamesFunction,
+  styled,
+  IStyleFunctionOrObject,
+} from '@fluentui/react/lib/Utilities';
 import { ITheme, IStyle } from '@fluentui/react/lib/Styling';
+import { PropsWithChildren } from 'react';
 
 export interface IMarkdownParagraphProps {
   styles?: IStyleFunctionOrObject<IMarkdownParagraphStyleProps, IMarkdownParagraphStyles>;
@@ -22,29 +29,29 @@ const getStyles: IStyleFunction<IMarkdownParagraphStyleProps, IMarkdownParagraph
     root: [
       theme.fonts.medium,
       {
-        marginBottom: 4
+        marginBottom: 4,
       },
       isTodo && {
         padding: 8,
-        background: theme.semanticColors.warningBackground
-      }
-    ]
+        background: theme.semanticColors.warningBackground,
+      },
+    ],
   };
 };
 
 const getClassNames = classNamesFunction<IMarkdownParagraphStyleProps, IMarkdownParagraphStyles>();
 
-const MarkdownParagraphBase: React.StatelessComponent<IMarkdownParagraphProps> = props => {
+const MarkdownParagraphBase: React.FunctionComponent<PropsWithChildren<IMarkdownParagraphProps>> = props => {
   const { children, theme } = props;
   const classNames: IClassNames<IMarkdownParagraphStyles> = getClassNames(props.styles, {
     theme: theme!,
-    isTodo: typeof children === 'string' && children.indexOf('TODO') === 0
+    isTodo: typeof children === 'string' && children.indexOf('TODO') === 0,
   });
 
   return <p className={classNames.root}>{children}</p>;
 };
 
-export const MarkdownParagraph: React.StatelessComponent<IMarkdownParagraphProps> = styled<
+export const MarkdownParagraph: React.FunctionComponent<PropsWithChildren<IMarkdownParagraphProps>> = styled<
   IMarkdownParagraphProps,
   IMarkdownParagraphStyleProps,
   IMarkdownParagraphStyles
diff --git a/Engine/Source/Programs/Horde/HordeDashboard/src/base/components/Markdown/index.ts b/Engine/Source/Programs/Horde/HordeDashboard/src/base/components/Markdown/index.ts
--- a/Engine/Source/Programs/Horde/HordeDashboard/src/base/components/Markdown/index.ts
+++ b/Engine/Source/Programs/Horde/HordeDashboard/src/base/components/Markdown/index.ts
@@ -3,3 +3,4 @@ export * from './Markdown';
 export * from './MarkdownHeader';
 export * from './MarkdownLink';
 export * from './MarkdownParagraph';
+export * from './MarkdownPre';
diff --git a/Engine/Source/Programs/Horde/HordeDashboard/src/base/components/MarkdownTable/MarkdownCell.tsx b/Engine/Source/Programs/Horde/HordeDashboard/src/base/components/MarkdownTable/MarkdownCell.tsx
--- a/Engine/Source/Programs/Horde/HordeDashboard/src/base/components/MarkdownTable/MarkdownCell.tsx
+++ b/Engine/Source/Programs/Horde/HordeDashboard/src/base/components/MarkdownTable/MarkdownCell.tsx
@@ -2,15 +2,16 @@ import * as React from 'react';
 import { classNamesFunction, styled } from '@fluentui/react';
 import { IMarkdownTableCellProps, IMarkdownTableStyleProps, IMarkdownTableStyles } from './MarkdownTable.types';
 import { getStyles } from './MarkdownTable.styles';
+import { PropsWithChildren } from 'react';
 
 const getClassNames = classNamesFunction<IMarkdownTableStyleProps, IMarkdownTableStyles>();
 
-export class MarkdownCellBase extends React.PureComponent<IMarkdownTableCellProps> {
+export class MarkdownCellBase extends React.PureComponent<PropsWithChildren<IMarkdownTableCellProps>> {
   public render(): JSX.Element {
     const { as = 'td', children, styles, theme } = this.props;
 
     const classNames = getClassNames(styles, {
-      theme: theme!
+      theme: theme!,
     });
 
     const Tag = as;
@@ -28,5 +29,5 @@ export const MarkdownCell = styled<IMarkdownTableCellProps, IMarkdownTableStyleP
   MarkdownCellBase,
   getStyles,
   undefined,
-  { scope: 'MarkdownTable' }
+  { scope: 'MarkdownTable' },
 );
diff --git a/Engine/Source/Programs/Horde/HordeDashboard/src/base/components/MarkdownTable/MarkdownTBody.tsx b/Engine/Source/Programs/Horde/HordeDashboard/src/base/components/MarkdownTable/MarkdownTBody.tsx
--- a/Engine/Source/Programs/Horde/HordeDashboard/src/base/components/MarkdownTable/MarkdownTBody.tsx
+++ b/Engine/Source/Programs/Horde/HordeDashboard/src/base/components/MarkdownTable/MarkdownTBody.tsx
@@ -2,15 +2,16 @@ import * as React from 'react';
 import { classNamesFunction, styled } from '@fluentui/react';
 import { IMarkdownTableProps, IMarkdownTableStyleProps, IMarkdownTableStyles } from './MarkdownTable.types';
 import { getStyles } from './MarkdownTable.styles';
+import { PropsWithChildren } from 'react';
 
 const getClassNames = classNamesFunction<IMarkdownTableStyleProps, IMarkdownTableStyles>();
 
-export class MarkdownTBodyBase extends React.PureComponent<IMarkdownTableProps> {
+export class MarkdownTBodyBase extends React.PureComponent<PropsWithChildren<IMarkdownTableProps>> {
   public render(): JSX.Element {
     const { children, styles, theme } = this.props;
 
     const classNames = getClassNames(styles, {
-      theme: theme!
+      theme: theme!,
     });
 
     return (
@@ -25,5 +26,5 @@ export const MarkdownTBody = styled<IMarkdownTableProps, IMarkdownTableStyleProp
   MarkdownTBodyBase,
   getStyles,
   undefined,
-  { scope: 'MarkdownTable' }
+  { scope: 'MarkdownTable' },
 );
diff --git a/Engine/Source/Programs/Horde/HordeDashboard/src/base/components/MarkdownTable/MarkdownTHead.tsx b/Engine/Source/Programs/Horde/HordeDashboard/src/base/components/MarkdownTable/MarkdownTHead.tsx
--- a/Engine/Source/Programs/Horde/HordeDashboard/src/base/components/MarkdownTable/MarkdownTHead.tsx
+++ b/Engine/Source/Programs/Horde/HordeDashboard/src/base/components/MarkdownTable/MarkdownTHead.tsx
@@ -2,15 +2,16 @@ import * as React from 'react';
 import { classNamesFunction, styled } from '@fluentui/react';
 import { IMarkdownTableProps, IMarkdownTableStyleProps, IMarkdownTableStyles } from './MarkdownTable.types';
 import { getStyles } from './MarkdownTable.styles';
+import { PropsWithChildren } from 'react';
 
 const getClassNames = classNamesFunction<IMarkdownTableStyleProps, IMarkdownTableStyles>();
 
-export class MarkdownTHeadBase extends React.PureComponent<IMarkdownTableProps> {
+export class MarkdownTHeadBase extends React.PureComponent<PropsWithChildren<IMarkdownTableProps>> {
   public render(): JSX.Element {
     const { children, styles, theme } = this.props;
 
     const classNames = getClassNames(styles, {
-      theme: theme!
+      theme: theme!,
     });
 
     return (
@@ -25,5 +26,5 @@ export const MarkdownTHead = styled<IMarkdownTableProps, IMarkdownTableStyleProp
   MarkdownTHeadBase,
   getStyles,
   undefined,
-  { scope: 'MarkdownTable' }
+  { scope: 'MarkdownTable' },
 );
diff --git a/Engine/Source/Programs/Horde/HordeDashboard/src/base/components/MarkdownTable/MarkdownTable.styles.ts b/Engine/Source/Programs/Horde/HordeDashboard/src/base/components/MarkdownTable/MarkdownTable.styles.ts
--- a/Engine/Source/Programs/Horde/HordeDashboard/src/base/components/MarkdownTable/MarkdownTable.styles.ts
+++ b/Engine/Source/Programs/Horde/HordeDashboard/src/base/components/MarkdownTable/MarkdownTable.styles.ts
@@ -8,7 +8,7 @@ const GlobalClassNames = {
   tbody: 'ms-MarkdownTable-tbody',
   tr: 'ms-MarkdownTable-tr',
   td: 'ms-MarkdownTable-td',
-  th: 'ms-MarkdownTable-th'
+  th: 'ms-MarkdownTable-th',
 };
 
 const cellStyles: IStyle = {
@@ -18,13 +18,13 @@ const cellStyles: IStyle = {
 
   selectors: {
     '&:first-of-type': {
-      paddingLeft: 0
+      paddingLeft: 0,
     },
 
     '&[style*="text-align"] .ms-mdImage img': {
-      display: 'inline-block'
-    }
-  }
+      display: 'inline-block',
+    },
+  },
 };
 
 export const getStyles = (props: IMarkdownTableStyleProps): IMarkdownTableStyles => {
@@ -41,9 +41,9 @@ export const getStyles = (props: IMarkdownTableStyleProps): IMarkdownTableStyles
       {
         width: '100%',
         overflowX: 'auto',
-        marginBottom: 20
+        marginBottom: 20,
         // TODO: add background gradient styles (see Page.module.scss)
-      }
+      },
     ],
 
     table: [
@@ -55,17 +55,17 @@ export const getStyles = (props: IMarkdownTableStyleProps): IMarkdownTableStyles
 
         selectors: {
           '.ms-mdImage img': {
-            maxWidth: '100%'
-          }
-        }
-      }
+            maxWidth: '100%',
+          },
+        },
+      },
     ],
 
     tr: [
       classNames.tr,
       {
-        borderBottom: `1px solid ${palette.neutralTertiary}`
-      }
+        borderBottom: `1px solid ${palette.neutralTertiary}`,
+      },
     ],
 
     thead: [
@@ -73,10 +73,10 @@ export const getStyles = (props: IMarkdownTableStyleProps): IMarkdownTableStyles
       {
         selectors: {
           [`.${classNames.tr}`]: {
-            borderColor: palette.neutralPrimary
-          }
-        }
-      }
+            borderColor: palette.neutralPrimary,
+          },
+        },
+      },
     ],
 
     tbody: [
@@ -84,18 +84,18 @@ export const getStyles = (props: IMarkdownTableStyleProps): IMarkdownTableStyles
       {
         selectors: {
           [`.${classNames.tr}:last-of-type`]: {
-            borderBottom: 'none'
-          }
-        }
-      }
+            borderBottom: 'none',
+          },
+        },
+      },
     ],
 
     td: [
       classNames.td,
       {
         color: palette.neutralSecondary,
-        ...cellStyles
-      }
+        ...cellStyles,
+      },
     ],
 
     th: [
@@ -105,8 +105,8 @@ export const getStyles = (props: IMarkdownTableStyleProps): IMarkdownTableStyles
         color: palette.neutralPrimary,
         whiteSpace: 'nowrap',
         textAlign: 'left',
-        ...cellStyles
-      }
-    ]
+        ...cellStyles,
+      },
+    ],
   };
 };
diff --git a/Engine/Source/Programs/Horde/HordeDashboard/src/base/components/MarkdownTable/MarkdownTable.tsx b/Engine/Source/Programs/Horde/HordeDashboard/src/base/components/MarkdownTable/MarkdownTable.tsx
--- a/Engine/Source/Programs/Horde/HordeDashboard/src/base/components/MarkdownTable/MarkdownTable.tsx
+++ b/Engine/Source/Programs/Horde/HordeDashboard/src/base/components/MarkdownTable/MarkdownTable.tsx
@@ -2,16 +2,17 @@ import * as React from 'react';
 import { classNamesFunction, styled } from '@fluentui/react';
 import { IMarkdownTableProps, IMarkdownTableStyleProps, IMarkdownTableStyles } from './MarkdownTable.types';
 import { getStyles } from './MarkdownTable.styles';
+import { PropsWithChildren } from 'react';
 
 const getClassNames = classNamesFunction<IMarkdownTableStyleProps, IMarkdownTableStyles>();
 
-export class MarkdownTableBase extends React.PureComponent<IMarkdownTableProps> {
+export class MarkdownTableBase extends React.PureComponent<PropsWithChildren<IMarkdownTableProps>> {
   public render(): JSX.Element {
     const { children, className, styles, theme } = this.props;
 
     const classNames = getClassNames(styles, {
       theme: theme!,
-      className
+      className,
     });
 
     return (
@@ -28,5 +29,5 @@ export const MarkdownTable = styled<IMarkdownTableProps, IMarkdownTableStyleProp
   MarkdownTableBase,
   getStyles,
   undefined,
-  { scope: 'MarkdownTable' }
+  { scope: 'MarkdownTable' },
 );
diff --git a/Engine/Source/Programs/Horde/HordeDashboard/src/base/components/MarkdownTable/MarkdownTable.types.ts b/Engine/Source/Programs/Horde/HordeDashboard/src/base/components/MarkdownTable/MarkdownTable.types.ts
--- a/Engine/Source/Programs/Horde/HordeDashboard/src/base/components/MarkdownTable/MarkdownTable.types.ts
+++ b/Engine/Source/Programs/Horde/HordeDashboard/src/base/components/MarkdownTable/MarkdownTable.types.ts
@@ -15,11 +15,11 @@ export interface IMarkdownTableCellProps extends IMarkdownTableProps {
    * Render the table cell as a th or td.
    * @default 'td'
    */
-  // TODO: doc says default is 'td', where is it set? if it has a default, why is 'as' required? making optional for now.
   as?: 'th' | 'td';
 }
 
-export type IMarkdownTableStyleProps = Required<Pick<IMarkdownTableProps, 'theme'>> & Pick<IMarkdownTableProps, 'className'>;
+export type IMarkdownTableStyleProps = Required<Pick<IMarkdownTableProps, 'theme'>> &
+  Pick<IMarkdownTableProps, 'className'>;
 
 export interface IMarkdownTableStyles {
   root: IStyle;
diff --git a/Engine/Source/Programs/Horde/HordeDashboard/src/base/components/MarkdownTable/MarkdownTr.tsx b/Engine/Source/Programs/Horde/HordeDashboard/src/base/components/MarkdownTable/MarkdownTr.tsx
--- a/Engine/Source/Programs/Horde/HordeDashboard/src/base/components/MarkdownTable/MarkdownTr.tsx
+++ b/Engine/Source/Programs/Horde/HordeDashboard/src/base/components/MarkdownTable/MarkdownTr.tsx
@@ -2,15 +2,16 @@ import * as React from 'react';
 import { classNamesFunction, styled } from '@fluentui/react';
 import { IMarkdownTableProps, IMarkdownTableStyleProps, IMarkdownTableStyles } from './MarkdownTable.types';
 import { getStyles } from './MarkdownTable.styles';
+import { PropsWithChildren } from 'react';
 
 const getClassNames = classNamesFunction<IMarkdownTableStyleProps, IMarkdownTableStyles>();
 
-export class MarkdownTrBase extends React.PureComponent<IMarkdownTableProps> {
+export class MarkdownTrBase extends React.PureComponent<PropsWithChildren<IMarkdownTableProps>> {
   public render(): JSX.Element {
     const { children, styles, theme } = this.props;
 
     const classNames = getClassNames(styles, {
-      theme: theme!
+      theme: theme!,
     });
 
     return (
@@ -25,5 +26,5 @@ export const MarkdownTr = styled<IMarkdownTableProps, IMarkdownTableStyleProps,
   MarkdownTrBase,
   getStyles,
   undefined,
-  { scope: 'MarkdownTable' }
+  { scope: 'MarkdownTable' },
 );
diff --git a/Engine/Source/Programs/Horde/HordeDashboard/src/base/components/SideRail/SideRail.tsx b/Engine/Source/Programs/Horde/HordeDashboard/src/base/components/SideRail/SideRail.tsx
--- a/Engine/Source/Programs/Horde/HordeDashboard/src/base/components/SideRail/SideRail.tsx
+++ b/Engine/Source/Programs/Horde/HordeDashboard/src/base/components/SideRail/SideRail.tsx
@@ -155,7 +155,7 @@ function _isElement(x: any): x is JSX.Element {
    return !!(x && (x as JSX.Element).props && (x as JSX.Element).type);
 }
 
-export const SideRail: React.StatelessComponent<ISideRailProps> = styled<ISideRailProps, ISideRailStyleProps, ISideRailStyles>(
+export const SideRail: React.FunctionComponent<ISideRailProps> = styled<ISideRailProps, ISideRailStyleProps, ISideRailStyles>(
    SideRailBase,
    getStyles,
    undefined,
diff --git a/Engine/Source/Programs/Horde/HordeDashboard/src/components/AgentView.tsx b/Engine/Source/Programs/Horde/HordeDashboard/src/components/AgentView.tsx
--- a/Engine/Source/Programs/Horde/HordeDashboard/src/components/AgentView.tsx
+++ b/Engine/Source/Programs/Horde/HordeDashboard/src/components/AgentView.tsx
@@ -1,15 +1,14 @@
 // Copyright Epic Games, Inc. All Rights Reserved.  
-import { Checkbox, CommandButton, ConstrainMode, ContextualMenu, DefaultButton, DetailsHeader, DetailsList, DetailsListLayoutMode, Dialog, DialogType, DirectionalHint, Dropdown, FontSizes, FontWeights, getTheme, IBasePickerProps, IColumn, Icon, IconButton, IContextualMenuItem, IContextualMenuProps, IDetailsHeaderProps, IDetailsHeaderStyles, IDetailsListProps, ITag, ITagItemStyles, ITooltipHostStyles, Link as ReactLink, mergeStyles, mergeStyleSets, PrimaryButton, ProgressIndicator, ScrollablePane, ScrollbarVisibility, SearchBox, Selection, SelectionMode, Slider, Spinner, SpinnerSize, Stack, Sticky, StickyPositionType, TagItem, Text, TextField } from '@fluentui/react';
-import { action, observable } from 'mobx';
+import { Checkbox, CommandButton, ConstrainMode, ContextualMenu, TagPicker, DefaultButton, DetailsHeader, DetailsList, DetailsListLayoutMode, Dialog, DialogType, DirectionalHint, Dropdown, FontSizes, FontWeights, getTheme, IBasePickerProps, IColumn, Icon, IconButton, IContextualMenuItem, IContextualMenuProps, IDetailsHeaderProps, IDetailsHeaderStyles, IDetailsListProps, ITag, ITagItemStyles, ITooltipHostStyles, Link as ReactLink, mergeStyles, mergeStyleSets, PrimaryButton, ProgressIndicator, ScrollablePane, ScrollbarVisibility, SearchBox, Selection, SelectionMode, Slider, Spinner, SpinnerSize, Stack, Sticky, StickyPositionType, TagItem, Text, TextField } from '@fluentui/react';
+import { action, makeObservable, observable } from 'mobx';
 import { observer } from 'mobx-react-lite';
 import moment from 'moment-timezone';
 import React, { createRef, useEffect, useState } from 'react';
-import { Link, useHistory, useParams } from 'react-router-dom';
+import { Link, useNavigate, useParams } from 'react-router-dom';
 import Marquee from 'react-text-marquee';
 import backend from '../backend';
 import { agentStore } from '../backend/AgentStore';
 import { AgentData, BatchUpdatePoolRequest, GetAgentResponse, LeaseState, PoolData } from '../backend/Api';
-import { HTagPicker } from '../base/components/HordeFluentComponents/HTagPicker/HTagPicker';
 import { copyToClipboard } from '../base/utilities/clipboard';
 import { useWindowSize } from '../base/utilities/hooks';
 import { hexToRGB, hordeClasses, linearInterpolate } from '../styles/Styles';
@@ -192,6 +191,7 @@ function getTaskTime(agent: GetAgentResponse): number {
 
 
 class LocalState {
+
    // agent update function params
    editingPools = false;  // if editing pools dialog is open
 
@@ -496,6 +496,7 @@ class LocalState {
    }
 
    constructor() {
+      makeObservable(this);
       this.columnsState = [
          {
             key: 'name',
@@ -622,6 +623,9 @@ type PoolEditorItem = {
 
 // state class for the edit pools modal
 class EditPoolsModalState {
+   constructor() {
+      makeObservable(this);
+   }
    // pool editor observables
    newIdSuffix = 0;
    @observable isOpen = false;
@@ -840,6 +844,11 @@ type AgentDataSlim = {
 
 // state class for the select pools modal
 class SelectPoolsModalState {
+
+   constructor() {
+      makeObservable(this);
+   }
+
    @observable isOpen = false;
    @observable.shallow availablePools: SelectPoolItem[] = [];
    @observable.shallow selectedPools: SelectPoolItem[] = [];
@@ -1553,7 +1562,7 @@ export const PoolSelectionModal: React.FC = observer(() => {
             }
          }}
       >
-         <HTagPicker
+         <TagPicker
             onResolveSuggestions={onResolveSuggestions}
             onEmptyResolveSuggestions={() => { return selectPoolsModalState.availablePools; }}
             onItemSelected={(selectedItem: ITag | undefined) => { return selectPoolsModalState.transitionPool(selectedItem as SelectPoolItem); }}
@@ -1591,7 +1600,7 @@ export const PoolSelectionModal: React.FC = observer(() => {
 
 export const AgentView: React.FC = observer(() => {
    const { agentId } = useParams<{ agentId: string }>();
-   const history = useHistory();
+   const navigate = useNavigate();
    const [initAgentUpdater, setInitAgentUpdater] = useState(false);
 
    // adjust automatically to viewport changes
@@ -1887,7 +1896,7 @@ export const AgentView: React.FC = observer(() => {
    }
 
    function onHistoryModalDismiss() {
-      history.push('/agents');
+      navigate('/agents');
    }
 
    let agentItems = agentStore.agents.filter(filterAgents, localState.agentFilter).sort(sortAgents);
@@ -2122,7 +2131,7 @@ export const AgentView: React.FC = observer(() => {
                      {getAgentStatusIcon(agent)}
                   </Stack.Item>
                   <Stack.Item align={"center"}>
-                     <ReactLink styles={{ root: { paddingTop: '1px' } }} title="Lease and Session History" onClick={() => { history.push(`/agents/${agent.id}`); }}>{agent.name}</ReactLink>
+                     <ReactLink styles={{ root: { paddingTop: '1px' } }} title="Lease and Session History" onClick={() => { navigate(`/agents/${agent.id}`); }}>{agent.name}</ReactLink>
                   </Stack.Item>
                </Stack>
             );
@@ -2391,7 +2400,7 @@ export const AgentView: React.FC = observer(() => {
          case 'systemInfoRAM':
             return getPropFromDevice(agent, "RAM");
          default:
-            return <span>{agent[column!.fieldName as keyof AgentData]}</span>;
+            return <span>{agent[column!.fieldName as keyof AgentData] as string}</span>;
       }
    }
 });
\ No newline at end of file
diff --git a/Engine/Source/Programs/Horde/HordeDashboard/src/components/AuditLog.tsx b/Engine/Source/Programs/Horde/HordeDashboard/src/components/AuditLog.tsx
--- a/Engine/Source/Programs/Horde/HordeDashboard/src/components/AuditLog.tsx
+++ b/Engine/Source/Programs/Horde/HordeDashboard/src/components/AuditLog.tsx
@@ -1,6 +1,6 @@
 // Copyright Epic Games, Inc. All Rights Reserved.
 import { ComboBox, DefaultButton, Dropdown, FocusZone, FocusZoneDirection, IDropdownOption, List, mergeStyleSets, SearchBox, Selection, SelectionMode, SelectionZone, Spinner, SpinnerSize, Stack, Text } from "@fluentui/react";
-import { action, observable } from 'mobx';
+import { action, makeObservable, observable } from 'mobx';
 import { observer } from "mobx-react-lite";
 import moment from "moment";
 import { useEffect, useState } from "react";
@@ -35,7 +35,7 @@ export const auditStyleNormal = mergeStyleSets(logStyleNormal, {
 class AuditLogHandler {
 
    constructor() {
-
+      makeObservable(this);
       this.clear();
    }
 
diff --git a/Engine/Source/Programs/Horde/HordeDashboard/src/components/AutomationSuiteDetails.tsx b/Engine/Source/Programs/Horde/HordeDashboard/src/components/AutomationSuiteDetails.tsx
--- a/Engine/Source/Programs/Horde/HordeDashboard/src/components/AutomationSuiteDetails.tsx
+++ b/Engine/Source/Programs/Horde/HordeDashboard/src/components/AutomationSuiteDetails.tsx
@@ -1,6 +1,6 @@
 import { DetailsList, DetailsListLayoutMode, Dropdown, IColumn, IconButton, IDropdownOption, mergeStyleSets, Modal, SelectionMode, Spinner, SpinnerSize, Stack, Text, TextField } from "@fluentui/react";
 import * as d3 from "d3";
-import { action, observable } from 'mobx';
+import { action, makeObservable, observable } from 'mobx';
 import { observer } from 'mobx-react-lite';
 import React, { useState } from "react";
 import backend from "../backend";
@@ -49,6 +49,7 @@ class SuiteHandler {
 
    constructor(suite: GetTestSuiteResponse, suiteRefs: GetTestDataRefResponse[], metaData: GetTestMetaResponse) {
 
+      makeObservable(this);
 
       this.suite = suite;
 
diff --git a/Engine/Source/Programs/Horde/HordeDashboard/src/components/AutomationSuiteTest.tsx b/Engine/Source/Programs/Horde/HordeDashboard/src/components/AutomationSuiteTest.tsx
--- a/Engine/Source/Programs/Horde/HordeDashboard/src/components/AutomationSuiteTest.tsx
+++ b/Engine/Source/Programs/Horde/HordeDashboard/src/components/AutomationSuiteTest.tsx
@@ -1,6 +1,6 @@
 
 import { FontIcon, Image, Label, ScrollablePane, ScrollbarVisibility, Spinner, SpinnerSize, Stack, Text } from "@fluentui/react";
-import { action, observable } from 'mobx';
+import { action, makeObservable, observable } from 'mobx';
 import { observer } from 'mobx-react-lite';
 import React, { useState } from "react";
 import { Link } from "react-router-dom";
@@ -38,6 +38,9 @@ type TestEvent = {
 class SuiteTestHandler {
 
    constructor(suite: GetTestSuiteResponse, test: GetTestResponse, details: GetTestDataDetailsResponse, testRef: GetTestDataRefResponse, metaData: GetTestMetaResponse) {
+
+      makeObservable(this);
+
       this.suite = suite;
       this.test = test;
 
@@ -254,7 +257,7 @@ const TestInfoBox: React.FC<{ handler: SuiteTestHandler }> = observer(({ handler
 
    const infoItems: JSX.Element[] = [];
    if (handler.jobId && handler.stepId) {
-      infoItems.push(infoItem("Horde", "Job Step", `job/${handler.jobId}/?step=${handler.stepId}`));
+      infoItems.push(infoItem("Horde", "Job Step", `/job/${handler.jobId}/?step=${handler.stepId}`));
    }
 
    const timestamp = handler.testRef.id.substring(0, 8)
diff --git a/Engine/Source/Programs/Horde/HordeDashboard/src/components/AutomationView.tsx b/Engine/Source/Programs/Horde/HordeDashboard/src/components/AutomationView.tsx
--- a/Engine/Source/Programs/Horde/HordeDashboard/src/components/AutomationView.tsx
+++ b/Engine/Source/Programs/Horde/HordeDashboard/src/components/AutomationView.tsx
@@ -3,11 +3,11 @@
 import { isNumber } from '@datadog/browser-core';
 import { Checkbox, DefaultButton, Dropdown, FocusZone, FocusZoneDirection, FontIcon, IContextualMenuItem, IContextualMenuProps, IDropdownOption, Label, PrimaryButton, ScrollablePane, ScrollbarVisibility, Spinner, SpinnerSize, Stack, Text } from '@fluentui/react';
 import * as d3 from "d3";
-import { action, observable } from 'mobx';
+import { action, makeObservable, observable } from 'mobx';
 import { observer } from 'mobx-react-lite';
 import moment from 'moment';
 import React, { useState } from 'react';
-import { useHistory, useLocation } from 'react-router-dom';
+import { useNavigate, useLocation } from 'react-router-dom';
 import backend from '../backend';
 import { GetTestDataRefResponse, GetTestMetaResponse, GetTestResponse, GetTestStreamResponse, GetTestSuiteResponse, TestOutcome } from '../backend/Api';
 import dashboard, { StatusColor } from '../backend/Dashboard';
@@ -59,6 +59,7 @@ type FilterState = "Success" | "Failed" | "Consecutive Failures" | "Skipped";
 class TestDataHandler {
 
    constructor(search: URLSearchParams) {
+      makeObservable(this);
       this.state = this.stateFromSearch(search);
       this.load();
    }
@@ -1734,7 +1735,7 @@ export const AutomationView: React.FC = observer(() => {
 
    const [state, setState] = useState<{ handler?: TestDataHandler, search?: string }>({});
    const location = useLocation();
-   const history = useHistory();
+   const navigate = useNavigate();
 
    let handler = state.handler;
 
@@ -1752,7 +1753,7 @@ export const AutomationView: React.FC = observer(() => {
 
    if (state.search !== csearch) {
       location.search = csearch;
-      history.replace(location);
+      navigate(location, {replace: true});
       setState({ ...state, search: csearch });
       return null;
    }
diff --git a/Engine/Source/Programs/Horde/HordeDashboard/src/components/ChangeCallout.tsx b/Engine/Source/Programs/Horde/HordeDashboard/src/components/ChangeCallout.tsx
--- a/Engine/Source/Programs/Horde/HordeDashboard/src/components/ChangeCallout.tsx
+++ b/Engine/Source/Programs/Horde/HordeDashboard/src/components/ChangeCallout.tsx
@@ -3,68 +3,71 @@ import { Callout, DirectionalHint, FontWeights, getTheme, mergeStyleSets, Stack,
 import React, { useState } from 'react';
 import backend from '../backend';
 import { JobData, ChangeSummaryData } from '../backend/Api';
-import { observable, action } from 'mobx';
+import { observable, action, makeObservable } from 'mobx';
 import { observer } from 'mobx-react-lite';
 import { Link } from 'react-router-dom';
 
 const theme = getTheme();
 const styles = mergeStyleSets({
-    container: {
-        overflow: 'auto',
-        maxHeight: 600
-    },
-    callout: {
-        maxWidth: 600
-    },
-    header: {
-        padding: '18px 14px 12px'
-    },
-    title: [
-        theme.fonts.mediumPlus,
-        {
-            margin: 0,
-            fontWeight: FontWeights.regular
-        }
-    ],
-    inner: {
-        height: '100%',
-        padding: '0 14px 10px'
-
-    },
-    actions: {
-        position: 'relative',
-        marginTop: 20,
-        width: '100%',
-        whiteSpace: 'nowrap'
-    },
-    subtext: [
-        theme.fonts.medium,
-        {
-            margin: 0,
-            fontWeight: FontWeights.light
-        }
-    ]
+   container: {
+      overflow: 'auto',
+      maxHeight: 600
+   },
+   callout: {
+      maxWidth: 600
+   },
+   header: {
+      padding: '18px 14px 12px'
+   },
+   title: [
+      theme.fonts.mediumPlus,
+      {
+         margin: 0,
+         fontWeight: FontWeights.regular
+      }
+   ],
+   inner: {
+      height: '100%',
+      padding: '0 14px 10px'
+
+   },
+   actions: {
+      position: 'relative',
+      marginTop: 20,
+      width: '100%',
+      whiteSpace: 'nowrap'
+   },
+   subtext: [
+      theme.fonts.medium,
+      {
+         margin: 0,
+         fontWeight: FontWeights.light
+      }
+   ]
 });
 
 class CommitCache {
 
-    constructor(change: number, streamId: string) {
-        this.change = change;
+   constructor(change: number, streamId: string) {
 
-        backend.getCommit(streamId, change).then((value) => {
-            this.setData(value!);
-        });
-    }
+      makeObservable(this);
 
-    @action
-    setData(commitData: ChangeSummaryData) {
-        this.commitData = commitData;
-    }
+      this.change = change;
 
-    @observable.ref
-    commitData?: ChangeSummaryData;
+      backend.getCommit(streamId, change).then((value) => {
+         this.setData(value!);
+      });
+   }
 
-    change: number;
+   @action
+   setData(commitData: ChangeSummaryData) {
+      this.commitData = commitData;
+   }
+
+   @observable.ref
+   commitData?: ChangeSummaryData;
+
+   change: number;
 
 }
 
@@ -72,64 +75,64 @@ const commits = new Map<number, CommitCache>();
 
 export const ChangeCallout: React.FC<{ job: JobData }> = observer(({ job }) => {
 
-    const [visible, setVisible] = useState(false);
-    const [searchRef] = useState(React.createRef<HTMLDivElement>());
+   const [visible, setVisible] = useState(false);
+   const [searchRef] = useState(React.createRef<HTMLDivElement>());
 
-    if (!job.change) {
-        return <div style={{paddingTop:3}}>Latest</div>;
-    }
+   if (!job.change) {
+      return <div style={{ paddingTop: 3 }}>Latest</div>;
+   }
 
-    const change = job.change;
+   const change = job.change;
 
-    let cache = commits.get(change);
-    if (!cache) {
-        cache = new CommitCache(change, job.streamId);
-        commits.set(change, cache);
-    }
+   let cache = commits.get(change);
+   if (!cache) {
+      cache = new CommitCache(change, job.streamId);
+      commits.set(change, cache);
+   }
 
-    const commit = cache.commitData;
+   const commit = cache.commitData;
 
 
-    let key = 0;
-    const commitLines = commit?.description.split(/\r\n|\n|\r/);
+   let key = 0;
+   const commitLines = commit?.description.split(/\r\n|\n|\r/);
 
-    if (!commit) {
-        return <div style={{paddingTop:3}}>{change}</div>;
-    }
-    return (
-        <Stack>
-            <Link to="" className="cl-callout-link" onClick={(ev) => {ev.preventDefault(); setVisible(!visible);}} >
-            <div ref={searchRef} style={{paddingTop:3}} >
-                {change}
+   if (!commit) {
+      return <div style={{ paddingTop: 3 }}>{change}</div>;
+   }
+   return (
+      <Stack>
+         <Link to="" className="cl-callout-link" onClick={(ev) => { ev.preventDefault(); setVisible(!visible); }} >
+            <div ref={searchRef} style={{ paddingTop: 3 }} >
+               {change}
             </div>
             <div className={styles.container} data-is-scrollable={true}>
-                <Callout isBeakVisible={false}
-                    onDismiss={(ev) => setVisible(false)}
-                    hidden={!visible}
-                    target={searchRef?.current}
-                    className={styles.callout}
-                    role="alertdialog"
-                    gapSpace={0}
-                    setInitialFocus={true}
-                    directionalHint={DirectionalHint.bottomLeftEdge}>
-                    <div className={styles.inner}>
-                        <div className={styles.header}>
-                            <p className={styles.title}>
-                                {`Change ${commit.number} committed by ${commit.authorInfo?.name}`}
-                            </p>
-                        </div>
-                        <div className={styles.inner}>
-                            <p className={styles.subtext}>
-                                {commitLines?.map(line => <Text key={`commitline_${job.id}_${key++}`} block>{line}</Text>)}
-                            </p>
-                        </div>
-                    </div>
-
-                </Callout>
+               <Callout isBeakVisible={false}
+                  onDismiss={(ev) => setVisible(false)}
+                  hidden={!visible}
+                  target={searchRef?.current}
+                  className={styles.callout}
+                  role="alertdialog"
+                  gapSpace={0}
+                  setInitialFocus={true}
+                  directionalHint={DirectionalHint.bottomLeftEdge}>
+                  <div className={styles.inner}>
+                     <div className={styles.header}>
+                        <p className={styles.title}>
+                           {`Change ${commit.number} committed by ${commit.authorInfo?.name}`}
+                        </p>
+                     </div>
+                     <div className={styles.inner}>
+                        <p className={styles.subtext}>
+                           {commitLines?.map(line => <Text key={`commitline_${job.id}_${key++}`} block>{line}</Text>)}
+                        </p>
+                     </div>
+                  </div>
+
+               </Callout>
             </div>
-            </Link>
-        </Stack>
-    );
+         </Link>
+      </Stack>
+   );
 });
 
 
diff --git a/Engine/Source/Programs/Horde/HordeDashboard/src/components/DebugView.tsx b/Engine/Source/Programs/Horde/HordeDashboard/src/components/DebugView.tsx
--- a/Engine/Source/Programs/Horde/HordeDashboard/src/components/DebugView.tsx
+++ b/Engine/Source/Programs/Horde/HordeDashboard/src/components/DebugView.tsx
@@ -1,12 +1,12 @@
 // Copyright Epic Games, Inc. All Rights Reserved.
 
 import React from 'react';
-import { useHistory, useParams } from 'react-router-dom';
+import { useNavigate, useParams } from 'react-router-dom';
 import backend from '../backend';
 
 export const DebugView: React.FC = () => {
 
-    const history = useHistory();
+    const navigate = useNavigate();
     const params = useParams<{ leaseId?: string }>();
 
     if (params.leaseId) {
@@ -17,7 +17,7 @@ export const DebugView: React.FC = () => {
 
                 if (lease) {
 
-                    history.replace(`/log/${lease.logId}?leaseId=${lease.id}`);
+                  navigate(`/log/${lease.logId}?leaseId=${lease.id}`, {replace: true});
 
                 } else {
                     console.error("Unable to get lease ", params.leaseId);
diff --git a/Engine/Source/Programs/Horde/HordeDashboard/src/components/DevicePoolTelemetry.tsx b/Engine/Source/Programs/Horde/HordeDashboard/src/components/DevicePoolTelemetry.tsx
--- a/Engine/Source/Programs/Horde/HordeDashboard/src/components/DevicePoolTelemetry.tsx
+++ b/Engine/Source/Programs/Horde/HordeDashboard/src/components/DevicePoolTelemetry.tsx
@@ -1,6 +1,6 @@
 import { DefaultButton, DetailsList, DetailsListLayoutMode, Dropdown, FocusZone, FocusZoneDirection, IColumn, IconButton, IDropdownOption, Modal, ScrollablePane, ScrollbarVisibility, SelectionMode, Spinner, SpinnerSize, Stack, Text } from "@fluentui/react";
 import * as d3 from "d3";
-import { action, observable } from "mobx";
+import { action, makeObservable, observable } from "mobx";
 import { observer } from "mobx-react-lite";
 import moment from "moment";
 import { useState } from "react";
@@ -74,6 +74,7 @@ type DeviceProblem = { deviceId: string, deviceName: string, problems: number, p
 class PoolTelemetryHandler {
 
    constructor() {
+      makeObservable(this);
       this.set();
    }
 
diff --git a/Engine/Source/Programs/Horde/HordeDashboard/src/components/ErrorHandler.tsx b/Engine/Source/Programs/Horde/HordeDashboard/src/components/ErrorHandler.tsx
--- a/Engine/Source/Programs/Horde/HordeDashboard/src/components/ErrorHandler.tsx
+++ b/Engine/Source/Programs/Horde/HordeDashboard/src/components/ErrorHandler.tsx
@@ -1,6 +1,6 @@
 // Copyright Epic Games, Inc. All Rights Reserved.
 
-import { action, observable } from 'mobx';
+import { action, makeObservable, observable } from 'mobx';
 import { observer } from 'mobx-react-lite';
 import { Checkbox, Dialog, DialogType, PrimaryButton, Stack, Text } from '@fluentui/react';
 import React from 'react';
@@ -42,13 +42,17 @@ export type ErrorInfo = {
 
 }
 
-export class ErrorHandler {
+class ErrorHandlerClass {
+
+   constructor() {
+      makeObservable(this);
+   }
 
    @observable
-   static error?: ErrorInfo;
+   error?: ErrorInfo;
 
    @action
-   static set(infoIn: ErrorInfo, update?: boolean): void {
+   set(infoIn: ErrorInfo, update?: boolean): void {
 
       this.print(infoIn);
 
@@ -78,7 +82,7 @@ export class ErrorHandler {
       this.error = info;
    }
 
-   static get message(): string {
+   get message(): string {
 
       const error = this.error;
 
@@ -106,7 +110,7 @@ export class ErrorHandler {
 
    }
 
-   static unauthorized(infoIn?: ErrorInfo): boolean {
+   unauthorized(infoIn?: ErrorInfo): boolean {
 
       let error = infoIn;
 
@@ -130,13 +134,13 @@ export class ErrorHandler {
 
    }
 
-   static get reject(): string {
+   get reject(): string {
 
       return this.message;
 
    }
 
-   static hash(error: ErrorInfo): string {
+   hash(error: ErrorInfo): string {
 
       const url = error.url ?? error?.response?.url;
 
@@ -153,7 +157,7 @@ export class ErrorHandler {
       return "";
    }
 
-   static print(info: ErrorInfo) {
+   print(info: ErrorInfo) {
 
       // this will also go out to datadog if configured
       const message = `${info.response ? JSON.stringify(info.response) : ""} ${info.reason ? info.reason : ""}`;
@@ -165,7 +169,7 @@ export class ErrorHandler {
    }
 
    @action
-   static clear(): void {
+   clear(): void {
 
       if (!this.error) {
          return;
@@ -184,12 +188,14 @@ export class ErrorHandler {
    }
 
 
-   static filter = new Set<string>();
+   filter = new Set<string>();
 
-   static filterError = false;
+   filterError = false;
 
 }
 
+export const ErrorHandler = new ErrorHandlerClass();
+
 export const ErrorDialog: React.FC = observer(() => {
 
    const error = ErrorHandler.error;
@@ -250,3 +256,5 @@ export const ErrorDialog: React.FC = observer(() => {
       </Dialog>
    );
 });
+
+export default ErrorHandler;
diff --git a/Engine/Source/Programs/Horde/HordeDashboard/src/components/HistoryModal.tsx b/Engine/Source/Programs/Horde/HordeDashboard/src/components/HistoryModal.tsx
--- a/Engine/Source/Programs/Horde/HordeDashboard/src/components/HistoryModal.tsx
+++ b/Engine/Source/Programs/Horde/HordeDashboard/src/components/HistoryModal.tsx
@@ -1,7 +1,7 @@
 // Copyright Epic Games, Inc. All Rights Reserved.
 
 import { ConstrainMode, ContextualMenu, DefaultButton, DetailsHeader, DetailsList, DetailsListLayoutMode, DetailsRow, Dialog, DialogFooter, DialogType, GroupedList, GroupHeader, IColumn, IContextualMenuItem, IContextualMenuProps, IDetailsHeaderProps, IDetailsHeaderStyles, IDetailsListProps, IGroup, ITooltipHostStyles, mergeStyleSets, Modal, Pivot, PivotItem, PrimaryButton, ScrollablePane, ScrollbarVisibility, Selection, SelectionMode, Spinner, SpinnerSize, Stack, Sticky, StickyPositionType, Text } from "@fluentui/react";
-import { action, observable } from "mobx";
+import { action, makeObservable, observable } from "mobx";
 import { observer } from "mobx-react-lite";
 import React, { useState } from "react";
 import { Link } from "react-router-dom";
@@ -44,6 +44,10 @@ const historyStyles = mergeStyleSets({
 });
 
 class HistoryModalState {
+   constructor() {
+      makeObservable(this);
+   }
+   
    @observable.ref selectedAgent: AgentData | undefined = undefined;
    @observable.shallow currentData: any = [];
    @observable.shallow infoItems: InfoPanelItem[] = [];
@@ -688,7 +692,7 @@ export const HistoryModal: React.FC<{ agentId: string | undefined, onDismiss: (.
                   return <Stack styles={{ root: { height: '100%', } }} horizontal><Stack.Item align={"center"}><Text styles={{ root: { fontSize: 12 } }} key={"leaseText_" + lease.id}>{lease.name}</Text></Stack.Item></Stack>;
                }
             default:
-               return <span>{lease[column!.fieldName as keyof LeaseData]}</span>;
+               return <span>{lease[column!.fieldName as keyof LeaseData] as string}</span>;
          }
       }
       else if (state.mode === "sessions") {
@@ -701,7 +705,7 @@ export const HistoryModal: React.FC<{ agentId: string | undefined, onDismiss: (.
             case 'endTime':
                return <Stack styles={{ root: { height: '100%', } }} horizontal horizontalAlign={'center'}><Stack.Item align={"center"}>{session.finishTime?.toString().substring(0, 24) || ""}</Stack.Item></Stack>
             default:
-               return <span>{session[column!.fieldName as keyof SessionData]}</span>;
+               return <span>{session[column!.fieldName as keyof SessionData] as string}</span>;
          }
       }
    }
diff --git a/Engine/Source/Programs/Horde/HordeDashboard/src/components/IssueViewV2.tsx b/Engine/Source/Programs/Horde/HordeDashboard/src/components/IssueViewV2.tsx
--- a/Engine/Source/Programs/Horde/HordeDashboard/src/components/IssueViewV2.tsx
+++ b/Engine/Source/Programs/Horde/HordeDashboard/src/components/IssueViewV2.tsx
@@ -1,9 +1,9 @@
 import { Checkbox, CommandBar, CommandBarButton, DefaultButton, getFocusStyle, ICommandBarItemProps, Icon, IconButton, IContextualMenuItem, IContextualMenuItemProps, IContextualMenuProps, IGroup, Label, List, mergeStyleSets, MessageBar, MessageBarType, Modal, Persona, PersonaSize, Pivot, PivotItem, PrimaryButton, Spinner, SpinnerSize, Stack, TagPicker, Text, TextField } from "@fluentui/react";
-import { action, observable } from 'mobx';
+import { action, makeObservable, observable } from 'mobx';
 import { observer } from "mobx-react-lite";
 import moment from "moment";
 import React, { useEffect, useRef, useState } from "react";
-import { Link, useHistory, useLocation } from "react-router-dom";
+import { Link, useNavigate, useLocation } from "react-router-dom";
 import backend from "../backend";
 import { CreateExternalIssueResponse, CreateJobRequest, EventSeverity, GetChangeSummaryResponse, GetExternalIssueProjectResponse, GetExternalIssueResponse, GetIssueResponse, GetIssueSpanResponse, GetIssueStepResponse, GetIssueStreamResponse, GetLogEventResponse, GetTemplateRefResponse, GetThinUserInfoResponse, GetUserResponse, IssueSeverity, UpdateIssueRequest } from "../backend/Api";
 import dashboard, { StatusColor } from "../backend/Dashboard";
@@ -56,6 +56,10 @@ type ChangeGroup = IGroup & ChangeItem;
 
 class IssueDetails {
 
+   constructor() {
+      makeObservable(this);
+   }
+
    async set(issueId: number) {
 
       if (this.issue?.id === issueId) {
@@ -1797,7 +1801,7 @@ const IssueCommandBar: React.FC = () => {
 
 export const IssueModalV2: React.FC<{ popHistoryOnClose: boolean, issueId?: string | null, streamId?: string, onCloseExternal?: () => void }> = observer(({ popHistoryOnClose, issueId, onCloseExternal, streamId }) => {
 
-   const history = useHistory();
+   const navigate = useNavigate();
    const query = useQuery();
    const location = useLocation();
    const [editShown, setEditShown] = useState(false);
@@ -1820,13 +1824,13 @@ export const IssueModalV2: React.FC<{ popHistoryOnClose: boolean, issueId?: stri
 
    const onClose = () => {
       if (queryId) {
-         if (popHistoryOnClose) {
-            history.goBack();
+         if (popHistoryOnClose) {            
+            navigate(-1);
          } else {
             let search = new URLSearchParams(location.search);
             search = new URLSearchParams(Array.from(search.entries()).filter(e => e[0] !== 'issue'));
             location.search = search.toString();
-            history.replace(location);
+            navigate(location, {replace: true});
          }
       }
    }
@@ -3026,7 +3030,7 @@ export const IssueForceCloseModal: React.FC<{ onClose: () => void }> = ({ onClos
 
 const TestFixModal: React.FC<{ onClose: () => void }> = ({ onClose }) => {
 
-   const history = useHistory();
+   const navigate = useNavigate();
    const [state, setState] = useState<{ loadingTemplates?: boolean, templates?: Map<string, GetTemplateRefResponse[]>, error?: string, submitting?: boolean, shelvedCL?: string, baseCL?: string, streamId?: string, templateId?: string, target?: string, updateIssues?: boolean }>({ updateIssues: true });
 
    const issue = details.issue!;
@@ -3185,7 +3189,7 @@ const TestFixModal: React.FC<{ onClose: () => void }> = ({ onClose }) => {
                }
 
                redirected = true;
-               history.push(`/job/${data.id}`);
+               navigate(`/job/${data.id}`);
 
             }).catch(reason => {
                // "Not Found" is generally a permissions error
diff --git a/Engine/Source/Programs/Horde/HordeDashboard/src/components/JobDetailArtifacts.tsx b/Engine/Source/Programs/Horde/HordeDashboard/src/components/JobDetailArtifacts.tsx
--- a/Engine/Source/Programs/Horde/HordeDashboard/src/components/JobDetailArtifacts.tsx
+++ b/Engine/Source/Programs/Horde/HordeDashboard/src/components/JobDetailArtifacts.tsx
@@ -6,211 +6,216 @@ import { ArtifactData, GetArtifactZipRequest } from '../backend/Api';
 import { JobDetails } from '../backend/JobDetails';
 import { hordeClasses } from '../styles/Styles';
 import { observer } from 'mobx-react-lite';
-import { observable, action } from 'mobx';
+import { observable, action, makeObservable } from 'mobx';
 import backend from '../backend';
 
 const classNames = mergeStyleSets({
-	fileIconHeaderIcon: {
-	  padding: 0,
-	  fontSize: '24px',
-	},
-	fileIconCell: {
-	  textAlign: 'center',
-	  selectors: {
-		'&:before': {
-		  content: '.',
-		  display: 'inline-block',
-		  verticalAlign: 'middle',
-		  height: '100%',
-		  width: '0px',
-		  visibility: 'hidden',
-		},
-	  },
-	},
-	fileIconImg: {
-		verticalAlign: 'middle',
-		maxHeight: '24px',
-		maxWidth: '24px',
-		fontSize: '24px'
-	},
+   fileIconHeaderIcon: {
+      padding: 0,
+      fontSize: '24px',
+   },
+   fileIconCell: {
+      textAlign: 'center',
+      selectors: {
+         '&:before': {
+            content: '.',
+            display: 'inline-block',
+            verticalAlign: 'middle',
+            height: '100%',
+            width: '0px',
+            visibility: 'hidden',
+         },
+      },
+   },
+   fileIconImg: {
+      verticalAlign: 'middle',
+      maxHeight: '24px',
+      maxWidth: '24px',
+      fontSize: '24px'
+   },
 });
 
 class ArtifactState {
-	@observable selectedItemCount = 0;
-	@observable isDownloading = false;
-	jobDetails: JobDetails | null = null;
-	stepId?: string | undefined = undefined;
-
-	selection: Selection = new Selection({ onSelectionChanged: () => { this.setSelectedItemCount(this.selection.getSelectedCount()); }});
-
-	@action
-	private setSelectedItemCount(count: number) {
-		this.selectedItemCount = count;
-	}
-
-	@action 
-	setIsDownloading(isDownloading: boolean){
-		this.isDownloading = isDownloading;
-	}
-
-	setDetails(jobDetails: JobDetails, stepId?: string) {
-		if(jobDetails!.id! !== this.jobDetails?.id || stepId !== this.stepId) {
-			this.selection.setAllSelected(false);
-		}
-		this.jobDetails = jobDetails;
-		this.stepId = stepId;
-	}
-
-	async downloadItems(fromDownloadButton?: boolean | undefined, ctrlKeyPressed?: boolean | undefined, forceFileName?: string | undefined) {
-		const request: GetArtifactZipRequest = {};
-		request.jobId = this.jobDetails!.id!;
-		request.stepId = this.stepId;
-		if(this.selectedItemCount !== 0) {
-			const items: ArtifactData[] = [];
-			this.selection.getSelection().forEach(function(item) { 
-				items.push(item as ArtifactData);
-			});
-
-			request.artifactIds = items.map(item => item.id);
-			if(request.artifactIds.length === 1) {
-				request.fileName = items[0].name;
-			}			
-		}
-		request.isForcedDownload = fromDownloadButton;
-		request.isOpenNewTab = ctrlKeyPressed;
-		if(!request.fileName && forceFileName) {
-			request.fileName = forceFileName;
-		}
-		this.setIsDownloading(true);
-		backend.downloadJobArtifacts(request).then(response => {
-
-		}).finally(() => {
-			this.setIsDownloading(false);
-		});
-	}
-
-	columns: IColumn[] = [
-		{
-			key: 'type',
-			name: 'File Type',
-			className: classNames.fileIconCell,
-			iconClassName: classNames.fileIconHeaderIcon,
-			iconName: 'Page',
-			isIconOnly: true,
-			fieldName: 'type',
-			minWidth: 24,
-			maxWidth: 24,
-			onRender: (item: ArtifactData) => {
-				let iconName = "Page";
-				switch(item.mimeType){
-					case "text/plain":
-						iconName = "TextDocument";
-						break;
-					case "text/xml":
-						iconName = "FileCode";
-						break;
-					case "application/octet-stream":
-						iconName = "ProcessMetaTask";
-						break;
-					case "application/json":
-						iconName = "Code";
-						break;	
-				}
-				return <Icon iconName={iconName} className={classNames.fileIconImg}/>;
-			}
-		},
-		{
-			key: 'name',
-			name: 'Name',
-			fieldName: 'name',
-			minWidth: 210,
-			maxWidth: 350,
-			isRowHeader: true,
-			data: 'string',
-			isPadded: true,
-			onRender: (item: ArtifactData) => {
-				return (<Stack styles={{ root: { height: '100%' }}} horizontal>
-							<Stack.Item align={"center"}>
-								<Link href={`${backend.serverUrl}/api/v1/artifacts/${item.id}/download?filename=${encodeURI(item.name)}&code=${item.code}`} target="_blank">
-									{item.name}
-								</Link>
-							</Stack.Item>
-						</Stack>);
-			}
-		}
-	];
+
+   constructor() {
+      makeObservable(this);
+   }
+
+   @observable selectedItemCount = 0;
+   @observable isDownloading = false;
+   jobDetails: JobDetails | null = null;
+   stepId?: string | undefined = undefined;
+
+   selection: Selection = new Selection({ onSelectionChanged: () => { this.setSelectedItemCount(this.selection.getSelectedCount()); } });
+
+   @action
+   private setSelectedItemCount(count: number) {
+      this.selectedItemCount = count;
+   }
+
+   @action
+   setIsDownloading(isDownloading: boolean) {
+      this.isDownloading = isDownloading;
+   }
+
+   setDetails(jobDetails: JobDetails, stepId?: string) {
+      if (jobDetails!.id! !== this.jobDetails?.id || stepId !== this.stepId) {
+         this.selection.setAllSelected(false);
+      }
+      this.jobDetails = jobDetails;
+      this.stepId = stepId;
+   }
+
+   async downloadItems(fromDownloadButton?: boolean | undefined, ctrlKeyPressed?: boolean | undefined, forceFileName?: string | undefined) {
+      const request: GetArtifactZipRequest = {};
+      request.jobId = this.jobDetails!.id!;
+      request.stepId = this.stepId;
+      if (this.selectedItemCount !== 0) {
+         const items: ArtifactData[] = [];
+         this.selection.getSelection().forEach(function (item) {
+            items.push(item as ArtifactData);
+         });
+
+         request.artifactIds = items.map(item => item.id);
+         if (request.artifactIds.length === 1) {
+            request.fileName = items[0].name;
+         }
+      }
+      request.isForcedDownload = fromDownloadButton;
+      request.isOpenNewTab = ctrlKeyPressed;
+      if (!request.fileName && forceFileName) {
+         request.fileName = forceFileName;
+      }
+      this.setIsDownloading(true);
+      backend.downloadJobArtifacts(request).then(response => {
+
+      }).finally(() => {
+         this.setIsDownloading(false);
+      });
+   }
+
+   columns: IColumn[] = [
+      {
+         key: 'type',
+         name: 'File Type',
+         className: classNames.fileIconCell,
+         iconClassName: classNames.fileIconHeaderIcon,
+         iconName: 'Page',
+         isIconOnly: true,
+         fieldName: 'type',
+         minWidth: 24,
+         maxWidth: 24,
+         onRender: (item: ArtifactData) => {
+            let iconName = "Page";
+            switch (item.mimeType) {
+               case "text/plain":
+                  iconName = "TextDocument";
+                  break;
+               case "text/xml":
+                  iconName = "FileCode";
+                  break;
+               case "application/octet-stream":
+                  iconName = "ProcessMetaTask";
+                  break;
+               case "application/json":
+                  iconName = "Code";
+                  break;
+            }
+            return <Icon iconName={iconName} className={classNames.fileIconImg} />;
+         }
+      },
+      {
+         key: 'name',
+         name: 'Name',
+         fieldName: 'name',
+         minWidth: 210,
+         maxWidth: 350,
+         isRowHeader: true,
+         data: 'string',
+         isPadded: true,
+         onRender: (item: ArtifactData) => {
+            return (<Stack styles={{ root: { height: '100%' } }} horizontal>
+               <Stack.Item align={"center"}>
+                  <Link href={`${backend.serverUrl}/api/v1/artifacts/${item.id}/download?filename=${encodeURI(item.name)}&code=${item.code}`} target="_blank">
+                     {item.name}
+                  </Link>
+               </Stack.Item>
+            </Stack>);
+         }
+      }
+   ];
 
 }
 
 const artifactState = new ArtifactState();
-export const JobDetailArtifacts: React.FC<{ jobDetails: JobDetails; stepId?: string, topPadding?:number  }> = observer(({ jobDetails, stepId, topPadding }) => {
-	artifactState.setDetails(jobDetails, stepId);
-
-	function getKey(item: any, index?: number): string {
-		return item.key;
-	}
-
-	// main header
-    const onRenderDetailsHeader: IDetailsListProps['onRenderDetailsHeader'] = (props) => {
-        const customStyles: Partial<IDetailsHeaderStyles> = {};
-        if (props) {
-            return (
-                <Sticky stickyPosition={StickyPositionType.Header} isScrollSynced={true}>
-                    <DetailsHeader {...props} styles={customStyles}/>
-                </Sticky>
-            );
-        }
-        return null;
-	};
-
-	let buttonText = "Downloading...";
-	if(!artifactState.isDownloading) {
-		buttonText = "Download All";
-		if(artifactState.selectedItemCount > 0) {
-			if(artifactState.selectedItemCount !== jobDetails.artifacts.length) {
-				buttonText = `Download ${artifactState.selectedItemCount} File${artifactState.selectedItemCount > 1 ? 's' : ''}`;
-			}
-		}
-	}
-
-	let artifacts: ArtifactData[] = jobDetails.artifacts;
-	if(stepId) {
-		artifacts = artifacts.filter(artifact => artifact.stepId === stepId);
-	}
-
-	let height = Math.min(36 * artifacts.length + 60, 500);
-	
-	return (<Stack styles={{ root: { paddingTop: topPadding ?? 18, paddingRight: 12 } }}>
-		<Stack className={hordeClasses.raised}>
-			<Stack tokens={{ childrenGap: 12 }}>
-				<Stack horizontal horizontalAlign="space-between" styles={{ root: { minHeight: 32 }}}>
-					<Text variant="mediumPlus" styles={{ root: { fontFamily: "Horde Open Sans SemiBold" } }}>Artifacts</Text>
-					<Stack horizontal>
-						{artifactState.isDownloading && <Spinner  styles={{ root: { marginRight: 10 }}} size={SpinnerSize.medium}></Spinner>}
-						<PrimaryButton styles={{ root: { fontFamily: 'Horde Open Sans Semibold !important' }}} disabled={artifactState.isDownloading} onClick={artifactState.downloadItems.bind(artifactState, true, false, undefined)}>{buttonText}</PrimaryButton>
-					</Stack>
-				</Stack>
-				<Stack>
-					<Stack.Item styles={{ root: { height: height, maxHeight: height, position: 'relative' }}}>
-						<ScrollablePane scrollbarVisibility={ScrollbarVisibility.auto}>
-							<DetailsList
-								items={artifacts}
-								compact={true}
-								columns={artifactState.columns}
-								selectionMode={SelectionMode.multiple}
-								getKey={getKey}
-								setKey="multiple"
-								layoutMode={DetailsListLayoutMode.justified}
-								onRenderDetailsHeader={onRenderDetailsHeader}
-								isHeaderVisible={true}
-								selection={artifactState.selection}
-								selectionPreservedOnEmptyClick={true}
-							/>
-						</ScrollablePane>
-					</Stack.Item>
-				</Stack>
-				
-			</Stack>
-
-		</Stack></Stack>);
+export const JobDetailArtifacts: React.FC<{ jobDetails: JobDetails; stepId?: string, topPadding?: number }> = observer(({ jobDetails, stepId, topPadding }) => {
+   artifactState.setDetails(jobDetails, stepId);
+
+   function getKey(item: any, index?: number): string {
+      return item.key;
+   }
+
+   // main header
+   const onRenderDetailsHeader: IDetailsListProps['onRenderDetailsHeader'] = (props) => {
+      const customStyles: Partial<IDetailsHeaderStyles> = {};
+      if (props) {
+         return (
+            <Sticky stickyPosition={StickyPositionType.Header} isScrollSynced={true}>
+               <DetailsHeader {...props} styles={customStyles} />
+            </Sticky>
+         );
+      }
+      return null;
+   };
+
+   let buttonText = "Downloading...";
+   if (!artifactState.isDownloading) {
+      buttonText = "Download All";
+      if (artifactState.selectedItemCount > 0) {
+         if (artifactState.selectedItemCount !== jobDetails.artifacts.length) {
+            buttonText = `Download ${artifactState.selectedItemCount} File${artifactState.selectedItemCount > 1 ? 's' : ''}`;
+         }
+      }
+   }
+
+   let artifacts: ArtifactData[] = jobDetails.artifacts;
+   if (stepId) {
+      artifacts = artifacts.filter(artifact => artifact.stepId === stepId);
+   }
+
+   let height = Math.min(36 * artifacts.length + 60, 500);
+
+   return (<Stack styles={{ root: { paddingTop: topPadding ?? 18, paddingRight: 12 } }}>
+      <Stack className={hordeClasses.raised}>
+         <Stack tokens={{ childrenGap: 12 }}>
+            <Stack horizontal horizontalAlign="space-between" styles={{ root: { minHeight: 32 } }}>
+               <Text variant="mediumPlus" styles={{ root: { fontFamily: "Horde Open Sans SemiBold" } }}>Artifacts</Text>
+               <Stack horizontal>
+                  {artifactState.isDownloading && <Spinner styles={{ root: { marginRight: 10 } }} size={SpinnerSize.medium}></Spinner>}
+                  <PrimaryButton styles={{ root: { fontFamily: 'Horde Open Sans Semibold !important' } }} disabled={artifactState.isDownloading} onClick={artifactState.downloadItems.bind(artifactState, true, false, undefined)}>{buttonText}</PrimaryButton>
+               </Stack>
+            </Stack>
+            <Stack>
+               <Stack.Item styles={{ root: { height: height, maxHeight: height, position: 'relative' } }}>
+                  <ScrollablePane scrollbarVisibility={ScrollbarVisibility.auto}>
+                     <DetailsList
+                        items={artifacts}
+                        compact={true}
+                        columns={artifactState.columns}
+                        selectionMode={SelectionMode.multiple}
+                        getKey={getKey}
+                        setKey="multiple"
+                        layoutMode={DetailsListLayoutMode.justified}
+                        onRenderDetailsHeader={onRenderDetailsHeader}
+                        isHeaderVisible={true}
+                        selection={artifactState.selection}
+                        selectionPreservedOnEmptyClick={true}
+                     />
+                  </ScrollablePane>
+               </Stack.Item>
+            </Stack>
+
+         </Stack>
+
+      </Stack></Stack>);
 });
\ No newline at end of file
