Horde: Updating dashboard dependencies to latest versions and cleanup

diff --git a/Engine/Source/Programs/Horde/HordeDashboard/src/components/JobDetailTrace.tsx b/Engine/Source/Programs/Horde/HordeDashboard/src/components/JobDetailTrace.tsx
--- a/Engine/Source/Programs/Horde/HordeDashboard/src/components/JobDetailTrace.tsx
+++ b/Engine/Source/Programs/Horde/HordeDashboard/src/components/JobDetailTrace.tsx
@@ -6,150 +6,151 @@ import { GetJobStepTraceResponse } from '../backend/Api';
 import { JobDetails } from '../backend/JobDetails';
 import { hordeClasses } from '../styles/Styles';
 import { observer } from 'mobx-react-lite';
-import { observable, action } from 'mobx';
+import { observable, action, makeObservable } from 'mobx';
 import backend from '../backend';
 import { FlameGraph } from 'react-flame-graph';
 import moment from 'moment-timezone';
 
 const classNames = mergeStyleSets({
-	fileIconHeaderIcon: {
-	  padding: 0,
-	  fontSize: '24px',
-	},
-	fileIconCell: {
-	  textAlign: 'center',
-	  selectors: {
-		'&:before': {
-		  content: '.',
-		  display: 'inline-block',
-		  verticalAlign: 'middle',
-		  height: '100%',
-		  width: '0px',
-		  visibility: 'hidden',
-		},
-	  },
-	},
-	fileIconImg: {
-		verticalAlign: 'middle',
-		maxHeight: '24px',
-		maxWidth: '24px',
-		fontSize: '24px'
-	},
+   fileIconHeaderIcon: {
+      padding: 0,
+      fontSize: '24px',
+   },
+   fileIconCell: {
+      textAlign: 'center',
+      selectors: {
+         '&:before': {
+            content: '.',
+            display: 'inline-block',
+            verticalAlign: 'middle',
+            height: '100%',
+            width: '0px',
+            visibility: 'hidden',
+         },
+      },
+   },
+   fileIconImg: {
+      verticalAlign: 'middle',
+      maxHeight: '24px',
+      maxWidth: '24px',
+      fontSize: '24px'
+   },
 });
 
 type TraceData = {
-    name: string;
-    value: number;
-    children: TraceData[]
-    tooltip?: string;
+   name: string;
+   value: number;
+   children: TraceData[]
+   tooltip?: string;
 }
 
 class TraceState {
-    @observable data: TraceData | undefined = undefined;
-    stepId: string | undefined = undefined;
-
-    getTrace(jobDetails: JobDetails, stepId: string) {
-        if(stepId !== this.stepId) {
-            this.stepId = stepId;
-            const batch = jobDetails.batchByStepId(stepId);
-            backend.getJobStepTrace(jobDetails.id!, batch!.id, stepId).then(result => {
-                this._setData(result);
-            }).catch(error => {
-            }).finally(() => {
-            })
-        }
-    }
-
-    private _recursiveGenerateChildren(traceParent: GetJobStepTraceResponse, parent: TraceData) {
-        if(traceParent.Children && traceParent.Children.length !== 0) {
-            let parentRunTime = traceParent.Finish - traceParent.Start;
-            let parentValue = parent.value;
-
-            // set us to the start of time for this parent
-            let lastChildEndTime = traceParent.Start;
-            traceParent.Children.forEach(child => {
-                // get percentage of the runtime
-                let runTime = child.Finish - child.Start;
-                let runPercentage = (runTime) / parentRunTime;
-
-                // get how much of that slice is from the parent value
-                let childSlice = parentValue * runPercentage;
-
-                // push deadtime as a child if applicable
-                if(child.Start !== lastChildEndTime) {
-                    // get dead time percentage
-                    let idleTime = child.Start - lastChildEndTime;
-                    let idlePercentage = (idleTime) / parentRunTime;
-
-                    if(idlePercentage > 0) {
-                        let idleSlice = parentValue * idlePercentage;
-
-                        let idleChild: TraceData = { name: "Agent Idle", value: idleSlice, children: [], tooltip: `Agent was idle ${moment.utc(idleTime / 10000).format("mm:ss.SSS")}` };
-                        parent.children.push(idleChild);
-                    }
-
-                }
-
-
-                // set us to the child's finish time
-                lastChildEndTime = child.Finish;
-
-                let name = child.Name;
-                if(child.Resource) {
-                    name += `: ${child.Resource}`;
-                }
-
-                let newChild: TraceData = { name: name, value: childSlice, children: [], tooltip: `${child.Service}: (${moment.utc(runTime / 10000).format("mm:ss.SSS")})` };
-                this._recursiveGenerateChildren(child, newChild);
-                parent.children.push(newChild);
-            });
-        }
-    }
-    
-    @action
-    private _setData(trace: GetJobStepTraceResponse | undefined) {
-        if(trace) {
-            let name = trace.Name;
-            if(trace.Resource) {
-                name += `: ${trace.Resource}`;
+
+   @observable data: TraceData | undefined = undefined;
+   stepId: string | undefined = undefined;
+
+   getTrace(jobDetails: JobDetails, stepId: string) {
+      if (stepId !== this.stepId) {
+         this.stepId = stepId;
+         const batch = jobDetails.batchByStepId(stepId);
+         backend.getJobStepTrace(jobDetails.id!, batch!.id, stepId).then(result => {
+            this._setData(result);
+         }).catch(error => {
+         }).finally(() => {
+         })
+      }
+   }
+
+   private _recursiveGenerateChildren(traceParent: GetJobStepTraceResponse, parent: TraceData) {
+      if (traceParent.Children && traceParent.Children.length !== 0) {
+         let parentRunTime = traceParent.Finish - traceParent.Start;
+         let parentValue = parent.value;
+
+         // set us to the start of time for this parent
+         let lastChildEndTime = traceParent.Start;
+         traceParent.Children.forEach(child => {
+            // get percentage of the runtime
+            let runTime = child.Finish - child.Start;
+            let runPercentage = (runTime) / parentRunTime;
+
+            // get how much of that slice is from the parent value
+            let childSlice = parentValue * runPercentage;
+
+            // push deadtime as a child if applicable
+            if (child.Start !== lastChildEndTime) {
+               // get dead time percentage
+               let idleTime = child.Start - lastChildEndTime;
+               let idlePercentage = (idleTime) / parentRunTime;
+
+               if (idlePercentage > 0) {
+                  let idleSlice = parentValue * idlePercentage;
+
+                  let idleChild: TraceData = { name: "Agent Idle", value: idleSlice, children: [], tooltip: `Agent was idle ${moment.utc(idleTime / 10000).format("mm:ss.SSS")}` };
+                  parent.children.push(idleChild);
+               }
+
             }
-            let graph: TraceData = { name: name, value: 100, children: [], tooltip: `${trace.Service ?? 'Total Runtime'}: (${moment.utc((trace.Finish - trace.Start) / 10000).format("mm:ss.SSS")})` }
-            this._recursiveGenerateChildren(trace, graph);
-            this.data = graph;
-        }
-        else {
-            this.data = undefined;
-        }
-    }
 
 
-    constructor() {
+            // set us to the child's finish time
+            lastChildEndTime = child.Finish;
 
-    }
+            let name = child.Name;
+            if (child.Resource) {
+               name += `: ${child.Resource}`;
+            }
+
+            let newChild: TraceData = { name: name, value: childSlice, children: [], tooltip: `${child.Service}: (${moment.utc(runTime / 10000).format("mm:ss.SSS")})` };
+            this._recursiveGenerateChildren(child, newChild);
+            parent.children.push(newChild);
+         });
+      }
+   }
+
+   @action
+   private _setData(trace: GetJobStepTraceResponse | undefined) {
+      if (trace) {
+         let name = trace.Name;
+         if (trace.Resource) {
+            name += `: ${trace.Resource}`;
+         }
+         let graph: TraceData = { name: name, value: 100, children: [], tooltip: `${trace.Service ?? 'Total Runtime'}: (${moment.utc((trace.Finish - trace.Start) / 10000).format("mm:ss.SSS")})` }
+         this._recursiveGenerateChildren(trace, graph);
+         this.data = graph;
+      }
+      else {
+         this.data = undefined;
+      }
+   }
+
+
+   constructor() {
+      makeObservable(this);
+   }
 }
 
 const traceState = new TraceState();
 export const JobDetailTrace: React.FC<{ jobDetails: JobDetails; stepId: string }> = observer(({ jobDetails, stepId }) => {
 
-    traceState.getTrace(jobDetails, stepId!);
-
-    // let height = Math.min(36 * artifacts.length + 60, 500);	
-	return (<Stack styles={{ root: { paddingTop: 18, paddingRight: 12 } }}>
-		<Stack className={hordeClasses.raised}>
-			<Stack tokens={{ childrenGap: 12 }}>
-                <Stack horizontal horizontalAlign="space-between" styles={{ root: { minHeight: 32 }}}>
-					<Text variant="mediumPlus" styles={{ root: { fontFamily: "Horde Open Sans SemiBold" } }}>Step Trace</Text>
-				</Stack>
-                { traceState.data && 
-                    <FlameGraph
-                        data={traceState.data}
-                        height={200}
-                        width={1014}
-                        onChange={(node: { name: any; }) => {
-                            console.log(`"${node.name}" focused`);
-                        }}
-                    />
-                }
-			</Stack>
-		</Stack></Stack>);
+   traceState.getTrace(jobDetails, stepId!);
+
+   // let height = Math.min(36 * artifacts.length + 60, 500);	
+   return (<Stack styles={{ root: { paddingTop: 18, paddingRight: 12 } }}>
+      <Stack className={hordeClasses.raised}>
+         <Stack tokens={{ childrenGap: 12 }}>
+            <Stack horizontal horizontalAlign="space-between" styles={{ root: { minHeight: 32 } }}>
+               <Text variant="mediumPlus" styles={{ root: { fontFamily: "Horde Open Sans SemiBold" } }}>Step Trace</Text>
+            </Stack>
+            {traceState.data &&
+               <FlameGraph
+                  data={traceState.data}
+                  height={200}
+                  width={1014}
+                  onChange={(node: { name: any; }) => {
+                     console.log(`"${node.name}" focused`);
+                  }}
+               />
+            }
+         </Stack>
+      </Stack></Stack>);
 });
\ No newline at end of file
diff --git a/Engine/Source/Programs/Horde/HordeDashboard/src/components/JobFilterBar.tsx b/Engine/Source/Programs/Horde/HordeDashboard/src/components/JobFilterBar.tsx
--- a/Engine/Source/Programs/Horde/HordeDashboard/src/components/JobFilterBar.tsx
+++ b/Engine/Source/Programs/Horde/HordeDashboard/src/components/JobFilterBar.tsx
@@ -1,6 +1,6 @@
 // Copyright Epic Games, Inc. All Rights Reserved.
 
-import { action, observable } from "mobx";
+import { action, makeObservable, observable } from "mobx";
 import { Dropdown, IDropdownOption, Stack } from "@fluentui/react";
 import React, { useEffect, useState } from "react";
 import { useBackend } from "../backend";
@@ -38,6 +38,10 @@ const dropDownStyle: any = () => {
 
 export class JobFilter {
 
+   constructor() {
+      makeObservable(this);
+   }
+
     @action
     setUpdated() {
         this.updated++;
diff --git a/Engine/Source/Programs/Horde/HordeDashboard/src/components/JobView.tsx b/Engine/Source/Programs/Horde/HordeDashboard/src/components/JobView.tsx
--- a/Engine/Source/Programs/Horde/HordeDashboard/src/components/JobView.tsx
+++ b/Engine/Source/Programs/Horde/HordeDashboard/src/components/JobView.tsx
@@ -451,7 +451,7 @@ const JobList: React.FC<{ tab: string; filter: JobFilterSimple, controller: Call
 
    function onRenderItemColumnInner(item: JobItem, index?: number, column?: IColumn) {
 
-      const fieldContent = item[column!.fieldName as keyof JobItem];
+      const fieldContent = item[column!.fieldName as keyof JobItem] as string;
 
       const jobColumn = item.jobTab?.columns?.find(c => c.heading === column!.key);
 
diff --git a/Engine/Source/Programs/Horde/HordeDashboard/src/components/JobViewAll.tsx b/Engine/Source/Programs/Horde/HordeDashboard/src/components/JobViewAll.tsx
--- a/Engine/Source/Programs/Horde/HordeDashboard/src/components/JobViewAll.tsx
+++ b/Engine/Source/Programs/Horde/HordeDashboard/src/components/JobViewAll.tsx
@@ -69,7 +69,7 @@ const JobFilterPanel: React.FC = () => {
     const { streamId } = useParams<{ streamId: string }>();
 
     return <Stack>
-        <JobFilterBar streamId={streamId} />
+        <JobFilterBar streamId={streamId!} />
     </Stack>
 }
 
diff --git a/Engine/Source/Programs/Horde/HordeDashboard/src/components/JobViewIncremental.tsx b/Engine/Source/Programs/Horde/HordeDashboard/src/components/JobViewIncremental.tsx
--- a/Engine/Source/Programs/Horde/HordeDashboard/src/components/JobViewIncremental.tsx
+++ b/Engine/Source/Programs/Horde/HordeDashboard/src/components/JobViewIncremental.tsx
@@ -470,7 +470,7 @@ const JobList: React.FC<{ tab: string; filter: JobFilterSimple, controller: Call
 
    function onRenderItemColumnInner(item: JobItem, index?: number, column?: IColumn) {
 
-      const fieldContent = item[column!.fieldName as keyof JobItem];
+      const fieldContent = item[column!.fieldName as keyof JobItem] as string;
 
       const jobColumn = item.jobTab?.columns?.find(c => c.heading === column!.key);
 
@@ -504,7 +504,7 @@ const JobList: React.FC<{ tab: string; filter: JobFilterSimple, controller: Call
 
       }
 
-      const commit = commitCache.getCommit(streamId, item.job!.change!);
+      const commit = commitCache.getCommit(streamId!, item.job!.change!);
 
       if (column!.key === "Author") {
 
diff --git a/Engine/Source/Programs/Horde/HordeDashboard/src/components/LogSource.tsx b/Engine/Source/Programs/Horde/HordeDashboard/src/components/LogSource.tsx
--- a/Engine/Source/Programs/Horde/HordeDashboard/src/components/LogSource.tsx
+++ b/Engine/Source/Programs/Horde/HordeDashboard/src/components/LogSource.tsx
@@ -1,6 +1,6 @@
 // Copyright Epic Games, Inc. All Rights Reserved.
 
-import { action, observable } from "mobx";
+import { action, makeObservable, observable } from "mobx";
 import moment, { Moment } from 'moment-timezone';
 import backend from '../backend';
 import { AgentData, BatchData, EventSeverity, GetLogEventResponse, IssueData, LeaseData, LogData, NodeData, StepData, StreamData } from "../backend/Api";
@@ -13,624 +13,628 @@ const stripNewlines = /(\r\n|\n|\r)/gm;
 
 export abstract class LogSource {
 
-    init(): Promise<void> {
-        return new Promise<void>((resolve, reject) => {
-            this.updateLogData().then(() => {
-                resolve();
-            }).catch(reason => reject(reason));
-        });
-    }
+   constructor() {
+      makeObservable(this);
+   }
 
-    get errors(): GetLogEventResponse[] {
-        return [];
-    }
+   init(): Promise<void> {
+      return new Promise<void>((resolve, reject) => {
+         this.updateLogData().then(() => {
+            resolve();
+         }).catch(reason => reject(reason));
+      });
+   }
 
-    get warnings(): GetLogEventResponse[] {
-        return [];
-    }
+   get errors(): GetLogEventResponse[] {
+      return [];
+   }
 
-    get events(): GetLogEventResponse[] {
-        return [];
-    }
+   get warnings(): GetLogEventResponse[] {
+      return [];
+   }
 
-    get crumbs(): BreadcrumbItem[] {
-        return [];
-    }
+   get events(): GetLogEventResponse[] {
+      return [];
+   }
 
-    get crumbTitle(): string | undefined {
-        return undefined;
-    }
+   get crumbs(): BreadcrumbItem[] {
+      return [];
+   }
 
-    get summary(): string {
-        return "";
-    }
+   get crumbTitle(): string | undefined {
+      return undefined;
+   }
 
-    get issues(): IssueData[] {
-        return [];
-    }
+   get summary(): string {
+      return "";
+   }
 
-    get percentComplete(): number | undefined {
-        return undefined;
-    }
+   get issues(): IssueData[] {
+      return [];
+   }
 
-    getAgentId(): string | undefined {
-        return "";
-    }
+   get percentComplete(): number | undefined {
+      return undefined;
+   }
 
-    getLeaseId(): string | undefined {
-        return undefined;
-    }
+   getAgentId(): string | undefined {
+      return "";
+   }
 
-    initComplete() {
-        if (this.active) {
-            setTimeout(() => { this.poll(); }, this.pollMS);
-        }
-    }
+   getLeaseId(): string | undefined {
+      return undefined;
+   }
 
-    download(json: boolean) {
+   initComplete() {
+      if (this.active) {
+         setTimeout(() => { this.poll(); }, this.pollMS);
+      }
+   }
 
-    }
+   download(json: boolean) {
 
-    resize(lineCount: number): LogItem[] | undefined {
+   }
 
-        // resize items if necessary
-        let count = lineCount - this.logItems.length;
+   resize(lineCount: number): LogItem[] | undefined {
 
-        if (count <= 0) {
-            return undefined;
-       }
-       
-        const items = [...this.logItems];
+      // resize items if necessary
+      let count = lineCount - this.logItems.length;
 
-        let line = this.logItems.length + 1;
-        while (count--) {
-            items.push({
-                lineNumber: line++,
-                requested: false
-            });
-        }
+      if (count <= 0) {
+         return undefined;
+      }
 
-        return items;
+      const items = [...this.logItems];
 
-    }
+      let line = this.logItems.length + 1;
+      while (count--) {
+         items.push({
+            lineNumber: line++,
+            requested: false
+         });
+      }
 
-    loadLines(index: number, count: number): Promise<boolean> {
+      return items;
 
-        return new Promise<boolean>((resolve, reject) => {
+   }
 
-            if (!this.logData) {
-                return reject("invalid log data");
-            }
+   loadLines(index: number, count: number): Promise<boolean> {
 
-            let anyRequested = false;
-
-            for (let i = 0; i < count; i++) {
-                const offset = i + index;
-                if (offset >= this.logItems.length) {
-                    break;
-                }
-                const item = this.logItems![i + index];
-                if (!item.requested) {
-                    item.requested = true;
-                    anyRequested = true;
-                }
-            }
+      return new Promise<boolean>((resolve, reject) => {
+
+         if (!this.logData) {
+            return reject("invalid log data");
+         }
 
-            if (!anyRequested) {
-                return resolve(false);
+         let anyRequested = false;
+
+         for (let i = 0; i < count; i++) {
+            const offset = i + index;
+            if (offset >= this.logItems.length) {
+               break;
+            }
+            const item = this.logItems![i + index];
+            if (!item.requested) {
+               item.requested = true;
+               anyRequested = true;
             }
+         }
 
-            backend.getLogLines(this.logData.id, index, count).then(data => {
+         if (!anyRequested) {
+            return resolve(false);
+         }
 
-                const start = this.startTime;
+         backend.getLogLines(this.logData.id, index, count).then(data => {
 
-                if (!start) {
-                    return reject("Invalid log start time");
-                }
-                for (let i = 0; i < data.count; i++) {
-                    const line = data.lines![i];
+            const start = this.startTime;
 
-                    const offset = i + data.index;
-                    if (offset >= this.logItems.length) {
-                        break;
-                    }
-                    const item = this.logItems[offset];
-                    item.line = line;
+            if (!start) {
+               return reject("Invalid log start time");
+            }
+            for (let i = 0; i < data.count; i++) {
+               const line = data.lines![i];
 
+               const offset = i + data.index;
+               if (offset >= this.logItems.length) {
+                  break;
+               }
+               const item = this.logItems[offset];
+               item.line = line;
 
-                    const fixWhitespace = (text: string | undefined): string | undefined => {
-                        return text?.replace(stripNewlines, "").trimEnd();
-                    };
 
-                    line.message = fixWhitespace(line.message)!;
-                    line.format = fixWhitespace(line.format);
+               const fixWhitespace = (text: string | undefined): string | undefined => {
+                  return text?.replace(stripNewlines, "").trimEnd();
+               };
 
-                }
+               line.message = fixWhitespace(line.message)!;
+               line.format = fixWhitespace(line.format);
 
-                this.setLogItems([...this.logItems]);
-                resolve(true);
+            }
 
-            }).catch(reason => reject(reason));
-        });
+            this.setLogItems([...this.logItems]);
+            resolve(true);
 
-    }
+         }).catch(reason => reject(reason));
+      });
 
+   }
 
-    private updateLogData(): Promise<void> {
 
-        return new Promise<void>((resolve, reject) => {
+   private updateLogData(): Promise<void> {
 
-            backend.getLogData(this.logId).then(data => {
+      return new Promise<void>((resolve, reject) => {
 
-                this.logData = data!;
+         backend.getLogData(this.logId).then(data => {
 
-                if (this.logData.lineCount) {
-                    const items = this.resize(this.logData!.lineCount);
-                    if (items) {
-                        this.setLogItems(items);
-                    }
-                }
+            this.logData = data!;
 
-                resolve();
+            if (this.logData.lineCount) {
+               const items = this.resize(this.logData!.lineCount);
+               if (items) {
+                  this.setLogItems(items);
+               }
+            }
 
-            }).catch(reason => {
-                reject(reason);
-            });
-        });
+            resolve();
 
-    }
+         }).catch(reason => {
+            reject(reason);
+         });
+      });
 
-    poll() {
+   }
 
-        if (!this.logData) {
-            return;
-        }
+   poll() {
 
-        if (!this.active) {
-            this.stopPolling();
-            return;
-        }
+      if (!this.logData) {
+         return;
+      }
 
-        if (!this.pollMS) {
-            return;
-        }
+      if (!this.active) {
+         this.stopPolling();
+         return;
+      }
 
-        this.pollID = setTimeout(() => { this.poll(); }, this.pollMS);
+      if (!this.pollMS) {
+         return;
+      }
 
-        if (this.polling) {
-            return;
-        }
+      this.pollID = setTimeout(() => { this.poll(); }, this.pollMS);
 
-        this.polling = true;
+      if (this.polling) {
+         return;
+      }
 
-        this.pollInner().then(async () => {
-            await this.updateLogData().catch(reason => {
-                console.log(`Error updating log data:\n${reason}`);
-            });
-        }).finally(() => {
-            this.polling = false;
-        });
+      this.polling = true;
 
+      this.pollInner().then(async () => {
+         await this.updateLogData().catch(reason => {
+            console.log(`Error updating log data:\n${reason}`);
+         });
+      }).finally(() => {
+         this.polling = false;
+      });
 
-    }
 
-    stopPolling() {
-        clearTimeout(this.pollID);
-        this.pollMS = 0;
-        this.pollID = undefined;
-    }
+   }
 
-    pollInner(): Promise<void> {
-        return new Promise<void>((resolve, reject) => {
-            resolve();
-        });
-    }
+   stopPolling() {
+      clearTimeout(this.pollID);
+      this.pollMS = 0;
+      this.pollID = undefined;
+   }
+
+   pollInner(): Promise<void> {
+      return new Promise<void>((resolve, reject) => {
+         resolve();
+      });
+   }
 
-    clear() {
-        this.stopPolling();
-    }
+   clear() {
+      this.stopPolling();
+   }
 
-    @action
-    setLogItems(items: LogItem[]) {
-        this.logItems = items;
-    }
+   @action
+   setLogItems(items: LogItem[]) {
+      this.logItems = items;
+   }
 
-    @action
-    setActive(active: boolean) {
-        this.active = active;
-    }
+   @action
+   setActive(active: boolean) {
+      this.active = active;
+   }
 
-    @action
-    setFatalError(error: string) {
-        this.fatalError = error;
-    }
+   @action
+   setFatalError(error: string) {
+      this.fatalError = error;
+   }
 
-    @observable.ref
-    logItems: LogItem[] = [];
+   @observable.ref
+   logItems: LogItem[] = [];
 
-    @observable
-    active = true;
+   @observable
+   active = true;
 
-    @observable
-    fatalError?: string;
+   @observable
+   fatalError?: string;
 
-    logId = "";
-    logData?: LogData;
-    startLine?: number;
-    trailing?: boolean;
+   logId = "";
+   logData?: LogData;
+   startLine?: number;
+   trailing?: boolean;
 
-    pollID?: any = undefined;
-    pollMS = 10000;
-    polling = false;
+   pollID?: any = undefined;
+   pollMS = 10000;
+   polling = false;
 
-    startTime?: Moment;
+   startTime?: Moment;
 
-    query?: URLSearchParams;
+   query?: URLSearchParams;
 
-    static create(logId: string, query: URLSearchParams): LogSource {
+   static create(logId: string, query: URLSearchParams): LogSource {
 
-        const leaseId = query.get("leaseId");
+      const leaseId = query.get("leaseId");
 
-        let line: number | undefined;
-        if (query.get("lineindex")) {
-            line = parseInt(query.get("lineindex")!);
-        }
+      let line: number | undefined;
+      if (query.get("lineindex")) {
+         line = parseInt(query.get("lineindex")!);
+      }
 
-        const source = leaseId ? new LeaseLogSource(leaseId) : new JobLogSource();
-        source.startLine = line;
-        source.logId = logId;
-        source.query = query;
+      const source = leaseId ? new LeaseLogSource(leaseId) : new JobLogSource();
+      source.startLine = line;
+      source.logId = logId;
+      source.query = query;
 
-        return source;
+      return source;
 
-    }
+   }
 }
 
 export class JobLogSource extends LogSource {
 
-    init(): Promise<void> {
-        return new Promise<void>((resolve, reject) => {
-            super.init().then(() => {
+   init(): Promise<void> {
+      return new Promise<void>((resolve, reject) => {
+         super.init().then(() => {
 
-                const details = this.jobDetails;
+            const details = this.jobDetails;
 
-                const done = () => {
-                    if (details.fatalError) {
-                        this.setFatalError(details.fatalError);
-                    }
-                    this.detailsUpdated();
-                    this.refreshJobData();
-                    this.initComplete();
-                    resolve();
-                };
+            const done = () => {
+               if (details.fatalError) {
+                  this.setFatalError(details.fatalError);
+               }
+               this.detailsUpdated();
+               this.refreshJobData();
+               this.initComplete();
+               resolve();
+            };
+
+            // @todo: refactor with job details, we need to get stepId from logId is the problem
+            details.set(this.logData!.jobId, this.logId, undefined, undefined, 10000, () => {
+               const step = details.getSteps().find(s => s.logId === this.logId);
+               if (step) {
+                  details.set(details.id!, this.logId, step.id, undefined, 10000, undefined, () => {
+                     done();
+                  });
+               } else {
+                  done();
+               }
 
-                // @todo: refactor with job details, we need to get stepId from logId is the problem
-                details.set(this.logData!.jobId, this.logId, undefined, undefined, 10000, () => {
-                    const step = details.getSteps().find(s => s.logId === this.logId);
-                    if (step) {
-                        details.set(details.id!, this.logId, step.id, undefined, 10000, undefined, () => {
-                            done();
-                        });
-                    } else {
-                        done();
-                    }
+            });
+         }).catch(reason => reject(reason));
+      });
+   }
 
-                });
-            }).catch(reason => reject(reason));
-        });
-    }
+   clear() {
+      super.clear();
+      this.jobDetails.clear();
+   }
 
-    clear() {
-        super.clear();
-        this.jobDetails.clear();
-    }
+   get errors(): GetLogEventResponse[] {
+      return this.jobDetails.events.filter(e => e.severity === EventSeverity.Error);
+   }
 
-    get errors(): GetLogEventResponse[] {
-        return this.jobDetails.events.filter(e => e.severity === EventSeverity.Error);
-    }
+   get warnings(): GetLogEventResponse[] {
+      return this.jobDetails.events.filter(e => e.severity === EventSeverity.Warning);
+   }
 
-    get warnings(): GetLogEventResponse[] {
-        return this.jobDetails.events.filter(e => e.severity === EventSeverity.Warning);
-    }
+   get events(): GetLogEventResponse[] {
+      return this.jobDetails.events;
+   }
 
-    get events(): GetLogEventResponse[] {
-        return this.jobDetails.events;
-    }
 
+   get issues(): IssueData[] {
+      return this.jobDetails.issues;
+   }
 
-    get issues(): IssueData[] {
-        return this.jobDetails.issues;
-    }
+   get summary(): string {
 
-    get summary(): string {
+      let text = "";
+      const step = this.step;
+      const batch = this.batch;
 
-        let text = "";
-        const step = this.step;
-        const batch = this.batch;
+      if (step) {
+         return getStepSummaryMarkdown(this.jobDetails, step.id);
+      }
+      if (batch) {
+         return getBatchSummaryMarkdown(this.jobDetails, batch.id);
+      }
 
-        if (step) {
-            return getStepSummaryMarkdown(this.jobDetails, step.id);
-        }
-        if (batch) {
-            return getBatchSummaryMarkdown(this.jobDetails, batch.id);
-        }
+      return text;
 
-        return text;
-        
-    }
+   }
 
-    get percentComplete(): number | undefined {
+   get percentComplete(): number | undefined {
 
-        if (this.step && this.node) {
-            return getStepPercent(this.step);
-        }
+      if (this.step && this.node) {
+         return getStepPercent(this.step);
+      }
 
-        return undefined;
-    }
+      return undefined;
+   }
 
-    getAgentId(): string {
-        return this.agentId ?? "";
-    }
+   getAgentId(): string {
+      return this.agentId ?? "";
+   }
 
-    download(json: boolean) {
+   download(json: boolean) {
 
-        const name = `++${this.projectName}+${this.stream?.name}-CL ${this.change}-${this.jobName}.${json ? "jsonl" : "log"}`;
-        backend.downloadLog(this.logId, name, json);
+      const name = `++${this.projectName}+${this.stream?.name}-CL ${this.change}-${this.jobName}.${json ? "jsonl" : "log"}`;
+      backend.downloadLog(this.logId, name, json);
 
-    }
+   }
 
-    private detailsUpdated() {
+   private detailsUpdated() {
 
-        const details = this.jobDetails;
+      const details = this.jobDetails;
 
-        if (!this.startTime) {
-            const details = this.jobDetails;
-            const batch = details.batches.find((b) => b.logId === this.logId);
-            const step = details.getSteps().find(s => s.logId === this.logId);
-            this.startTime = moment.utc(step?.startTime ?? batch?.startTime);
-        }
+      if (!this.startTime) {
+         const details = this.jobDetails;
+         const batch = details.batches.find((b) => b.logId === this.logId);
+         const step = details.getSteps().find(s => s.logId === this.logId);
+         this.startTime = moment.utc(step?.startTime ?? batch?.startTime);
+      }
 
-        const active = details.getLogActive(this.logData!.id);
+      const active = details.getLogActive(this.logData!.id);
 
-        if (this.active !== active) {
-            this.setActive(active);
-        }
-    }
+      if (this.active !== active) {
+         this.setActive(active);
+      }
+   }
 
-    private refreshJobData() {
+   private refreshJobData() {
 
-        const details = this.jobDetails;
+      const details = this.jobDetails;
 
-        this.jobName = "UnknownJob";
-        this.projectName = "";
-        this.change = "";
-        this.agentId = undefined;
-        this.batch = undefined;
-        this.step = undefined;
-        this.stream = undefined;
+      this.jobName = "UnknownJob";
+      this.projectName = "";
+      this.change = "";
+      this.agentId = undefined;
+      this.batch = undefined;
+      this.step = undefined;
+      this.stream = undefined;
 
-        this.stream = details.stream;
+      this.stream = details.stream;
 
-        if (this.stream) {
+      if (this.stream) {
 
-            if (this.stream.project) {
-                this.projectName = this.stream.project?.name ?? "Unknown Project";
-                if (this.projectName === "Engine") {
-                    this.projectName = "UE4";
-                }
+         if (this.stream.project) {
+            this.projectName = this.stream.project?.name ?? "Unknown Project";
+            if (this.projectName === "Engine") {
+               this.projectName = "UE4";
             }
-        }
+         }
+      }
 
-        this.change = details.jobdata?.change ?? "Latest";
+      this.change = details.jobdata?.change ?? "Latest";
 
-        this.batch = details.batches.find((b) => b.logId === this.logId);
-        this.step = details.getSteps().find(s => s.logId === this.logId);
+      this.batch = details.batches.find((b) => b.logId === this.logId);
+      this.step = details.getSteps().find(s => s.logId === this.logId);
 
-        this.agentId = this.batch?.agentId ?? details.batchByStepId(this.step?.id!)?.agentId;
+      this.agentId = this.batch?.agentId ?? details.batchByStepId(this.step?.id!)?.agentId;
 
-        this.node = undefined;
+      this.node = undefined;
 
-        if (this.step) {
-            this.node = details.nodeByStepId(this.step.id);
-            this.jobName = details.getStepName(this.step.id) ?? "Unknown Step Node";
-        }
+      if (this.step) {
+         this.node = details.nodeByStepId(this.step.id);
+         this.jobName = details.getStepName(this.step.id) ?? "Unknown Step Node";
+      }
 
-        if (this.batch) {
-            this.jobName = `Batch-${this.batch.id}`;
-        }
+      if (this.batch) {
+         this.jobName = `Batch-${this.batch.id}`;
+      }
 
-    }
+   }
 
-    get clText(): string {
+   get clText(): string {
 
-        const data = this.jobDetails.jobdata;
-        if (!data) {
-            return "";
-        }
+      const data = this.jobDetails.jobdata;
+      if (!data) {
+         return "";
+      }
 
-        let clText = "";
-        if (data.preflightChange) {
-            clText = `Preflight ${data.preflightChange} `;
-            clText += ` - Base ${data.change ? data.change : "Latest"}`;
+      let clText = "";
+      if (data.preflightChange) {
+         clText = `Preflight ${data.preflightChange} `;
+         clText += ` - Base ${data.change ? data.change : "Latest"}`;
 
-        } else {
-            clText = `${data.change ? "CL " + data.change : "Latest CL"}`;
-        }
+      } else {
+         clText = `${data.change ? "CL " + data.change : "Latest CL"}`;
+      }
 
-        return clText;
-    }
+      return clText;
+   }
 
-    get crumbs(): BreadcrumbItem[] {
+   get crumbs(): BreadcrumbItem[] {
 
-        const data = this.jobDetails.jobdata!;
+      const data = this.jobDetails.jobdata!;
 
 
-        if (!this.stream) {
-            return [];
-        }
+      if (!this.stream) {
+         return [];
+      }
 
-        let projectName = this.stream.project?.name;
-        if (projectName === "Engine") {
-            projectName = "UE4";
-        }
+      let projectName = this.stream.project?.name;
+      if (projectName === "Engine") {
+         projectName = "UE4";
+      }
 
-        const crumbItems: BreadcrumbItem[] = [
-            {
-                text: projectName ?? "Unknown Project",
-                link: `/project/${this.stream?.project?.id}`
-            },
-            {
-                text: this.stream.name,
-                link: `/stream/${this.stream.id}`
-            },
-            {
-                text: `${this.clText}: ${data?.name ?? ""}`,
-                link: `/job/${data?.id}`
-            },
-            {
-                text: this.jobName,
-                link: this.step ? `/job/${data?.id}?step=${this.step?.id}` : `/job/${data?.id}?batch=${this.batch?.id}`
-            },
-            {
-                text: this.jobName + " (Log)"
+      const crumbItems: BreadcrumbItem[] = [
+         {
+            text: projectName ?? "Unknown Project",
+            link: `/project/${this.stream?.project?.id}`
+         },
+         {
+            text: this.stream.name,
+            link: `/stream/${this.stream.id}`
+         },
+         {
+            text: `${this.clText}: ${data?.name ?? ""}`,
+            link: `/job/${data?.id}`
+         },
+         {
+            text: this.jobName,
+            link: this.step ? `/job/${data?.id}?step=${this.step?.id}` : `/job/${data?.id}?batch=${this.batch?.id}`
+         },
+         {
+            text: this.jobName + " (Log)"
 
-            },
+         },
 
 
-        ];
+      ];
 
-        return crumbItems;
+      return crumbItems;
 
-    }
+   }
 
-    get crumbTitle(): string | undefined {
+   get crumbTitle(): string | undefined {
 
-        if (!this.stream) {
-            return undefined;
-        }
-        return `Horde - ${this.clText}: ${this.jobDetails.jobdata?.name ?? ""} - ${this.jobName}`
-    }
+      if (!this.stream) {
+         return undefined;
+      }
+      return `Horde - ${this.clText}: ${this.jobDetails.jobdata?.name ?? ""} - ${this.jobName}`
+   }
 
-    stream?: StreamData;
-    jobName = "";
-    projectName = "";
-    change: string | number = "";
-    agentId?: string;
-    batch?: BatchData;
-    step?: StepData;
-    node?: NodeData;
+   stream?: StreamData;
+   jobName = "";
+   projectName = "";
+   change: string | number = "";
+   agentId?: string;
+   batch?: BatchData;
+   step?: StepData;
+   node?: NodeData;
 
-    jobDetails: JobDetails = new JobDetails(undefined, undefined, undefined, true);
+   jobDetails: JobDetails = new JobDetails(undefined, undefined, undefined, true);
 }
 
 class LeaseLogSource extends LogSource {
 
-    constructor(leaseId: string) {
-        super();
-        this.leaseId = leaseId;
-    }
+   constructor(leaseId: string) {
+      super();
+      this.leaseId = leaseId;
+   }
 
-    getAgentId(): string | undefined {
-        return this.lease?.agentId;
-    }
+   getAgentId(): string | undefined {
+      return this.lease?.agentId;
+   }
 
-    getLeaseId(): string | undefined {
-        return this.leaseId;
-    }
+   getLeaseId(): string | undefined {
+      return this.leaseId;
+   }
 
-    get summary(): string {
+   get summary(): string {
 
-        let text = "";
+      let text = "";
 
-        const duration = getLeaseElapsed(this.lease);
+      const duration = getLeaseElapsed(this.lease);
 
-        if (duration) {
+      if (duration) {
 
-                text = `${!!this.lease?.finishTime ? "Ran" : "Running"} on [${this.agent?.id}](?agentId=${encodeURIComponent(this.agent?.id ?? "")}) for ${duration}`    
-        }
+         text = `${!!this.lease?.finishTime ? "Ran" : "Running"} on [${this.agent?.id}](?agentId=${encodeURIComponent(this.agent?.id ?? "")}) for ${duration}`
+      }
 
-        return text;
-    }
+      return text;
+   }
 
 
-    init(): Promise<void> {
-        return new Promise<void>((resolve, reject) => {
-            super.init().then(() => {
-                Promise.all([backend.getLease(this.leaseId)]).then(async (values) => {
+   init(): Promise<void> {
+      return new Promise<void>((resolve, reject) => {
+         super.init().then(() => {
+            Promise.all([backend.getLease(this.leaseId)]).then(async (values) => {
 
-                    this.lease = values[0];
+               this.lease = values[0];
 
-                    this.agent = await backend.getAgent(this.lease.agentId!);
+               this.agent = await backend.getAgent(this.lease.agentId!);
 
-                    this.leaseUpdated();
-                    this.initComplete();
-                    resolve();
+               this.leaseUpdated();
+               this.initComplete();
+               resolve();
 
-                });
-            }).catch(reason => reject(reason));
-        });
-    }
+            });
+         }).catch(reason => reject(reason));
+      });
+   }
 
-    download(json: boolean) {
-        const name = `${this.agent!.name}-Lease-${this.leaseId}.${json ? "jsonl" : "log"}`;
-        backend.downloadLog(this.logId, name, json);
+   download(json: boolean) {
+      const name = `${this.agent!.name}-Lease-${this.leaseId}.${json ? "jsonl" : "log"}`;
+      backend.downloadLog(this.logId, name, json);
 
-    }
+   }
 
-    pollInner(): Promise<void> {
-        return new Promise<void>((resolve, reject) => {
+   pollInner(): Promise<void> {
+      return new Promise<void>((resolve, reject) => {
 
-            if (!this.lease) {
-                resolve();
-                return;
-            }
+         if (!this.lease) {
+            resolve();
+            return;
+         }
 
-            Promise.all([backend.getAgent(this.getAgentId()!), backend.getLease(this.leaseId)]).then((values => {
-                this.agent = values[0];
-                this.lease = values[1];
-                this.leaseUpdated();
-                resolve();
-            })).catch(reason => reject(reason));
-        });
-    }
-
-    leaseUpdated() {
-
-        this.startTime = moment.utc(this.lease!.startTime);
-
-        const active = !this.lease!.finishTime;
-
-        if (this.active !== active) {
-            this.setActive(active);
-        }
-    }
-
-    get crumbs(): BreadcrumbItem[] {
-
-        const crumbItems: BreadcrumbItem[] = [
-            {
-                text: `Agents`,
-                link: '/agents'
-            },
-            {
-                text: `${this.agent?.name}`,
-                link: `/agents/${this.agent?.id}`
-            },
-            {
-                text: `${this.lease?.name ?? this.leaseId}`,
-                link: `/agents/${this.agent?.id}`
-            }
+         Promise.all([backend.getAgent(this.getAgentId()!), backend.getLease(this.leaseId)]).then((values => {
+            this.agent = values[0];
+            this.lease = values[1];
+            this.leaseUpdated();
+            resolve();
+         })).catch(reason => reject(reason));
+      });
+   }
+
+   leaseUpdated() {
+
+      this.startTime = moment.utc(this.lease!.startTime);
+
+      const active = !this.lease!.finishTime;
+
+      if (this.active !== active) {
+         this.setActive(active);
+      }
+   }
+
+   get crumbs(): BreadcrumbItem[] {
+
+      const crumbItems: BreadcrumbItem[] = [
+         {
+            text: `Agents`,
+            link: '/agents'
+         },
+         {
+            text: `${this.agent?.name}`,
+            link: `/agents/${this.agent?.id}`
+         },
+         {
+            text: `${this.lease?.name ?? this.leaseId}`,
+            link: `/agents/${this.agent?.id}`
+         }
 
-        ];
+      ];
 
-        return crumbItems;
+      return crumbItems;
 
-    }
+   }
 
 
-    leaseId: string;
-    agent?: AgentData;
-    lease?: LeaseData;
+   leaseId: string;
+   agent?: AgentData;
+   lease?: LeaseData;
 
 
 }
\ No newline at end of file
diff --git a/Engine/Source/Programs/Horde/HordeDashboard/src/components/LogView.tsx b/Engine/Source/Programs/Horde/HordeDashboard/src/components/LogView.tsx
--- a/Engine/Source/Programs/Horde/HordeDashboard/src/components/LogView.tsx
+++ b/Engine/Source/Programs/Horde/HordeDashboard/src/components/LogView.tsx
@@ -1,11 +1,11 @@
 // Copyright Epic Games, Inc. All Rights Reserved.
 
 import { DefaultButton, DetailsList, DetailsListLayoutMode, DetailsRow, DirectionalHint, Dropdown, FocusZone, FocusZoneDirection, IColumn, Icon, IconButton, IContextualMenuProps, IDetailsListProps, ITextField, List, Modal, ProgressIndicator, ScrollToMode, Selection, SelectionMode, SelectionZone, Separator, Spinner, SpinnerSize, Stack, Text, TextField } from '@fluentui/react';
-import { action, observable } from 'mobx';
+import { action, makeObservable, observable } from 'mobx';
 import { observer } from 'mobx-react-lite';
 import moment from 'moment-timezone';
 import React, { useEffect, useState } from 'react';
-import { Link, useHistory, useLocation, useParams } from 'react-router-dom';
+import { Link, useNavigate, useLocation, useParams } from 'react-router-dom';
 import backend from '../backend';
 import { ArtifactData, EventSeverity, GetChangeSummaryResponse, GetJobStepRefResponse, GetLogEventResponse, LogLevel } from '../backend/Api';
 import { CommitCache } from '../backend/CommitCache';
@@ -31,6 +31,7 @@ import { TopNav } from './TopNav';
 class LogHandler {
 
    constructor() {
+      makeObservable(this);
       LogHandler.instance = this;
    }
 
@@ -187,7 +188,7 @@ const commitCache = new CommitCache();
 
 const StepHistoryModal: React.FC<{ jobDetails: JobDetails, stepId: string | undefined, onClose: () => void }> = observer(({ jobDetails, stepId, onClose }) => {
 
-   const history = useHistory();
+   const navigate = useNavigate();
    const location = useLocation();
    const [commitState, setCommitState] = useState<{ target?: ChangeContextMenuTarget, commit?: GetChangeSummaryResponse, rangeCL?: number }>({});
 
@@ -274,7 +275,7 @@ const StepHistoryModal: React.FC<{ jobDetails: JobDetails, stepId: string | unde
          const url = `${location.pathname}?agentId=${agentId}`;
 
          return <Stack verticalAlign="center" horizontalAlign="center" tokens={{ childrenGap: 0, padding: 0 }} style={{ width: "100%", height: "100%" }}>
-            <a href={url} onClick={(ev) => { ev.preventDefault(); ev.stopPropagation(); history.replace(url); }}><Stack horizontal horizontalAlign={"end"} verticalFill={true} tokens={{ childrenGap: 0, padding: 0 }}><Text>{agentId}</Text></Stack></a>
+            <a href={url} onClick={(ev) => { ev.preventDefault(); ev.stopPropagation(); navigate(url, {replace: true}); }}><Stack horizontal horizontalAlign={"end"} verticalFill={true} tokens={{ childrenGap: 0, padding: 0 }}><Text>{agentId}</Text></Stack></a>
          </Stack>
       }
 
@@ -414,7 +415,7 @@ export const LogList: React.FC<{ logId: string }> = observer(({ logId }) => {
    const windowSize = useWindowSize();
    const query = useQuery();
    const location = useLocation();
-   const history = useHistory();
+   const navigate = useNavigate();
    const [tsFormat, setTSFormat] = useState('relative');
    const [handler, setHandler] = useState<LogHandler | undefined>(undefined);
    const [searchState, setSearchState] = useState<{ search?: string, results?: number[], curRequest?: any }>({});
@@ -617,7 +618,7 @@ export const LogList: React.FC<{ logId: string }> = observer(({ logId }) => {
                      ev.stopPropagation();
                      location.search = `?issue=${issueId}`;
                      setIssueHistory(true);
-                     history.push(location);
+                     navigate(location);
                   }}><Text variant="small" style={{ ...fontStyle }}>Issue</Text><div style={{ ...fontStyle }}>&nbsp;</div><Text variant="small" style={{ ...fontStyle }}>{`${issueId}`}</Text></DefaultButton>
             </Stack>
          }
@@ -649,7 +650,7 @@ export const LogList: React.FC<{ logId: string }> = observer(({ logId }) => {
          }
       
          return (
-            <Stack key={`key_log_line_${item.lineNumber}`} style={{ width: "max-content", height: handler.lineHeight }} onClick={() => {history.replace(location.pathname + `?lineindex=${item.lineNumber - 1}`)} }>
+            <Stack key={`key_log_line_${item.lineNumber}`} style={{ width: "max-content", height: handler.lineHeight }} onClick={() => {navigate(location.pathname + `?lineindex=${item.lineNumber - 1}`, {replace: true})} }>
                <div style={{ position: "relative" }}>
                   <Stack className={styles.logLine} tokens={{ childrenGap: 8 }} horizontal disableShrink={true}>
                      <Stack styles={{ root: { color: "#c0c0c0", width: 80, textAlign: "right", userSelect: "none", fontSize: handler.fontSize } }}>{prefix + item.lineNumber}</Stack>
@@ -923,7 +924,7 @@ export const LogList: React.FC<{ logId: string }> = observer(({ logId }) => {
       {!!fixme && <IssueModalV2 issueId={query.get("issue")} popHistoryOnClose={issueHistory} />}
       {!!fixme && logHistory && <StepHistoryModal jobDetails={fixme!} stepId={fixme!.stepByLogId(logId)?.id} onClose={() => setLogHistory(false)} />}
       {!!fixme && logArtifacts && <StepArtifactsModal jobDetails={fixme!} stepId={fixme!.stepByLogId(logId)?.id} onClose={() => setLogArtifacts(false)} />}
-      {!!historyAgentId && <HistoryModal agentId={historyAgentId} onDismiss={() => { history.replace(baseUrl); setHistoryAgentId(undefined) }} />}
+      {!!historyAgentId && <HistoryModal agentId={historyAgentId} onDismiss={() => { navigate(baseUrl, {replace: true}); setHistoryAgentId(undefined) }} />}
       <Breadcrumbs items={logSource?.crumbs ?? []} title={logSource?.crumbTitle} />
       <Stack tokens={{ childrenGap: 0 }} style={{ backgroundColor: "#FFFFFF", paddingTop: 12 }}>
          <Stack horizontal >
diff --git a/Engine/Source/Programs/Horde/HordeDashboard/src/components/NewBuild.tsx b/Engine/Source/Programs/Horde/HordeDashboard/src/components/NewBuild.tsx
--- a/Engine/Source/Programs/Horde/HordeDashboard/src/components/NewBuild.tsx
+++ b/Engine/Source/Programs/Horde/HordeDashboard/src/components/NewBuild.tsx
@@ -1,6 +1,6 @@
 // Copyright Epic Games, Inc. All Rights Reserved.
 import { Checkbox, ComboBox, ContextualMenuItemType, DefaultButton, DirectionalHint, Dropdown, DropdownMenuItemType, IComboBoxOption, Icon, IconButton, IContextualMenuItem, IContextualMenuProps, IDropdownOption, Label, MessageBar, MessageBarType, Modal, Pivot, PivotItem, PrimaryButton, ScrollablePane, ScrollbarVisibility, Spinner, SpinnerSize, Stack, TagPicker, Text, TextField, TooltipHost, ValidationState } from '@fluentui/react';
-import { action, observable } from 'mobx';
+import { action, makeObservable, observable } from 'mobx';
 import { observer } from 'mobx-react-lite';
 import moment from 'moment';
 import React, { useState } from 'react';
@@ -29,6 +29,8 @@ class BuildParameters {
 
    constructor(template: GetTemplateRefResponse, jobDetails?: JobDetailsV2, existing?: BuildParameters) {
 
+      makeObservable(this);
+
       BuildParameters.paramKey++;
 
       this.template = template;
diff --git a/Engine/Source/Programs/Horde/HordeDashboard/src/components/NotificationDropdown.tsx b/Engine/Source/Programs/Horde/HordeDashboard/src/components/NotificationDropdown.tsx
--- a/Engine/Source/Programs/Horde/HordeDashboard/src/components/NotificationDropdown.tsx
+++ b/Engine/Source/Programs/Horde/HordeDashboard/src/components/NotificationDropdown.tsx
@@ -6,7 +6,7 @@ import React, { useEffect } from "react";
 import backend from "../backend";
 import { CreateSubscriptionRequest, GetNotificationResponse, GetSubscriptionResponse, JobState, JobStepState, LabelCompleteEventRecord, LabelState, StepCompleteEventRecord, SubscriptonNotificationType, UpdateNotificationsRequest } from "../backend/Api";
 import { hordeClasses } from "../styles/Styles";
-import { action, observable } from "mobx";
+import { action, makeObservable, observable } from "mobx";
 import { JobDetailsV2 } from './jobDetailsV2/JobDetailsViewCommon';
 import { useQuery } from "./JobDetailCommon";
 
@@ -15,6 +15,10 @@ type NotificationOutcome = "Warnings" | "Success" | "Failure";
 
 class SubscriptionHandler  {
 
+   constructor() {
+      makeObservable(this);
+   }
+
    @observable
    private updated = 0
 
diff --git a/Engine/Source/Programs/Horde/HordeDashboard/src/components/Preflight.tsx b/Engine/Source/Programs/Horde/HordeDashboard/src/components/Preflight.tsx
--- a/Engine/Source/Programs/Horde/HordeDashboard/src/components/Preflight.tsx
+++ b/Engine/Source/Programs/Horde/HordeDashboard/src/components/Preflight.tsx
@@ -1,7 +1,7 @@
 // Copyright Epic Games, Inc. All Rights Reserved.
 
 import React, { useState } from 'react';
-import { useHistory } from 'react-router-dom';
+import { useNavigate, Navigate } from 'react-router-dom';
 import backend, { useBackend } from '../backend';
 import { useQuery } from './JobDetailCommon';
 
@@ -9,90 +9,85 @@ import { useQuery } from './JobDetailCommon';
 export const PreflightRedirector: React.FC = () => {
 
 
-    const [state, setState] = useState({ preflightQueried: false })
+   const [state, setState] = useState({ preflightQueried: false })
 
-    const history = useHistory();
-    const query = useQuery();
-    const { projectStore } = useBackend();
+   const navigate = useNavigate();
+   const query = useQuery();
+   const { projectStore } = useBackend();
 
-    const streamName = !query.get("stream") ? "" : query.get("stream")!;
-    const change = !query.get("change") ? "" : query.get("change")!;
+   const streamName = !query.get("stream") ? "" : query.get("stream")!;
+   const change = !query.get("change") ? "" : query.get("change")!;
 
-    // whether to autosubmit
-    const autosubmit = !query.get("submit") ? "" : query.get("submit")!;
+   // whether to autosubmit
+   const autosubmit = !query.get("submit") ? "" : query.get("submit")!;
 
-    if (!change) {
-        console.error("No preflight change specified");
-        history.replace("/");
-        return null;
-    }
+   if (!change) {
+      console.error("No preflight change specified");
+      return <Navigate to="/" replace={true} />;
+   }
 
-    const cl = parseInt(change);
+   const cl = parseInt(change);
 
-    if (isNaN(cl)) {
-        console.error(`Bad change in preflight ${change}`);
-        history.replace("/");
-        return null;
-    }
+   if (isNaN(cl)) {
+      console.error(`Bad change in preflight ${change}`);
+      return <Navigate to="/" replace={true} />;
+   }
 
-    if (!streamName) {
-        console.error("No stream in query");
-        history.replace("/");
-        return null;
-    }
+   if (!streamName) {
+      console.error("No stream in query");
+      return <Navigate to="/" replace={true} />;
+   }
 
-    let stream = projectStore.streamByFullname(streamName);
+   let stream = projectStore.streamByFullname(streamName);
 
 
-    if (!stream) {
-        console.error(`Unable to resolve stream with name ${streamName}`);
-        history.replace("/");
-        return null;
-    }
+   if (!stream) {
+      console.error(`Unable to resolve stream with name ${streamName}`);
+      return <Navigate to="/" replace={true} />;      
+   }
 
-    const project = stream?.project;
+   const project = stream?.project;
 
-    if (!stream || !project) {
-        console.error("Bad stream or project id in StreamView");
-        history.replace("/");
-        return null;
-    }
+   if (!stream || !project) {
+      console.error("Bad stream or project id in StreamView");
+      return <Navigate to="/" replace={true} />;      
+   }
 
-    if (!state.preflightQueried) {
+   if (!state.preflightQueried) {
 
-        backend.getJobs({ filter: "id", count: 1, preflightChange: cl }).then(result => {
+      backend.getJobs({ filter: "id", count: 1, preflightChange: cl }).then(result => {
 
-            if (result && result.length === 1) {
+         if (result && result.length === 1) {
 
-                let url = `/job/${result[0].id}?newbuild=true&allowtemplatechange=true&shelvedchange=${change}&p4v=true`;
+            let url = `/job/${result[0].id}?newbuild=true&allowtemplatechange=true&shelvedchange=${change}&p4v=true`;
 
-                if (autosubmit === "true") {
-                    url += "&autosubmit=true";
-                }
-
-                history.replace(url);
-                return;
+            if (autosubmit === "true") {
+               url += "&autosubmit=true";
             }
 
-            let url = `/stream/${stream!.id}?tab=summary&newbuild=true&shelvedchange=${change}&p4v=true`;
+            navigate(url, { replace: true });
+            return;
+         }
 
-            if (autosubmit === "true") {
-                url += "&autosubmit=true";
-            }
+         let url = `/stream/${stream!.id}?tab=summary&newbuild=true&shelvedchange=${change}&p4v=true`;
+
+         if (autosubmit === "true") {
+            url += "&autosubmit=true";
+         }
 
-            history.replace(url);
+         navigate(url, { replace: true });
 
 
-        }).catch(reason => {
-            console.error(`Error getting job for preflight: `, reason);
-            history.replace("/");
-        })
+      }).catch(reason => {
+         console.error(`Error getting job for preflight: `, reason);
+         navigate("/", { replace: true });
+      })
 
-        setState({ preflightQueried: true })
+      setState({ preflightQueried: true })
 
-        return null;
+      return null;
 
-    }
+   }
 
-    return null;
+   return null;
 }
diff --git a/Engine/Source/Programs/Horde/HordeDashboard/src/components/ProjectHome.tsx b/Engine/Source/Programs/Horde/HordeDashboard/src/components/ProjectHome.tsx
--- a/Engine/Source/Programs/Horde/HordeDashboard/src/components/ProjectHome.tsx
+++ b/Engine/Source/Programs/Horde/HordeDashboard/src/components/ProjectHome.tsx
@@ -7,7 +7,7 @@ import { Stack, SearchBox, Text, FocusZoneDirection, ScrollbarVisibility, FocusZ
 import { detailClasses, hordeClasses } from '../styles/Styles';
 import { Breadcrumbs, BreadcrumbItem } from './Breadcrumbs';
 import { TopNav } from './TopNav';
-import { action, observable } from 'mobx';
+import { action, makeObservable, observable } from 'mobx';
 import backend, { useBackend } from '../backend';
 import { ProjectData, StreamData } from '../backend/Api';
 import dashboard from '../backend/Dashboard';
@@ -20,6 +20,9 @@ type FilteredCategory = {
 };
 
 class LocalState {
+   constructor() {
+      makeObservable(this);
+   }
     filterString = "";
     activeProject: ProjectData | undefined = undefined;
     @observable.shallow filteredCategories: FilteredCategory[] = [];
diff --git a/Engine/Source/Programs/Horde/HordeDashboard/src/components/StreamView.tsx b/Engine/Source/Programs/Horde/HordeDashboard/src/components/StreamView.tsx
--- a/Engine/Source/Programs/Horde/HordeDashboard/src/components/StreamView.tsx
+++ b/Engine/Source/Programs/Horde/HordeDashboard/src/components/StreamView.tsx
@@ -3,7 +3,7 @@
 import { DefaultButton, IContextualMenuProps, mergeStyleSets, Pivot, PivotItem, PrimaryButton, Stack, Text, TextField } from '@fluentui/react';
 import { observer } from 'mobx-react-lite';
 import React, { useEffect, useState } from 'react';
-import { useHistory, useLocation, useParams } from 'react-router-dom';
+import { Navigate, useNavigate, useLocation, useParams } from 'react-router-dom';
 import { useBackend } from '../backend';
 import dashboard from '../backend/Dashboard';
 import { JobFilterSimple } from '../base/utilities/filter';
@@ -59,7 +59,7 @@ const StreamViewInner: React.FC = observer(() => {
    const windowSize = useWindowSize();
 
    const { streamId } = useParams<{ streamId: string }>();
-   const history = useHistory();
+   const navigate = useNavigate();
    const location = useLocation();
    const query = useQuery();
 
@@ -85,8 +85,7 @@ const StreamViewInner: React.FC = observer(() => {
 
    if (!stream || !project) {
       console.error("Bad stream or project id in StreamView");
-      history.replace("/");
-      return null;
+      return <Navigate to="/" replace={true} />
    }
 
    let queryTab = query.get("tab") ?? undefined;
@@ -94,14 +93,12 @@ const StreamViewInner: React.FC = observer(() => {
    if (!queryTab || (queryTab.toLowerCase() !== "summary" && queryTab.toLowerCase() !== "all" && !stream.tabs.find(t => t.title === queryTab))) {
 
       if (!!stream.tabs.find(t => t.title.toLowerCase() === "incremental")) {
-         history.replace(`/stream/${streamId}?tab=Incremental`);
+         return <Navigate to={`/stream/${streamId}?tab=Incremental`} replace={true} />
       } else if (stream.tabs.length) {
-         history.replace(`/stream/${streamId}?tab=${stream.tabs[0].title}`);
+         return <Navigate to={`/stream/${streamId}?tab=${stream.tabs[0].title}`} replace={true} />
       } else {
-         history.replace(`/stream/${streamId}?tab=all`);
+         return <Navigate to={`/stream/${streamId}?tab=all`} replace={true} />
       }
-
-      return null;
    }
 
    const isSummary = queryTab === 'summary';
@@ -222,10 +219,10 @@ const StreamViewInner: React.FC = observer(() => {
          <NewBuild streamId={streamId!} jobKey={newBuildTab!} show={shown} onClose={(newJobId) => {
             setShown(false);
             if (newJobId) {
-               history.push(`/job/${newJobId}`);
+               navigate(`/job/${newJobId}`);
             } else {
                if (query.get("newbuild")) {
-                  history.replace(`/stream/${streamId}?tab=${queryTab}`);
+                  navigate(`/stream/${streamId}?tab=${queryTab}`, { replace: true });
                }
             }
          }} />
@@ -243,7 +240,7 @@ const StreamViewInner: React.FC = observer(() => {
                            onLinkClick={(item) => {
                               if (item!.props.itemKey !== queryTab) {
                                  // this still needs to be here, even though specific button href in item
-                                 history.push(`/stream/${streamId}?tab=${encodeURIComponent(item!.props.itemKey!)}`);
+                                 navigate(`/stream/${streamId}?tab=${encodeURIComponent(item!.props.itemKey!)}`);
                               }
                            }}>
                            {pivotItems}
diff --git a/Engine/Source/Programs/Horde/HordeDashboard/src/components/SubmitQueue.tsx b/Engine/Source/Programs/Horde/HordeDashboard/src/components/SubmitQueue.tsx
--- a/Engine/Source/Programs/Horde/HordeDashboard/src/components/SubmitQueue.tsx
+++ b/Engine/Source/Programs/Horde/HordeDashboard/src/components/SubmitQueue.tsx
@@ -2,7 +2,7 @@
 
 import { CollapseAllVisibility, ConstrainMode, DetailsHeader, DetailsList, DetailsListLayoutMode, DetailsRow, Text, FocusZone, FocusZoneDirection, IColumn, IDetailsListProps, ScrollablePane, ScrollbarVisibility, SelectionMode, Stack, Sticky, StickyPositionType, Modal, IconButton, PrimaryButton, DefaultButton, TextField } from '@fluentui/react';
 import React from 'react';
-import { useParams, useHistory } from 'react-router-dom';
+import { useParams } from 'react-router-dom';
 import { detailClasses } from '../styles/Styles';
 
 export const SubmitQueue: React.FC = () => {
@@ -87,8 +87,7 @@ export const SubmitQueue: React.FC = () => {
 };
 
 export const SubmitQueueChangeDialog: React.FC<{ show: boolean; streamId: string; onClose: () => void }> = ({ show, streamId, onClose }) => {
-
-	const history = useHistory();
+	
 	const headerText = "Add Change";
 
 	const close = () => {
diff --git a/Engine/Source/Programs/Horde/HordeDashboard/src/components/TopNav.tsx b/Engine/Source/Programs/Horde/HordeDashboard/src/components/TopNav.tsx
--- a/Engine/Source/Programs/Horde/HordeDashboard/src/components/TopNav.tsx
+++ b/Engine/Source/Programs/Horde/HordeDashboard/src/components/TopNav.tsx
@@ -3,7 +3,7 @@
 import { CommandBar, CommandBarButton, ContextualMenu, ContextualMenuItem, ContextualMenuItemType, DefaultButton, Dialog, DialogFooter, DialogType, IButton, IButtonProps, ICommandBarItemProps, IContextualMenuItem, IContextualMenuItemProps, IContextualMenuItemStyles, IContextualMenuStyles, Persona, PersonaSize, PrimaryButton, Stack, Text } from '@fluentui/react';
 import { observer } from 'mobx-react-lite';
 import React, { useRef, useState } from 'react';
-import { Link, useHistory } from 'react-router-dom';
+import { Link, useNavigate } from 'react-router-dom';
 import backend, { useBackend } from '../backend';
 import { ProjectData } from "../backend/Api";
 import dashboard from '../backend/Dashboard';
@@ -45,7 +45,7 @@ type ButtonData = {
 // Top level project button component
 const ProjectButton: React.FunctionComponent<IButtonProps> = (props) => {
 
-   const history = useHistory();
+   const navigate = useNavigate();
    const { project } = props.data as ButtonData;
 
    const buttonRef = React.createRef<IButton>();
@@ -67,7 +67,7 @@ const ProjectButton: React.FunctionComponent<IButtonProps> = (props) => {
                if (ev?.metaKey || ev?.ctrlKey) {
                   window.open(`/project/${project?.id}`);
                } else {
-                  history.push(`/project/${project?.id}`);
+                  navigate(`/project/${project?.id}`);
                }
 
             }}
@@ -110,7 +110,7 @@ const ProjectMenuItem: React.FunctionComponent<IContextualMenuItemProps> = props
    return <Link to={item.link} onClick={(ev: any) => { ev.stopPropagation(); return true; }}> <ContextualMenuItem onClick={(ev: any) => { ev.preventDefault(); return true; }} {...props} /></Link>;
 };
 
-const generateProjectMenu = (history: any, store: ProjectStore) => {
+const generateProjectMenu = (store: ProjectStore) => {
 
    const projects = store.projects;
 
@@ -231,7 +231,7 @@ const generateProjectMenu = (history: any, store: ProjectStore) => {
 // Top level admin button component
 const AdminButton: React.FunctionComponent<IButtonProps> = (props) => {
 
-   const history = useHistory();
+   const navigate = useNavigate();
 
    const buttonRef = React.createRef<IButton>();
    return (
@@ -257,7 +257,7 @@ const AdminButton: React.FunctionComponent<IButtonProps> = (props) => {
                if (ev?.metaKey || ev?.ctrlKey) {
                   window.open(`/agents`);
                } else {
-                  history.push(`/agents`);
+                  navigate(`/agents`);
                }
             }}
             {...props}
@@ -338,7 +338,7 @@ export const TopNav: React.FC<{ suppressServer?: boolean }> = observer(({ suppre
 
    const [showMenu, setShowMenu] = useState(false);
    const [showVersion, setShowVersion] = useState(false);
-   const history = useHistory();
+   const navigate = useNavigate();
    const { projectStore } = useBackend();
 
    const divRef = useRef(null);
@@ -346,7 +346,7 @@ export const TopNav: React.FC<{ suppressServer?: boolean }> = observer(({ suppre
    // subscribe
    if (dashboard.updated) { }
 
-   const generateAdminMenu = (history: any) => {
+   const generateAdminMenu = () => {
 
 
       const cbProps: ICommandBarItemProps[] = [];
@@ -489,7 +489,7 @@ export const TopNav: React.FC<{ suppressServer?: boolean }> = observer(({ suppre
                   key: 'settingsItem',
                   text: 'Preferences',
                   onClick: () => {
-                     history.push("/dashboard");
+                     navigate("/dashboard");
                   }
                }],
             bottomDivider: true
@@ -542,7 +542,7 @@ export const TopNav: React.FC<{ suppressServer?: boolean }> = observer(({ suppre
                   <CommandBar styles={{ root: { paddingTop: 17, backgroundColor: modeColors.header } }}
                      buttonAs={ProjectButton}
                      onReduceData={() => undefined}
-                     items={generateProjectMenu(history, projectStore)}
+                     items={generateProjectMenu(projectStore)}
                   />
                </Stack>
 
@@ -552,7 +552,7 @@ export const TopNav: React.FC<{ suppressServer?: boolean }> = observer(({ suppre
                      {showServer && <CommandBar styles={{ root: { paddingTop: 17, paddingLeft: 0, paddingRight: 32, backgroundColor: modeColors.header } }}
                         buttonAs={AdminButton}
                         onReduceData={() => undefined}
-                        items={generateAdminMenu(history)}
+                        items={generateAdminMenu()}
                      />}
                   </Stack>
 
@@ -561,7 +561,7 @@ export const TopNav: React.FC<{ suppressServer?: boolean }> = observer(({ suppre
                      <div ref={divRef}>
                         <a href="/" onClick={(ev) => { ev.preventDefault(); setShowMenu(true); }}>
                            <Persona styles={{ root: { selectors: { ".ms-Persona-initials": { fontWeight: "unset", fontFamily: "Horde Open Sans SemiBold" } } } }} imageShouldFadeIn={false} imageInitials={initials} imageUrl={dashboard.userImage32} size={PersonaSize.size32}
-                              onClick={() => { history.push("/index") }} />
+                              onClick={() => { navigate("/index") }} />
                         </a>
                      </div>
                      <ContextualMenu
diff --git a/Engine/Source/Programs/Horde/HordeDashboard/src/components/UserHome.tsx b/Engine/Source/Programs/Horde/HordeDashboard/src/components/UserHome.tsx
--- a/Engine/Source/Programs/Horde/HordeDashboard/src/components/UserHome.tsx
+++ b/Engine/Source/Programs/Horde/HordeDashboard/src/components/UserHome.tsx
@@ -1,7 +1,7 @@
 // Copyright Epic Games, Inc. All Rights Reserved.
 
 import { CollapseAllVisibility, DefaultButton, DetailsHeader, DetailsList, DetailsListLayoutMode, DetailsRow, FocusZone, FocusZoneDirection, FontIcon, IColumn, IconButton, IDetailsListProps, mergeStyleSets, PrimaryButton, ScrollablePane, ScrollbarVisibility, SelectionMode, Spinner, SpinnerSize, Stack, Text } from '@fluentui/react';
-import { action, observable } from 'mobx';
+import { action, makeObservable, observable } from 'mobx';
 import { observer } from 'mobx-react-lite';
 import moment from "moment-timezone";
 import { default as React, useEffect, useState } from 'react';
@@ -118,6 +118,10 @@ const ProjectsPanel: React.FC = observer(() => {
 
 class HealthHandler {
 
+   constructor() {
+      makeObservable(this);      
+   }
+
    startPolling() {
 
       if (this.polling) {
@@ -345,6 +349,10 @@ const HealthPanel: React.FC = observer(() => {
 const jobsRefreshTime = 3000;
 class UserJobsHandler {
 
+   constructor() {
+      makeObservable(this);
+   }
+
    set(jobIds: string[]): boolean {
 
       if (jobIds.length === this.jobIds.length) {
diff --git a/Engine/Source/Programs/Horde/HordeDashboard/src/components/UtilizationReportView.tsx b/Engine/Source/Programs/Horde/HordeDashboard/src/components/UtilizationReportView.tsx
--- a/Engine/Source/Programs/Horde/HordeDashboard/src/components/UtilizationReportView.tsx
+++ b/Engine/Source/Programs/Horde/HordeDashboard/src/components/UtilizationReportView.tsx
@@ -5,7 +5,7 @@ import { IVerticalStackedChartProps, IVSChartDataPoint, VerticalStackedBarChart
 import { observer } from 'mobx-react-lite';
 import moment from 'moment-timezone';
 import React, { useEffect, useState } from 'react';
-import { useHistory } from 'react-router-dom';
+import { useNavigate } from 'react-router-dom';
 import backend from '../backend';
 import { GetUtilizationTelemetryResponse } from '../backend/Api';
 import { useWindowSize } from '../base/utilities/hooks';
@@ -158,7 +158,7 @@ let currentSliderValue: number = 1;
 export const UtilizationReportView: React.FC = observer(() => {
 	const windowSize = useWindowSize();
 	const query = useQuery();
-	const history = useHistory();
+	const navigate = useNavigate();
 
 	const [optionsState, setOptionsState] = useState<ChartOptionsState>({
 		chartData: null,
@@ -593,7 +593,7 @@ export const UtilizationReportView: React.FC = observer(() => {
 		setOptionsState(newState);
 
 		let dateParamString = `${newState.selectedDate.getFullYear()}-${newState.selectedDate.getMonth() + 1}-${newState.selectedDate.getDate()}`;
-		history.push(`/reports/utilization?date=${dateParamString}&range=${newState.selectedRange}&interval=${newState.selectedFormat}&pools=${newState.selectedPoolKeys.join(';')}`);
+		navigate(`/reports/utilization?date=${dateParamString}&range=${newState.selectedRange}&interval=${newState.selectedFormat}&pools=${newState.selectedPoolKeys.join(';')}`);
 	};
 
 	const beginUpdateState = (options: UpdateStateOptions): void => {
diff --git a/Engine/Source/Programs/Horde/HordeDashboard/src/components/VersionModal.tsx b/Engine/Source/Programs/Horde/HordeDashboard/src/components/VersionModal.tsx
--- a/Engine/Source/Programs/Horde/HordeDashboard/src/components/VersionModal.tsx
+++ b/Engine/Source/Programs/Horde/HordeDashboard/src/components/VersionModal.tsx
@@ -49,7 +49,14 @@ export const VersionModal: React.FC<{ show: boolean, onClose: () => void }> = ({
       { key: 'column2', name: 'Value', fieldName: 'value', minWidth: 100, maxWidth: 200 },
    ];
 
-   const dashboardVersion: string | undefined = process?.env?.REACT_APP_VERSION_INFO;
+   let dashboardVersion: string | undefined;
+
+   try {
+      dashboardVersion = process?.env?.REACT_APP_VERSION_INFO;
+   } catch {
+      
+   }
+   
    if (dashboardVersion) {
       versionItems.push({ name: "Dashboard", value: dashboardVersion });
    }
diff --git a/Engine/Source/Programs/Horde/HordeDashboard/src/components/jobDetailsV2/JobDetailArtifactsV2.tsx b/Engine/Source/Programs/Horde/HordeDashboard/src/components/jobDetailsV2/JobDetailArtifactsV2.tsx
--- a/Engine/Source/Programs/Horde/HordeDashboard/src/components/jobDetailsV2/JobDetailArtifactsV2.tsx
+++ b/Engine/Source/Programs/Horde/HordeDashboard/src/components/jobDetailsV2/JobDetailArtifactsV2.tsx
@@ -5,7 +5,7 @@ import React, { useEffect } from 'react';
 import { ArtifactData, GetArtifactResponse, GetArtifactZipRequest } from '../../backend/Api';
 import { hordeClasses } from '../../styles/Styles';
 import { observer } from 'mobx-react-lite';
-import { observable, action } from 'mobx';
+import { observable, action, makeObservable } from 'mobx';
 import backend from '../../backend';
 import { JobDataView, JobDetailsV2 } from './JobDetailsViewCommon';
 import { ISideRailLink } from '../../base/components/SideRail';
@@ -39,6 +39,11 @@ const classNames = mergeStyleSets({
 });
 
 class ArtifactState {
+
+   constructor() {
+      makeObservable(this);
+   }
+
    @observable selectedItemCount = 0;
    @observable isDownloading = false;
    jobDetails: JobDetailsV2 | null = null;
@@ -188,7 +193,7 @@ class ArtifactsDataView extends JobDataView {
    }
 
    detailsUpdated() {
-      
+
       if (!this.details?.jobData) {
          return;
       }
diff --git a/Engine/Source/Programs/Horde/HordeDashboard/src/components/jobDetailsV2/JobDetailLabels.tsx b/Engine/Source/Programs/Horde/HordeDashboard/src/components/jobDetailsV2/JobDetailLabels.tsx
--- a/Engine/Source/Programs/Horde/HordeDashboard/src/components/jobDetailsV2/JobDetailLabels.tsx
+++ b/Engine/Source/Programs/Horde/HordeDashboard/src/components/jobDetailsV2/JobDetailLabels.tsx
@@ -1,6 +1,6 @@
 import { ConstrainMode, DefaultButton, DetailsList, DetailsListLayoutMode, IColumn, SelectionMode, Stack, Text } from "@fluentui/react";
 import { observer } from "mobx-react-lite";
-import { useHistory, useLocation } from "react-router-dom";
+import { useNavigate, useLocation } from "react-router-dom";
 import { LabelOutcome, LabelState } from "../../backend/Api";
 import { JobLabel } from "../../backend/JobDetails";
 import { getLabelColor } from "../../styles/colors";
@@ -12,7 +12,7 @@ import { JobDataView, JobDetailsV2 } from "./JobDetailsViewCommon";
 export const LabelsPanelV2: React.FC<{ jobDetails: JobDetailsV2, dataView: JobDataView }> = observer(({ jobDetails, dataView }) => {
 
    const query = useQuery();
-   const history = useHistory();
+   const navigate = useNavigate();
    const location = useLocation();
 
    dataView.subscribe();
@@ -159,14 +159,14 @@ export const LabelsPanelV2: React.FC<{ jobDetails: JobDetailsV2, dataView: JobDa
                      <DefaultButton
                         onClick={() => {
                            if (qlabel?.category === label.category && qlabel?.name === label.name) {
-                              history.push(location.pathname)
+                              navigate(location.pathname)
                            } else {
 
                               const idx = jobDetails.labelIndex(label.name, label.category);
                               if (idx >= 0) {
-                                 history.push(location.pathname + `?label=${idx}`)
+                                 navigate(location.pathname + `?label=${idx}`)
                               } else {
-                                 history.replace(location.pathname)
+                                 navigate(location.pathname, {replace: true})
                               }                              
                            }
 
diff --git a/Engine/Source/Programs/Horde/HordeDashboard/src/components/jobDetailsV2/JobDetailStepHistory.tsx b/Engine/Source/Programs/Horde/HordeDashboard/src/components/jobDetailsV2/JobDetailStepHistory.tsx
--- a/Engine/Source/Programs/Horde/HordeDashboard/src/components/jobDetailsV2/JobDetailStepHistory.tsx
+++ b/Engine/Source/Programs/Horde/HordeDashboard/src/components/jobDetailsV2/JobDetailStepHistory.tsx
@@ -4,7 +4,7 @@ import { DetailsList, DetailsListLayoutMode, DetailsRow, IColumn, IDetailsListPr
 import { observer } from "mobx-react-lite";
 import moment from "moment";
 import { useEffect } from "react";
-import { Link, useHistory, useLocation } from "react-router-dom";
+import { Link, useNavigate, useLocation } from "react-router-dom";
 import backend from "../../backend";
 import { GetJobStepRefResponse } from "../../backend/Api";
 import dashboard from "../../backend/Dashboard";
@@ -75,7 +75,7 @@ JobDetailsV2.registerDataView("StepHistoryDataView", (details: JobDetailsV2) =>
 
 export const StepHistoryPanel: React.FC<{ jobDetails: JobDetailsV2; stepId: string }> = observer(({ jobDetails, stepId }) => {
 
-   const history = useHistory();
+   const navigate = useNavigate();
    const location = useLocation();
 
    const dataView = jobDetails.getDataView<StepHistoryDataView>("StepHistoryDataView");
@@ -144,7 +144,7 @@ export const StepHistoryPanel: React.FC<{ jobDetails: JobDetailsV2; stepId: stri
 
          const url = `${location.pathname}${location.search}&agentId=${agentId}`;
 
-         return <a href={url} onClick={(ev) => { ev.preventDefault(); ev.stopPropagation(); history.replace(url); }}><Stack horizontal horizontalAlign={"end"} verticalFill={true} tokens={{ childrenGap: 0, padding: 0 }}><Text>{agentId}</Text></Stack></a>;
+         return <a href={url} onClick={(ev) => { ev.preventDefault(); ev.stopPropagation(); navigate(url, { replace: true } ); }}><Stack horizontal horizontalAlign={"end"} verticalFill={true} tokens={{ childrenGap: 0, padding: 0 }}><Text>{agentId}</Text></Stack></a>;
       }
 
       if (column.name === "Started") {
diff --git a/Engine/Source/Programs/Horde/HordeDashboard/src/components/jobDetailsV2/JobDetailStepTrends.tsx b/Engine/Source/Programs/Horde/HordeDashboard/src/components/jobDetailsV2/JobDetailStepTrends.tsx
--- a/Engine/Source/Programs/Horde/HordeDashboard/src/components/jobDetailsV2/JobDetailStepTrends.tsx
+++ b/Engine/Source/Programs/Horde/HordeDashboard/src/components/jobDetailsV2/JobDetailStepTrends.tsx
@@ -1,7 +1,7 @@
 // Copyright Epic Games, Inc. All Rights Reserved.
 
 import { ContextualMenu, Dialog, DialogType, IconButton, IIconProps, Spinner, SpinnerSize, Stack, Text, Toggle } from '@fluentui/react';
-import { IChartProps, ILineChartDataPoint } from '@fluentui/react-charting';
+import { IChartProps, ICustomizedCalloutData, ILineChartDataPoint, LineChart } from '@fluentui/react-charting';
 import { observer } from 'mobx-react-lite';
 import moment from 'moment';
 import React, { useCallback, useEffect, useState } from 'react';
@@ -9,9 +9,6 @@ import backend from '../../backend';
 import { GetBatchResponse, GetJobTimingResponse, GetLabelResponse, GetLabelStateResponse, GetLabelTimingInfoResponse, GetNodeResponse, GetStepResponse, GetStepTimingInfoResponse, GetTemplateRefResponse, JobData, JobQuery, JobTimingsQuery, StepData } from '../../backend/Api';
 import { JobLabel } from '../../backend/JobDetails';
 import TemplateCache from '../../backend/TemplateCache';
-import { HChartHoverCard } from '../../base/components/HordeFluentCharting/HHoverCard/HHoverCard';
-import { HLineChart } from '../../base/components/HordeFluentCharting/HLineChart/HLineChart';
-import { IHCustomizedCalloutData } from '../../base/components/HordeFluentCharting/HLineChart/HLineChart.types';
 import { ISideRailLink } from '../../base/components/SideRail';
 import { msecToElapsed } from '../../base/utilities/timeUtils';
 import { hordeClasses } from '../../styles/Styles';
@@ -409,7 +406,7 @@ function preventChartScroll(this: Document, e: WheelEvent) {
    e.returnValue = false
 }
 
-export const StepTrendsPanel: React.FC<{ jobDetails: JobDetailsV2}> = observer(({ jobDetails}) => {
+export const StepTrendsPanel: React.FC<{ jobDetails: JobDetailsV2 }> = observer(({ jobDetails }) => {
    const query = useQuery();
 
    const stepId = query.get("step") ? query.get("step")! : undefined;
@@ -422,15 +419,15 @@ export const StepTrendsPanel: React.FC<{ jobDetails: JobDetailsV2}> = observer((
          dataView?.clear();
       };
    }, [dataView]);
-   
+
    dataView.subscribe();
 
    if (jobDetails?.jobData) {
-      dataView.initialize(jobDetails.jobData?.preflightChange ? undefined : [sideRail]);       
+      dataView.initialize(jobDetails.jobData?.preflightChange ? undefined : [sideRail]);
    } else {
       jobDetails?.subscribe();
    }
-   
+
    if (jobDetails.jobData == null || !jobDetails?.viewsReady || jobDetails.jobData.preflightChange) {
       return null;
    }
@@ -1219,7 +1216,7 @@ const StepTrendsPanelInner: React.FC<{ jobDetails: JobDetailsV2; dataView: Trend
    let titleText = 'Trends';
    if (filteredChartState && filteredChartState?.chartProps) {
       if (filteredChartState.chartProps?.lineChartData?.length) {
-         placeholderDiv = <HLineChart
+         placeholderDiv = <LineChart
             data={filteredChartState.chartProps}
             hideLegend={false}
             legendProps={{ canSelectMultipleLegends: true, allowFocusOnLegends: false, enabledWrapLines: true }}
@@ -1228,21 +1225,16 @@ const StepTrendsPanelInner: React.FC<{ jobDetails: JobDetailsV2; dataView: Trend
             margins={{ top: 20, bottom: 35, left: filteredChartState.leftMargin, right: 45 }}
             yAxisTickCount={5}
             yMaxValue={filteredChartState.yMaxValue}
-            onRenderCalloutPerDataPoint={(props: (IHCustomizedCalloutData | undefined)[] | undefined) => {
-               if (!props || (!props[0] && !props[1])) {
-                  return null;
-               }
-
-               let hoverPoint1 = props[0];
-               let hoverPoint2 = props[1];
-
-               let dataPoint1 = hoverPoint1 ? hoverPoint1.values.find((point: { legend: string | number; y: string | number }) => point.y === hoverPoint1!.hoverYValue) : undefined;
-               let dataPoint2 = hoverPoint2 ? hoverPoint2.values.find((point: { legend: string | number; y: string | number }) => point.y === hoverPoint2!.hoverYValue) : undefined;
-
-               if (hoverPoint1 && dataPoint1) {
-                  return <HChartHoverCard props1={hoverPoint1} dataPoint1={dataPoint1} props2={hoverPoint2} dataPoint2={dataPoint2} />
-               }
-               return null;
+            onRenderCalloutPerDataPoint={(props?: ICustomizedCalloutData) => {
+               // This is currently disabled, due to package updates and incompatibility with 
+               // custom LineChart component hoisted out
+               // of react charting, which had some customizations for the callout
+               // namely comparison with next data point, and more importantly a way to 
+               // tell which y data point the user is hovering over... without that
+               // all matching x's are rendered and no way to customize (a failing of the fluent charting package)
+               // Trends data and rendering need to be redone, and already spent too much time trying to figure out
+               // how to make this work without the customization on the component, so here we are:
+               return <div />
             }}
          />
       }
diff --git a/Engine/Source/Programs/Horde/HordeDashboard/src/components/jobDetailsV2/JobDetailViewV2.tsx b/Engine/Source/Programs/Horde/HordeDashboard/src/components/jobDetailsV2/JobDetailViewV2.tsx
--- a/Engine/Source/Programs/Horde/HordeDashboard/src/components/jobDetailsV2/JobDetailViewV2.tsx
+++ b/Engine/Source/Programs/Horde/HordeDashboard/src/components/jobDetailsV2/JobDetailViewV2.tsx
@@ -1,7 +1,7 @@
 import { mergeStyleSets, Pivot, PivotItem, ScrollablePane, ScrollbarVisibility, Stack } from "@fluentui/react";
 import { observer } from "mobx-react-lite";
 import { useEffect, useState } from "react";
-import { useHistory, useParams } from "react-router-dom";
+import { useNavigate, useParams } from "react-router-dom";
 import { GetJobsTabResponse } from "../../backend/Api";
 import { useWindowSize } from "../../base/utilities/hooks";
 import { hordeClasses, modeColors } from "../../styles/Styles";
@@ -134,7 +134,7 @@ let crender = 0;
 
 const JobPanelPivot: React.FC<{ jobDetails: JobDetailsV2 }> = observer(({ jobDetails }) => {
    useQuery();
-   const history = useHistory();
+   const navigate = useNavigate();
    const [state, setState] = useState<{ iobserver?: IntersectionObserver }>({});
    const [, setRender] = useState(0);
 
@@ -153,7 +153,7 @@ const JobPanelPivot: React.FC<{ jobDetails: JobDetailsV2 }> = observer(({ jobDet
    // install intersection observer
    if (supported && rootArea && !state.iobserver) {
       if (window.location.hash) {
-         history.replace(window.location.pathname + window.location.search);
+         navigate(window.location.pathname + window.location.search, { replace: true });
       }
       setState({
          iobserver: new IntersectionObserver((entries: IntersectionObserverEntry[]) => {
diff --git a/Engine/Source/Programs/Horde/HordeDashboard/src/components/jobDetailsV2/JobDetailsViewCommon.tsx b/Engine/Source/Programs/Horde/HordeDashboard/src/components/jobDetailsV2/JobDetailsViewCommon.tsx
--- a/Engine/Source/Programs/Horde/HordeDashboard/src/components/jobDetailsV2/JobDetailsViewCommon.tsx
+++ b/Engine/Source/Programs/Horde/HordeDashboard/src/components/jobDetailsV2/JobDetailsViewCommon.tsx
@@ -1,9 +1,9 @@
 import { Dropdown, FontIcon, IDropdownOption, IPickerItemProps, ISuggestionItemProps, ITag, PrimaryButton, Stack, TagPicker, Text, ValidationState } from "@fluentui/react";
-import { action, observable } from "mobx";
+import { action, makeObservable, observable } from "mobx";
 import { observer } from "mobx-react-lite";
 import moment from "moment";
 import React, { useState } from "react";
-import { useHistory, useLocation } from "react-router-dom";
+import { useNavigate, useLocation } from "react-router-dom";
 import backend from "../../backend";
 import { BatchData, GetJobTimingResponse, GetLabelResponse, GetTemplateRefResponse, GroupData, JobData, JobState, LabelState, NodeData, ReportPlacement, StepData, StreamData } from "../../backend/Api";
 import { JobLabel } from "../../backend/JobDetails";
@@ -16,8 +16,9 @@ import { useQuery } from "../JobDetailCommon";
 import { LabelStatusIcon, StepStatusIcon } from "../StatusIcon";
 
 export abstract class JobDataView {
-
+   
    constructor(details: JobDetailsV2) {
+      makeObservable(this);
       this.details = details;
    }
 
@@ -121,8 +122,9 @@ export abstract class JobDataView {
 const defaultUpdateMS = 5000;
 export class JobDetailsV2 extends PollBase {
 
-   constructor(jobId: string) {
+   constructor(jobId: string) {      
       super(defaultUpdateMS)
+      makeObservable(this);
       this.jobId = jobId;
       this.filter = new JobDetailFilters(this);
       this.start();
@@ -768,6 +770,7 @@ type FilterPickerItem = ITag & {
 export class JobDetailFilters {
 
    constructor(details: JobDetailsV2) {
+      makeObservable(this);
       this.details = details;
    }
 
@@ -1008,7 +1011,7 @@ export const JobFilterBar: React.FC<{ jobDetails: JobDetailsV2 }> = observer(({
    const location = useLocation();
    const query = useQuery();
    const filterPicker = React.useRef(null);
-   const history = useHistory();
+   const navigate = useNavigate();
    const [state, setState] = useState<{}>({});
 
    const jobFilter = jobDetails.filter;
@@ -1170,7 +1173,7 @@ export const JobFilterBar: React.FC<{ jobDetails: JobDetailsV2 }> = observer(({
                      style={{ padding: 2, paddingLeft: 8, paddingRight: 8, fontFamily: "Horde Open Sans SemiBold", fontSize: 12, height: "unset" }}
                      text={props.item.name}
                      onClick={() => {
-                        history.push(location.pathname);
+                        navigate(location.pathname);
                         jobFilter.setSelected(undefined, undefined, undefined)
                      }} />
                </Stack>;
@@ -1248,7 +1251,7 @@ export const JobFilterBar: React.FC<{ jobDetails: JobDetailsV2 }> = observer(({
 
                const step = (item as FilterPickerItem).step;
                if (step) {
-                  history.push(location.pathname + `?step=${step.id}`);
+                  navigate(location.pathname + `?step=${step.id}`);
                   return null;
                }
 
@@ -1256,9 +1259,9 @@ export const JobFilterBar: React.FC<{ jobDetails: JobDetailsV2 }> = observer(({
                if (label) {
                   const idx = jobDetails.labelIndex(label.name, label.category);
                   if (idx >= 0) {
-                     history.push(location.pathname + `?label=${idx}`);
+                     navigate(location.pathname + `?label=${idx}`);
                   } else {
-                     history.replace(location.pathname);
+                     navigate(location.pathname, {replace: true});
                   }
                }
 
diff --git a/Engine/Source/Programs/Horde/HordeDashboard/src/components/jobDetailsV2/JobOperationsBar.tsx b/Engine/Source/Programs/Horde/HordeDashboard/src/components/jobDetailsV2/JobOperationsBar.tsx
--- a/Engine/Source/Programs/Horde/HordeDashboard/src/components/jobDetailsV2/JobOperationsBar.tsx
+++ b/Engine/Source/Programs/Horde/HordeDashboard/src/components/jobDetailsV2/JobOperationsBar.tsx
@@ -3,7 +3,7 @@
 import { CommandBar, CommandBarButton, ICommandBarItemProps, IContextualMenuItem, Stack } from '@fluentui/react';
 import { observer } from 'mobx-react-lite';
 import React, { useState } from 'react';
-import { Link, useHistory } from 'react-router-dom';
+import { Link, useNavigate } from 'react-router-dom';
 import { JobState, StepData } from '../../backend/Api';
 import dashboard from '../../backend/Dashboard';
 import { hordeClasses } from '../../styles/Styles';
@@ -25,7 +25,7 @@ enum ParameterState {
 export const JobOperations: React.FC<{ jobDetails: JobDetailsV2 }> = observer(({ jobDetails }) => {
 
    const query = useQuery();
-   const history = useHistory();
+   const navigate = useNavigate();
    const [abortShown, setAbortShown] = useState(false);
    const [editShown, setEditShown] = useState(false);
    const [parametersState, setParametersState] = useState(query.get("newbuild") ? ParameterState.Clone : ParameterState.Hidden);
@@ -124,10 +124,10 @@ export const JobOperations: React.FC<{ jobDetails: JobDetailsV2 }> = observer(({
          onClose={(newJobId) => {
             setParametersState(ParameterState.Hidden);
             if (newJobId) {
-               history.push(`/job/${newJobId}`);
+               navigate(`/job/${newJobId}`);
             } else {
                if (query.get("newbuild")) {
-                  history.replace(`/job/${jobId}`)
+                  navigate(`/job/${jobId}`, { replace: true });
                }
             }
          }} />
diff --git a/Engine/Source/Programs/Horde/HordeDashboard/src/components/jobDetailsV2/StepDetailView.tsx b/Engine/Source/Programs/Horde/HordeDashboard/src/components/jobDetailsV2/StepDetailView.tsx
--- a/Engine/Source/Programs/Horde/HordeDashboard/src/components/jobDetailsV2/StepDetailView.tsx
+++ b/Engine/Source/Programs/Horde/HordeDashboard/src/components/jobDetailsV2/StepDetailView.tsx
@@ -3,7 +3,7 @@
 import { Stack } from "@fluentui/react";
 import { observer } from "mobx-react-lite";
 import { useState } from "react";
-import { useHistory } from "react-router-dom";
+import { useNavigate } from "react-router-dom";
 import { hordeClasses } from "../../styles/Styles";
 import { HistoryModal } from "../HistoryModal";
 import { useQuery } from "../JobDetailCommon";
@@ -40,7 +40,7 @@ JobDetailsV2.registerDataView("StepDetailDataView", (details: JobDetailsV2) => n
 const StepDetailViewInner: React.FC<{ jobDetails: JobDetailsV2, stepId: string }> = observer(({ jobDetails, stepId }) => {
 
    let [historyAgentId, setHistoryAgentId] = useState<string | undefined>(undefined);
-   const history = useHistory();
+   const navigate = useNavigate();
    const query = useQuery();
 
    const dataView = jobDetails.getDataView<StepDetailDataView>("StepDetailDataView");
@@ -62,7 +62,7 @@ const StepDetailViewInner: React.FC<{ jobDetails: JobDetailsV2, stepId: string }
    const step = jobDetails.stepById(stepId);
 
    return <Stack>
-      <HistoryModal agentId={historyAgentId} onDismiss={() => { history.replace(`/job/${jobData.id}?step=${stepId}`); setHistoryAgentId(undefined); }} />
+      <HistoryModal agentId={historyAgentId} onDismiss={() => { navigate(`/job/${jobData.id}?step=${stepId}`, { replace: true }); setHistoryAgentId(undefined); }} />
 
       <Stack>
          <StepSummaryPanel jobDetails={jobDetails} stepId={stepId} />
