Horde: Updating dashboard dependencies to latest versions and cleanup

diff --git a/Engine/Source/Programs/Horde/HordeDashboard/src/components/jobDetailsV2/StepRetryModal.tsx b/Engine/Source/Programs/Horde/HordeDashboard/src/components/jobDetailsV2/StepRetryModal.tsx
--- a/Engine/Source/Programs/Horde/HordeDashboard/src/components/jobDetailsV2/StepRetryModal.tsx
+++ b/Engine/Source/Programs/Horde/HordeDashboard/src/components/jobDetailsV2/StepRetryModal.tsx
@@ -3,126 +3,126 @@ import { DefaultButton, IconButton, Modal, PrimaryButton, Stack, Text, TextField
 import React from 'react';
 import backend from '../../backend';
 import { CreateJobRequest } from '../../backend/Api';
-import { useHistory } from 'react-router-dom';
+import { useNavigate } from 'react-router-dom';
 import { hordeClasses } from '../../styles/Styles';
 import { JobDetailsV2 } from './JobDetailsViewCommon';
 
 export enum StepRetryType {
-    RunAgain,
-    TestFix
+   RunAgain,
+   TestFix
 }
 
 export const StepRetryModal: React.FC<{ stepId: string; jobDetails: JobDetailsV2; show: boolean; type: StepRetryType; onClose: () => void }> = ({ stepId, jobDetails, show, type, onClose }) => {
 
-   const history = useHistory();
-   
+   const navigate = useNavigate();
+
    const jobData = jobDetails.jobData;
    if (!jobData) {
       console.error("Attempting to display StepRetryModal without job data");
       return null;
    }
 
-    const node = jobDetails.nodeByStepId(stepId);    
+   const node = jobDetails.nodeByStepId(stepId);
 
-    const headerText = type === StepRetryType.RunAgain ? `Run "${node?.name}" Again?` : "Preflight Step";
+   const headerText = type === StepRetryType.RunAgain ? `Run "${node?.name}" Again?` : "Preflight Step";
 
-    let fixCL = 0;
+   let fixCL = 0;
 
-    const close = () => {
-        onClose();
-    };
+   const close = () => {
+      onClose();
+   };
 
-    const onRunAgain = async () => {
+   const onRunAgain = async () => {
 
-        const batch = jobDetails.batchByStepId(stepId);
-        if (!batch) {
-            console.error("Unable to get batch for step");
-            return;
-        }
+      const batch = jobDetails.batchByStepId(stepId);
+      if (!batch) {
+         console.error("Unable to get batch for step");
+         return;
+      }
+
+      await backend.updateJobStep(jobData.id, batch.id, stepId, {
+         retry: true
+      }).then((response) => {
+
+         if (response.stepId) {
+            navigate(`/job/${jobData.id}?step=${response.stepId!}`)
+         }
+
+      }).catch((reason) => {
+         // @todo: error ui
+         console.error(reason);
+      }).finally(() => {
+         close();
+      });
+
+   };
+
+   const onTryFix = async () => {
 
-        await backend.updateJobStep(jobData.id, batch.id, stepId, {
-            retry: true
-        }).then((response) => {
+      if (!fixCL || isNaN(fixCL)) {
+         console.error("Invalid Fix CL");
+         return;
+      }
 
-            if (response.stepId) {
-                history.push(`/job/${jobData.id}?step=${response.stepId!}`)
-            } 
-            
-        }).catch((reason) => {
-            // @todo: error ui
-            console.error(reason);
-        }).finally(() => {
-            close();
-        });
+      const job = jobData;
+      const node = jobDetails.nodeByStepId(stepId);
 
-    };
+      const args = [];
+      args.push(`-Target=Setup Build`);
+      args.push(`-Target=${node!.name}`);
 
-    const onTryFix = async () => {
+      job.arguments?.forEach(arg => {
 
-        if (!fixCL || isNaN(fixCL)) {
-            console.error("Invalid Fix CL");
+         if (arg.toLowerCase().indexOf("-target=") !== -1) {
             return;
-        }
-
-        const job = jobData;
-        const node = jobDetails.nodeByStepId(stepId);
-
-        const args = [];
-        args.push(`-Target=Setup Build`);
-        args.push(`-Target=${node!.name}`);        
-
-        job.arguments?.forEach(arg => {
-
-            if (arg.toLowerCase().indexOf("-target=") !== -1) {
-                return;
-            }
-
-            args.push(arg);
-        });
-
-        const data: CreateJobRequest = {
-            streamId: job.streamId,
-            templateId: job.templateId!,
-            arguments: args,
-            change: job.change,
-            preflightChange: fixCL
-        };
-
-        console.log("Submitting job");
-        console.log(data);
-
-        backend.createJob(data).then(data => {
-            history.push(`/job/${data.id}`);
-        });
-    };
-
-    const height = type === StepRetryType.TestFix ? 200 : 140;
-
-    return <Modal className={hordeClasses.modal} isOpen={show} styles={{ main: { padding: 8, width: 540, height: height, minHeight: height } }} onDismiss={() => { close() }}>
-        <Stack horizontal styles={{ root: { padding: 8 } }}>
-            <Stack.Item grow={2}>
-                <Text variant="mediumPlus">{headerText}</Text>
-            </Stack.Item>
-            <Stack.Item grow={0}>
-                <IconButton
-                    iconProps={{ iconName: 'Cancel' }}
-                    ariaLabel="Close popup modal"
-                    onClick={() => { close(); }}
-                />
-            </Stack.Item>
-        </Stack>
-
-        <Stack styles={{ root: { paddingLeft: 8, width: 500 } }}>
-            {type === StepRetryType.TestFix && <TextField label="Shelved Change" onChange={(ev, newValue) => { ev.preventDefault(); fixCL = parseInt(newValue ?? ""); }} />}
-        </Stack>
-
-        <Stack styles={{ root: { padding: 8 } }}>
-            <Stack horizontal tokens={{ childrenGap: 16 }} styles={{ root: { paddingTop: 12, paddingLeft: 8, paddingBottom: 8 } }}>
-                <Stack grow/>
-                <PrimaryButton text="Run" disabled={false} onClick={() => { type === StepRetryType.TestFix ? onTryFix() : onRunAgain(); }} />
-                <DefaultButton text="Cancel" disabled={false} onClick={() => { close(); }} />
-            </Stack>
-        </Stack>
-    </Modal>;
+         }
+
+         args.push(arg);
+      });
+
+      const data: CreateJobRequest = {
+         streamId: job.streamId,
+         templateId: job.templateId!,
+         arguments: args,
+         change: job.change,
+         preflightChange: fixCL
+      };
+
+      console.log("Submitting job");
+      console.log(data);
+
+      backend.createJob(data).then(data => {
+         navigate(`/job/${data.id}`);
+      });
+   };
+
+   const height = type === StepRetryType.TestFix ? 200 : 140;
+
+   return <Modal className={hordeClasses.modal} isOpen={show} styles={{ main: { padding: 8, width: 540, height: height, minHeight: height } }} onDismiss={() => { close() }}>
+      <Stack horizontal styles={{ root: { padding: 8 } }}>
+         <Stack.Item grow={2}>
+            <Text variant="mediumPlus">{headerText}</Text>
+         </Stack.Item>
+         <Stack.Item grow={0}>
+            <IconButton
+               iconProps={{ iconName: 'Cancel' }}
+               ariaLabel="Close popup modal"
+               onClick={() => { close(); }}
+            />
+         </Stack.Item>
+      </Stack>
+
+      <Stack styles={{ root: { paddingLeft: 8, width: 500 } }}>
+         {type === StepRetryType.TestFix && <TextField label="Shelved Change" onChange={(ev, newValue) => { ev.preventDefault(); fixCL = parseInt(newValue ?? ""); }} />}
+      </Stack>
+
+      <Stack styles={{ root: { padding: 8 } }}>
+         <Stack horizontal tokens={{ childrenGap: 16 }} styles={{ root: { paddingTop: 12, paddingLeft: 8, paddingBottom: 8 } }}>
+            <Stack grow />
+            <PrimaryButton text="Run" disabled={false} onClick={() => { type === StepRetryType.TestFix ? onTryFix() : onRunAgain(); }} />
+            <DefaultButton text="Cancel" disabled={false} onClick={() => { close(); }} />
+         </Stack>
+      </Stack>
+   </Modal>;
 
 };
\ No newline at end of file
diff --git a/Engine/Source/Programs/Horde/HordeDashboard/src/index.tsx b/Engine/Source/Programs/Horde/HordeDashboard/src/index.tsx
--- a/Engine/Source/Programs/Horde/HordeDashboard/src/index.tsx
+++ b/Engine/Source/Programs/Horde/HordeDashboard/src/index.tsx
@@ -3,13 +3,21 @@
 // this import must come first as it loads theme!
 import './styles/Styles';
 
-import React from 'react';
-import ReactDOM from 'react-dom';
+import { createRoot } from 'react-dom/client';
 import App from './App';
 import './index.css';
 import * as serviceWorker from './serviceWorker';
+import { configure } from 'mobx';
 
-ReactDOM.render(<App />, document.getElementById('root'));
+// configure mobx
+configure({
+   reactionRequiresObservable: true,
+   enforceActions: "observed"
+});
+
+
+const root = createRoot(document.getElementById('root')!);
+root.render(<App />);
 
 // If you want your app to work offline and load faster, you can change
 // unregister() to register() below. Note this comes with some pitfalls.
