Rename JWT handler for tokens issued by Horde itself

To avoid renaming and replacing file in the same CL.

diff --git a/Engine/Source/Programs/Horde/Horde.Build/Authentication/HordeJwtBearerHandler.cs b/Engine/Source/Programs/Horde/Horde.Build/Authentication/HordeJwtBearerHandler.cs
deleted file mode 100644
diff --git a/Engine/Source/Programs/Horde/Horde.Build/Authentication/HordeServerJwtBearerHandler.cs b/Engine/Source/Programs/Horde/Horde.Build/Authentication/HordeServerJwtBearerHandler.cs
new file mode 100644
--- /dev/null
+++ b/Engine/Source/Programs/Horde/Horde.Build/Authentication/HordeServerJwtBearerHandler.cs
@@ -0,0 +1,78 @@
+// Copyright Epic Games, Inc. All Rights Reserved.
+
+using System;
+using System.IdentityModel.Tokens.Jwt;
+using System.Text.Encodings.Web;
+using System.Threading;
+using System.Threading.Tasks;
+using Horde.Build.Server;
+using Horde.Build.Utilities;
+using Microsoft.AspNetCore.Authentication;
+using Microsoft.AspNetCore.Authentication.JwtBearer;
+using Microsoft.Extensions.Logging;
+using Microsoft.Extensions.Options;
+using Microsoft.IdentityModel.Protocols;
+
+namespace Horde.Build.Authentication
+{
+	/// <summary>
+	/// JWT handler for server-issued bearer tokens. These tokens are signed using a randomly generated key per DB instance.
+	/// </summary>
+	class HordeServerJwtBearerHandler : JwtBearerHandler
+	{
+		/// <summary>
+		/// Default name of the authentication scheme
+		/// </summary>
+		public const string AuthenticationScheme = "ServerJwt";
+
+		readonly AsyncCachedValue<IGlobals> _globals;
+
+		public HordeServerJwtBearerHandler(ILoggerFactory logger, UrlEncoder encoder, ISystemClock clock, GlobalsService globalsService, IOptionsMonitorCache<JwtBearerOptions> optionsCache)
+			: base(GetOptionsMonitor(optionsCache), logger, encoder, clock)
+		{
+			_globals = new AsyncCachedValue<IGlobals>(async () => await globalsService.GetAsync(), TimeSpan.FromSeconds(30.0));
+		}
+
+		private static IOptionsMonitor<JwtBearerOptions> GetOptionsMonitor(IOptionsMonitorCache<JwtBearerOptions> optionsCache)
+		{
+			ConfigureNamedOptions<JwtBearerOptions> namedOptions = new ConfigureNamedOptions<JwtBearerOptions>(AuthenticationScheme, options => { });
+			OptionsFactory<JwtBearerOptions> optionsFactory = new OptionsFactory<JwtBearerOptions>(new[] { namedOptions }, Array.Empty<IPostConfigureOptions<JwtBearerOptions>>());
+			return new OptionsMonitor<JwtBearerOptions>(optionsFactory, Array.Empty<IOptionsChangeTokenSource<JwtBearerOptions>>(), optionsCache);
+		}
+
+		protected override async Task<AuthenticateResult> HandleAuthenticateAsync()
+		{
+			// Get the current state
+			IGlobals globals = await _globals.GetAsync();
+
+			Options.TokenValidationParameters.ValidateAudience = false;
+
+			Options.TokenValidationParameters.RequireExpirationTime = false;
+			Options.TokenValidationParameters.ValidateLifetime = true;
+
+			Options.TokenValidationParameters.ValidIssuer = globals.JwtIssuer;
+			Options.TokenValidationParameters.ValidateIssuer = true;
+
+			Options.TokenValidationParameters.ValidateIssuerSigningKey = true;
+			Options.TokenValidationParameters.IssuerSigningKey = globals.JwtSigningKey;
+
+			// Silent fail if this JWT is not issued by the server
+			string? token;
+			if (!JwtUtils.TryGetBearerToken(Request, "Bearer ", out token))
+			{
+				return AuthenticateResult.NoResult();
+			}
+
+			// Validate that it's from the correct issuer, and silent fail if not. Allows multiple handlers for bearer tokens.
+			JwtSecurityToken? jwtToken;
+			if (!JwtUtils.TryParseJwt(token, out jwtToken) || !String.Equals(jwtToken.Issuer, globals.JwtIssuer, StringComparison.Ordinal))
+			{
+				return AuthenticateResult.NoResult();
+			}
+
+			// Pass it to the base class
+			AuthenticateResult result = await base.HandleAuthenticateAsync();
+			return result;
+		}
+	}
+}
diff --git a/Engine/Source/Programs/Horde/Horde.Build/Startup.cs b/Engine/Source/Programs/Horde/Horde.Build/Startup.cs
--- a/Engine/Source/Programs/Horde/Horde.Build/Startup.cs
+++ b/Engine/Source/Programs/Horde/Horde.Build/Startup.cs
@@ -617,8 +617,8 @@ namespace Horde.Build
 					break;
 			}
 
-			authBuilder.AddScheme<JwtBearerOptions, HordeJwtBearerHandler>(HordeJwtBearerHandler.AuthenticationScheme, options => { });
-			schemes.Add(HordeJwtBearerHandler.AuthenticationScheme);
+			authBuilder.AddScheme<JwtBearerOptions, HordeServerJwtBearerHandler>(HordeServerJwtBearerHandler.AuthenticationScheme, options => { });
+			schemes.Add(HordeServerJwtBearerHandler.AuthenticationScheme);
 
 			services.AddAuthorization(options =>
 				{
