Horde: Add support for authenticating via JWT access tokens

To avoid interfering with any existing auth handling, handler is installed under "HordeJwt".
Will resolve any extra user info required to create/configure a valid user request in Horde by downloading an ID token from /userinfo.

diff --git a/Engine/Source/Programs/Horde/Horde.Build/Authentication/HordeJwtBearerHandler.cs b/Engine/Source/Programs/Horde/Horde.Build/Authentication/HordeJwtBearerHandler.cs
new file mode 100644
--- /dev/null
+++ b/Engine/Source/Programs/Horde/Horde.Build/Authentication/HordeJwtBearerHandler.cs
@@ -0,0 +1,215 @@
+// Copyright Epic Games, Inc. All Rights Reserved.
+
+using System;
+using System.Collections.Concurrent;
+using System.IdentityModel.Tokens.Jwt;
+using System.Linq;
+using System.Net.Http;
+using System.Net.Http.Headers;
+using System.Security.Claims;
+using System.Text.Json;
+using System.Threading;
+using System.Threading.Tasks;
+using Horde.Build.Users;
+using Horde.Build.Utilities;
+using Microsoft.AspNetCore.Authentication;
+using Microsoft.AspNetCore.Authentication.JwtBearer;
+using Microsoft.Extensions.DependencyInjection;
+using Microsoft.IdentityModel.Tokens;
+
+namespace Horde.Build.Authentication;
+
+/// <summary>
+/// JWT bearer handler capable of requesting additional user info via OpenID to fully resolve a Horde user
+/// </summary>
+public class HordeJwtBearerHandler
+{
+	/// <summary>
+	/// Default name of the authentication scheme
+	/// </summary>
+	public const string AuthenticationScheme = "HordeJwt";
+
+	// Caches the lookup of "sub" claim in access token to actual user used internally by Horde.
+	// Required to avoid database lookup on every request when authenticating.
+	private readonly ConcurrentDictionary<string, IUser> _subToUser = new();
+	private readonly ConcurrentDictionary<string, JsonElement> _subToUserInfo = new();
+
+	private readonly ServerSettings _settings;
+
+	/// <summary>
+	/// Constructor
+	/// </summary>
+	public HordeJwtBearerHandler(ServerSettings settings)
+	{
+		_settings = settings;
+	}
+	
+	/// <summary>
+	/// Callback when JWT bearer request has been received
+	/// Looks up and extracts the token from Authorization header.
+	/// This cannot be done by the default JwtBearer handler as it can only handle "Bearer" prefix.
+	/// </summary>
+	/// <param name="context">Message context</param>
+	private Task OnMessageReceived(MessageReceivedContext context)
+	{
+		string authorization = context.Request.Headers.Authorization;
+		if (!String.IsNullOrEmpty(authorization))
+		{
+			string prefix = AuthenticationScheme + " ";
+			if (authorization.StartsWith(prefix, StringComparison.OrdinalIgnoreCase))
+			{
+				context.Token = authorization.Substring(prefix.Length).Trim();
+			}
+		}
+		
+		return Task.CompletedTask;
+	}
+
+	/// <summary>
+	/// Callback when JWT bearer token is being validated
+	/// If user is not cached, it will look up additional info via /userinfo and cache it
+	/// </summary>
+	/// <param name="context">Token validation context</param>
+	private async Task OnTokenValidated(TokenValidatedContext context)
+	{
+		if (context.Principal == null)
+		{
+			context.Fail("Principal not set in context");
+			return;
+		}
+
+		JwtSecurityToken? accessToken = context.SecurityToken as JwtSecurityToken;
+		if (accessToken == null)
+		{
+			context.Fail("Unable to read access token");
+			return;
+		}
+
+		ClaimsIdentity? identity = (ClaimsIdentity?)context.Principal?.Identity;
+		if (identity == null)
+		{
+			context.Fail("No identity specified");
+			return;
+		}
+		
+		if (_settings.OidcAuthority == null)
+		{
+			context.Fail("OidcAuthority not set in settings");
+			return;
+		}
+
+		if (!_subToUser.TryGetValue(accessToken.Subject, out IUser? user))
+		{
+			// No cached mapping of 'sub' claim to actual login ID used to fetch the user.
+			// Call /userinfo to get the ID token and get information needed to populate (and maybe create a new user)
+			
+			CancellationToken cancellationToken = context.HttpContext.RequestAborted;
+			Uri userInfoUri = new (_settings.OidcAuthority.TrimEnd('/') + "/v1/userinfo");
+			using HttpClient httpClient = new();
+			using HttpRequestMessage request = new(HttpMethod.Get, userInfoUri);
+			request.Headers.Authorization = new AuthenticationHeaderValue("Bearer", accessToken!.RawData);
+
+			HttpResponseMessage response = await httpClient.SendAsync(request, cancellationToken);
+			if (!response.IsSuccessStatusCode)
+			{
+				context.Fail("Bad status code from OpenID authority when fetching user info: " + response.StatusCode);
+				return;
+			}
+			
+			string userInfoResponse = await response.Content.ReadAsStringAsync(cancellationToken);
+
+			JsonElement userInfo;
+			MediaTypeHeaderValue? contentType = response.Content.Headers.ContentType;
+			if (contentType?.MediaType?.Equals("application/json", StringComparison.OrdinalIgnoreCase) ?? false)
+			{
+				userInfo = JsonDocument.Parse(userInfoResponse).RootElement;
+			}
+			else if (contentType?.MediaType?.Equals("application/jwt", StringComparison.OrdinalIgnoreCase) ?? false)
+			{
+				JwtSecurityToken userInfoEndpointJwt = new(userInfoResponse);
+				userInfo = JsonDocument.Parse(userInfoEndpointJwt.Payload.SerializeToJson()).RootElement;
+			}
+			else
+			{
+				context.Fail("Unknown response type: " + contentType?.MediaType);
+				return;
+			}
+
+			if (!userInfo.TryGetStringProperty("preferred_username", out string? login))
+			{
+				context.Fail("Unable to read property 'preferred_username' from /userinfo");
+				return;
+			}
+
+			if (!userInfo.TryGetStringProperty("name", out string? name))
+			{
+				context.Fail("Unable to read property 'name' from /userinfo");
+				return;
+			}
+
+			if (!userInfo.TryGetStringProperty("email", out string? email))
+			{
+				context.Fail("Unable to read property 'email' from /userinfo");
+				return;
+			}
+
+			IUserCollection userCollection = context.HttpContext.RequestServices.GetRequiredService<IUserCollection>();
+			user = await userCollection.FindOrAddUserByLoginAsync(login, name, email);
+
+			await userCollection.UpdateClaimsAsync(user.Id, accessToken.Claims.Select(x => new UserClaim(x.Type, x.Value)));
+
+			_subToUser[accessToken.Subject] = user;
+			_subToUserInfo[accessToken.Subject] = userInfo;
+		}
+
+		if (!_subToUserInfo.TryGetValue(accessToken.Subject, out JsonElement cachedUserInfo))
+		{
+			context.Fail("Cached user info not found");
+			return;
+		}
+
+		identity.AddClaim(new Claim(HordeClaimTypes.UserId, user.Id.ToString()));
+		HordeOpenIdConnectHandler.AddUserInfoClaims(_settings, cachedUserInfo, identity);
+	}
+
+	/// <summary>
+	/// Registers this instance as a JWT handler
+	/// </summary>
+	/// <param name="authBuilder">Authentication builder</param>
+	public void AddHordeJwtBearerConfiguration(AuthenticationBuilder authBuilder)
+	{
+		authBuilder.AddJwtBearer(AuthenticationScheme, options =>
+		{
+			options.Authority = _settings.OidcAuthority;
+			options.Audience = _settings.OidcAudience;
+			options.Events = new JwtBearerEvents() { OnMessageReceived = OnMessageReceived, OnTokenValidated = OnTokenValidated };
+
+			options.TokenValidationParameters.ValidAudience = _settings.OidcAudience;
+			options.TokenValidationParameters.RequireExpirationTime = true;
+			options.TokenValidationParameters.RequireSignedTokens = true;
+			options.TokenValidationParameters.ValidateIssuer = true;
+			options.TokenValidationParameters.ValidIssuer = _settings.OidcAuthority;
+			options.TokenValidationParameters.ValidateAudience = true;
+			options.TokenValidationParameters.ValidateIssuerSigningKey = true;
+			options.TokenValidationParameters.ValidateLifetime = true;
+			
+			options.SecurityTokenValidators.Clear();
+			options.SecurityTokenValidators.Add(new StrictTokenValidator());
+		});
+	}
+
+	private class StrictTokenValidator : JwtSecurityTokenHandler
+	{
+		public override ClaimsPrincipal ValidateToken(string token, TokenValidationParameters validationParameters, out SecurityToken validatedToken)
+		{
+			ClaimsPrincipal claimsPrincipal = base.ValidateToken(token, validationParameters, out validatedToken);
+			JwtSecurityToken jwtSecurityToken = ReadJwtToken(token);
+			if (jwtSecurityToken.Header?.Alg is not SecurityAlgorithms.RsaSha256)
+			{
+				throw new SecurityTokenValidationException("The JWT algorithm is not RS256.");
+			}
+
+			return claimsPrincipal;
+		}
+	}
+}
\ No newline at end of file
diff --git a/Engine/Source/Programs/Horde/Horde.Build/ServerSettings.cs b/Engine/Source/Programs/Horde/Horde.Build/ServerSettings.cs
--- a/Engine/Source/Programs/Horde/Horde.Build/ServerSettings.cs
+++ b/Engine/Source/Programs/Horde/Horde.Build/ServerSettings.cs
@@ -323,6 +323,11 @@ namespace Horde.Build
 		/// </summary>
 		public AuthMethod AuthMethod { get; set; } = AuthMethod.Anonymous;
 
+		/// <summary>
+		/// Audience for OIDC validation
+		/// </summary>
+		public string? OidcAudience { get; set; }
+		
 		/// <summary>
 		/// Issuer for tokens from the auth provider
 		/// </summary>
diff --git a/Engine/Source/Programs/Horde/Horde.Build/Startup.cs b/Engine/Source/Programs/Horde/Horde.Build/Startup.cs
--- a/Engine/Source/Programs/Horde/Horde.Build/Startup.cs
+++ b/Engine/Source/Programs/Horde/Horde.Build/Startup.cs
@@ -620,6 +620,13 @@ namespace Horde.Build
 			authBuilder.AddScheme<JwtBearerOptions, HordeServerJwtBearerHandler>(HordeServerJwtBearerHandler.AuthenticationScheme, options => { });
 			schemes.Add(HordeServerJwtBearerHandler.AuthenticationScheme);
 
+			if (settings.OidcAuthority != null && settings.OidcAudience != null)
+			{
+				HordeJwtBearerHandler hordeJwtBearer = new(settings);  
+				hordeJwtBearer.AddHordeJwtBearerConfiguration(authBuilder);  
+				schemes.Add(HordeJwtBearerHandler.AuthenticationScheme);				
+			}
+
 			services.AddAuthorization(options =>
 				{
 					options.DefaultPolicy = new AuthorizationPolicyBuilder(schemes.ToArray())
