Horde: Add trigger times to schedule responses
#rnx

diff --git a/Engine/Source/Programs/Horde/Horde.Build/Jobs/Schedules/Schedules.cs b/Engine/Source/Programs/Horde/Horde.Build/Jobs/Schedules/Schedules.cs
--- a/Engine/Source/Programs/Horde/Horde.Build/Jobs/Schedules/Schedules.cs
+++ b/Engine/Source/Programs/Horde/Horde.Build/Jobs/Schedules/Schedules.cs
@@ -71,6 +71,11 @@ namespace Horde.Build.Jobs.Schedules
 		/// </summary>
 		public DateTimeOffset LastTriggerTime { get; set; }
 
+		/// <summary>
+		/// Next trigger times for schedule
+		/// </summary>
+		public List<DateTime> NextTriggerTimesUTC { get; set; }
+
 		/// <summary>
 		/// List of active jobs
 		/// </summary>
@@ -96,6 +101,21 @@ namespace Horde.Build.Jobs.Schedules
 			LastTriggerChange = schedule.LastTriggerChange;
 			LastTriggerTime = schedule.LastTriggerTimeUtc;
 			ActiveJobs = new List<JobId>(schedule.ActiveJobs);
+
+			DateTime curTime = schedule.LastTriggerTimeUtc;
+			NextTriggerTimesUTC = new List<DateTime>();
+			for (int i = 0; i < 16; i++)
+			{
+				DateTime? nextTime = schedule.Config.GetNextTriggerTimeUtc(curTime, TimeZoneInfo.Utc);
+				if (nextTime == null)
+				{
+					break;
+				}
+				
+				curTime = nextTime.Value;
+				NextTriggerTimesUTC.Add(curTime);
+			}
+
 		}
 	}
 
