Horde: Add a few more properties to telemetry:

* State.Job: StreamId (String), PreflightDescription (String)
* State.JobStepRef: JobStartTimeUtc (DateTime), JobId (String), BatchId (String), StepId (String), Id (String)

diff --git a/Engine/Source/Programs/Horde/Horde.Build.Tests/IssueServiceTests.cs b/Engine/Source/Programs/Horde/Horde.Build.Tests/IssueServiceTests.cs
--- a/Engine/Source/Programs/Horde/Horde.Build.Tests/IssueServiceTests.cs
+++ b/Engine/Source/Programs/Horde/Horde.Build.Tests/IssueServiceTests.cs
@@ -209,6 +209,7 @@ namespace Horde.Build.Tests
 		public IJob CreateJob(StreamId streamId, int change, string name, IGraph graph, TimeSpan time = default, bool promoteByDefault = true, bool updateIssues = true)
 		{
 			JobId jobId = JobId.GenerateNewId();
+			DateTime utcNow = DateTime.UtcNow;
 
 			List<IJobStepBatch> batches = new List<IJobStepBatch>();
 			for (int groupIdx = 0; groupIdx < graph.Groups.Count; groupIdx++)
@@ -226,7 +227,7 @@ namespace Horde.Build.Tests
 					step.SetupGet(x => x.Id).Returns(stepId);
 					step.SetupGet(x => x.NodeIdx).Returns(nodeIdx);
 					step.SetupGet(x => x.LogId).Returns(logFile.Id);
-					step.SetupGet(x => x.StartTimeUtc).Returns(DateTime.UtcNow + time);
+					step.SetupGet(x => x.StartTimeUtc).Returns(utcNow + time);
 
 					steps.Add(step.Object);
 				}
@@ -245,6 +246,7 @@ namespace Horde.Build.Tests
 			job.SetupGet(x => x.Name).Returns(name);
 			job.SetupGet(x => x.StreamId).Returns(streamId);
 			job.SetupGet(x => x.TemplateId).Returns(new TemplateId("test-template"));
+			job.SetupGet(x => x.CreateTimeUtc).Returns(utcNow);
 			job.SetupGet(x => x.Change).Returns(change);
 			job.SetupGet(x => x.Batches).Returns(batches);
 			job.SetupGet(x => x.ShowUgsBadges).Returns(promoteByDefault);
@@ -267,7 +269,7 @@ namespace Horde.Build.Tests
 
 			JobStepRefId jobStepRefId = new JobStepRefId(job.Id, batch.Id, step.Id);
 			string nodeName = _graph.Groups[batch.GroupIdx].Nodes[step.NodeIdx].Name;
-			await JobStepRefCollection.InsertOrReplaceAsync(jobStepRefId, "TestJob", nodeName, job.StreamId, job.TemplateId, job.Change, step.LogId, null, null, outcome, job.UpdateIssues, null, null, 0.0f, 0.0f, step.StartTimeUtc!.Value, step.StartTimeUtc);
+			await JobStepRefCollection.InsertOrReplaceAsync(jobStepRefId, "TestJob", nodeName, job.StreamId, job.TemplateId, job.Change, step.LogId, null, null, outcome, job.UpdateIssues, null, null, 0.0f, 0.0f, job.CreateTimeUtc, step.StartTimeUtc!.Value, step.StartTimeUtc);
 		}
 
 		async Task AddEvent(IJob job, int batchIdx, int stepIdx, object data, EventSeverity severity = EventSeverity.Error)
diff --git a/Engine/Source/Programs/Horde/Horde.Build.Tests/Stubs/Collections/JobStepRefCollectionStub.cs b/Engine/Source/Programs/Horde/Horde.Build.Tests/Stubs/Collections/JobStepRefCollectionStub.cs
--- a/Engine/Source/Programs/Horde/Horde.Build.Tests/Stubs/Collections/JobStepRefCollectionStub.cs
+++ b/Engine/Source/Programs/Horde/Horde.Build.Tests/Stubs/Collections/JobStepRefCollectionStub.cs
@@ -100,7 +100,7 @@ namespace Horde.Build.Tests.Stubs.Collections
 			throw new NotImplementedException();
 		}
 
-		Task<IJobStepRef> IJobStepRefCollection.InsertOrReplaceAsync(JobStepRefId id, string jobName, string nodeName, StreamId streamId, TemplateRefId templateId, int change, LogId? logId, PoolId? poolId, AgentId? agentId, JobStepOutcome? outcome, bool updateIssues, int? lastSuccess, int? lastWarning, float waitTime, float initTime, DateTime startTime, DateTime? finishTime)
+		Task<IJobStepRef> IJobStepRefCollection.InsertOrReplaceAsync(JobStepRefId id, string jobName, string nodeName, StreamId streamId, TemplateRefId templateId, int change, LogId? logId, PoolId? poolId, AgentId? agentId, JobStepOutcome? outcome, bool updateIssues, int? lastSuccess, int? lastWarning, float waitTime, float initTime, DateTime jobStartTime, DateTime startTime, DateTime? finishTime)
 		{
 			throw new NotImplementedException();
 		}
diff --git a/Engine/Source/Programs/Horde/Horde.Build/Jobs/IJobStepRefCollection.cs b/Engine/Source/Programs/Horde/Horde.Build/Jobs/IJobStepRefCollection.cs
--- a/Engine/Source/Programs/Horde/Horde.Build/Jobs/IJobStepRefCollection.cs
+++ b/Engine/Source/Programs/Horde/Horde.Build/Jobs/IJobStepRefCollection.cs
@@ -42,9 +42,10 @@ namespace Horde.Build.Jobs
 		/// <param name="lastWarning">The last change that completed with a warning (or success)</param>
 		/// <param name="waitTime">Time taken for the batch containing this step to start</param>
 		/// <param name="initTime">Time taken for the batch containing this step to initializer</param>
+		/// <param name="jobStartTimeUtc">Start time of the job</param>
 		/// <param name="startTimeUtc">Start time</param>
 		/// <param name="finishTimeUtc">Finish time for the step, if known</param>
-		Task<IJobStepRef> InsertOrReplaceAsync(JobStepRefId id, string jobName, string stepName, StreamId streamId, TemplateId templateId, int change, LogId? logId, PoolId? poolId, AgentId? agentId, JobStepOutcome? outcome, bool updateIssues, int? lastSuccess, int? lastWarning, float waitTime, float initTime, DateTime startTimeUtc, DateTime? finishTimeUtc);
+		Task<IJobStepRef> InsertOrReplaceAsync(JobStepRefId id, string jobName, string stepName, StreamId streamId, TemplateId templateId, int change, LogId? logId, PoolId? poolId, AgentId? agentId, JobStepOutcome? outcome, bool updateIssues, int? lastSuccess, int? lastWarning, float waitTime, float initTime, DateTime jobStartTimeUtc, DateTime startTimeUtc, DateTime? finishTimeUtc);
 
 		/// <summary>
 		/// Gets the history of a given node
@@ -125,7 +126,7 @@ namespace Horde.Build.Jobs
 					logger.LogInformation("Updating step reference {StepId} for job {JobId}, batch {BatchId}, with outcome {JobStepOutcome}", step.Id, job.Id, batch.Id, outcome);
 				}				
 
-				await jobStepRefs.InsertOrReplaceAsync(new JobStepRefId(job.Id, batch.Id, step.Id), job.Name, nodeName, job.StreamId, job.TemplateId, job.Change, step.LogId, batch.PoolId, batch.AgentId, outcome, job.UpdateIssues, lastSuccess, lastWarning, waitTime, initTime, step.StartTimeUtc ?? DateTime.UtcNow, step.FinishTimeUtc);
+				await jobStepRefs.InsertOrReplaceAsync(new JobStepRefId(job.Id, batch.Id, step.Id), job.Name, nodeName, job.StreamId, job.TemplateId, job.Change, step.LogId, batch.PoolId, batch.AgentId, outcome, job.UpdateIssues, lastSuccess, lastWarning, waitTime, initTime, job.CreateTimeUtc, step.StartTimeUtc ?? DateTime.UtcNow, step.FinishTimeUtc);
 			}
 		}
 	}
diff --git a/Engine/Source/Programs/Horde/Horde.Build/Jobs/JobCollection.cs b/Engine/Source/Programs/Horde/Horde.Build/Jobs/JobCollection.cs
--- a/Engine/Source/Programs/Horde/Horde.Build/Jobs/JobCollection.cs
+++ b/Engine/Source/Programs/Horde/Horde.Build/Jobs/JobCollection.cs
@@ -416,6 +416,7 @@ namespace Horde.Build.Jobs
 				_telemetrySink.SendEvent("State.Job", new
 				{
 					Id = newJob.Id,
+					StreamId = newJob.StreamId,
 					Arguments = newJob.Arguments,
 					AutoSubmit = newJob.AutoSubmit,
 					Change = newJob.Change,
@@ -424,6 +425,7 @@ namespace Horde.Build.Jobs
 					GraphHash = newJob.GraphHash,
 					Name = newJob.Name,
 					PreflightChange = newJob.PreflightChange,
+					PreflightDescription = newJob.PreflightDescription,
 					Priority = newJob.Priority,
 					StartedByUserId = newJob.StartedByUserId,
 					TemplateId = newJob.TemplateId
diff --git a/Engine/Source/Programs/Horde/Horde.Build/Jobs/JobStepRefCollection.cs b/Engine/Source/Programs/Horde/Horde.Build/Jobs/JobStepRefCollection.cs
--- a/Engine/Source/Programs/Horde/Horde.Build/Jobs/JobStepRefCollection.cs
+++ b/Engine/Source/Programs/Horde/Horde.Build/Jobs/JobStepRefCollection.cs
@@ -53,6 +53,8 @@ namespace Horde.Build.Jobs
 			public float BatchWaitTime { get; set; }
 			public float BatchInitTime { get; set; }
 
+			public DateTime JobStartTimeUtc { get; set; }
+
 			[BsonIgnoreIfNull]
 			public DateTimeOffset? StartTime { get; set; }
 
@@ -73,7 +75,7 @@ namespace Horde.Build.Jobs
 			string IJobStepRef.NodeName => Name;
 			bool IJobStepRef.UpdateIssues => UpdateIssues ?? false;
 
-			public JobStepRef(JobStepRefId id, string jobName, string nodeName, StreamId streamId, TemplateId templateId, int change, LogId? logId, PoolId? poolId, AgentId? agentId, JobStepOutcome? outcome, bool updateIssues, int? lastSuccess, int? lastWarning, float batchWaitTime, float batchInitTime, DateTime startTimeUtc, DateTime? finishTimeUtc)
+			public JobStepRef(JobStepRefId id, string jobName, string nodeName, StreamId streamId, TemplateId templateId, int change, LogId? logId, PoolId? poolId, AgentId? agentId, JobStepOutcome? outcome, bool updateIssues, int? lastSuccess, int? lastWarning, float batchWaitTime, float batchInitTime, DateTime jobStartTimeUtc, DateTime startTimeUtc, DateTime? finishTimeUtc)
 			{
 				Id = id;
 				JobName = jobName;
@@ -90,9 +92,8 @@ namespace Horde.Build.Jobs
 				LastWarning = lastWarning;
 				BatchWaitTime = batchWaitTime;
 				BatchInitTime = batchInitTime;
-				StartTime = startTimeUtc;
+				JobStartTimeUtc = jobStartTimeUtc;
 				StartTimeUtc = startTimeUtc;
-				FinishTime = finishTimeUtc;
 				FinishTimeUtc = finishTimeUtc;
 			}
 		}
@@ -112,22 +113,26 @@ namespace Horde.Build.Jobs
 		}
 
 		/// <inheritdoc/>
-		public async Task<IJobStepRef> InsertOrReplaceAsync(JobStepRefId id, string jobName, string stepName, StreamId streamId, TemplateId templateId, int change, LogId? logId, PoolId? poolId, AgentId? agentId, JobStepOutcome? outcome, bool updateIssues, int? lastSuccess, int? lastWarning, float waitTime, float initTime, DateTime startTimeUtc, DateTime? finishTimeUtc)
+		public async Task<IJobStepRef> InsertOrReplaceAsync(JobStepRefId id, string jobName, string stepName, StreamId streamId, TemplateId templateId, int change, LogId? logId, PoolId? poolId, AgentId? agentId, JobStepOutcome? outcome, bool updateIssues, int? lastSuccess, int? lastWarning, float waitTime, float initTime, DateTime jobStartTimeUtc, DateTime startTimeUtc, DateTime? finishTimeUtc)
 		{
-			JobStepRef newJobStepRef = new JobStepRef(id, jobName, stepName, streamId, templateId, change, logId, poolId, agentId, outcome, updateIssues, lastSuccess, lastWarning, waitTime, initTime, startTimeUtc, finishTimeUtc);
+			JobStepRef newJobStepRef = new JobStepRef(id, jobName, stepName, streamId, templateId, change, logId, poolId, agentId, outcome, updateIssues, lastSuccess, lastWarning, waitTime, initTime, jobStartTimeUtc, startTimeUtc, finishTimeUtc);
 			await _jobStepRefs.ReplaceOneAsync(Builders<JobStepRef>.Filter.Eq(x => x.Id, newJobStepRef.Id), newJobStepRef, new ReplaceOptions { IsUpsert = true });
 
 			if (_telemetrySink.Enabled)
 			{
 				_telemetrySink.SendEvent("State.JobStepRef", new
 				{
-					Id = id,
+					Id = id.ToString(),
+					JobId = id.JobId.ToString(),
+					BatchId = id.BatchId.ToString(),
+					StepId = id.BatchId.ToString(),
 					AgentId = agentId,
 					BatchInitTime = initTime,
 					BatchWaitTime = waitTime,
 					Change = change,
 					FinishTime = finishTimeUtc,
 					JobName = jobName,
+					JobStartTime = jobStartTimeUtc,
 					StepName = stepName,
 					Outcome = outcome,
 					PoolId = poolId,
@@ -160,7 +165,7 @@ namespace Horde.Build.Jobs
 				filter &= filterBuilder.Ne(x => x.Outcome, JobStepOutcome.Failure);
 			}
 
-			List<JobStepRef> steps = await _jobStepRefs.Find(filter).SortByDescending(x => x.Change).ThenByDescending(x => x.StartTime).Limit(count).ToListAsync();
+			List<JobStepRef> steps = await _jobStepRefs.Find(filter).SortByDescending(x => x.Change).ThenByDescending(x => x.StartTimeUtc).Limit(count).ToListAsync();
 			return steps.ConvertAll<IJobStepRef>(x => x);
 		}
 
