Horde: Prevent non-screenshot events being ignored in Gauntlet issues.

diff --git a/Engine/Source/Programs/Horde/Horde.Build.Tests/IssueServiceTests.cs b/Engine/Source/Programs/Horde/Horde.Build.Tests/IssueServiceTests.cs
--- a/Engine/Source/Programs/Horde/Horde.Build.Tests/IssueServiceTests.cs
+++ b/Engine/Source/Programs/Horde/Horde.Build.Tests/IssueServiceTests.cs
@@ -627,6 +627,41 @@ namespace Horde.Build.Tests
 			}
 		}
 
+		[TestMethod]
+		public async Task EnsureWarningTest()
+		{
+			// #1
+			// Scenario: Job step completes successfully at CL 105
+			// Expected: No issues are created
+			{
+				IJob job = CreateJob(_mainStreamId, 105, "Test Build", _graph);
+				await UpdateCompleteStep(job, 0, 0, JobStepOutcome.Success);
+
+				List<IIssue> issues = await IssueCollection.FindIssuesAsync();
+				Assert.AreEqual(0, issues.Count);
+
+				List<IIssue> openIssues = await IssueCollection.FindIssuesAsync(resolved: false);
+				Assert.AreEqual(0, openIssues.Count);
+			}
+
+			// #2
+			// Scenario: Job step fails at CL 120 in Gauntlet
+			// Expected: Creates issue
+			{
+				IJob job = CreateJob(_mainStreamId, 120, "Test Build", _graph);
+				await AddEvent(job, 0, 0, new { level = nameof(LogLevel.Warning), id = KnownLogEvents.Gauntlet_TestEvent.Id, message = "" }, EventSeverity.Warning);
+				await UpdateCompleteStep(job, 0, 0, JobStepOutcome.Warnings);
+
+				ILogFile? log = await LogFileService.GetLogFileAsync(job.Batches[0].Steps[0].LogId!.Value, CancellationToken.None);
+				List<ILogEvent> events = await LogFileService.FindEventsAsync(log!);
+				Assert.AreEqual(1, events.Count);
+				Assert.AreEqual(1, events[0].LineCount);
+
+				List<IIssue> issues = await IssueCollection.FindIssuesAsync();
+				Assert.AreEqual(1, issues.Count);
+			}
+		}
+
 		[TestMethod]
 		public async Task GenericErrorTest()
 		{
diff --git a/Engine/Source/Programs/Horde/Horde.Build/Issues/Handlers/GauntletIssueHandler.cs b/Engine/Source/Programs/Horde/Horde.Build/Issues/Handlers/GauntletIssueHandler.cs
--- a/Engine/Source/Programs/Horde/Horde.Build/Issues/Handlers/GauntletIssueHandler.cs
+++ b/Engine/Source/Programs/Horde/Horde.Build/Issues/Handlers/GauntletIssueHandler.cs
@@ -110,11 +110,7 @@ namespace Horde.Build.Issues.Handlers
 					GetUnitTestNames(stepEvent.EventData, keys);
 					GetScreenshotTestNames(stepEvent.EventData, keys);
 
-					if (keys.Count == 0)
-					{
-						stepEvent.Ignored = true;
-					}
-					else
+					if (keys.Count > 0)
 					{
 						stepEvent.Fingerprint = new NewIssueFingerprint(Type, keys, null, null);
 					}
