Horde: Use server scheduler time zone for calculating next trigger times
#rnx

diff --git a/Engine/Source/Programs/Horde/Horde.Build/Jobs/Schedules/Schedules.cs b/Engine/Source/Programs/Horde/Horde.Build/Jobs/Schedules/Schedules.cs
--- a/Engine/Source/Programs/Horde/Horde.Build/Jobs/Schedules/Schedules.cs
+++ b/Engine/Source/Programs/Horde/Horde.Build/Jobs/Schedules/Schedules.cs
@@ -85,7 +85,8 @@ namespace Horde.Build.Jobs.Schedules
 		/// Constructor
 		/// </summary>
 		/// <param name="schedule">Schedule to construct from</param>
-		public GetScheduleResponse(ITemplateSchedule schedule)
+		/// <param name="schedulerTimeZone">The scheduler time zone</param>
+		public GetScheduleResponse(ITemplateSchedule schedule, TimeZoneInfo schedulerTimeZone)
 		{
 			Enabled = schedule.Config.Enabled;
 			MaxActive = schedule.Config.MaxActive;
@@ -106,7 +107,7 @@ namespace Horde.Build.Jobs.Schedules
 			NextTriggerTimesUTC = new List<DateTime>();
 			for (int i = 0; i < 16; i++)
 			{
-				DateTime? nextTime = schedule.Config.GetNextTriggerTimeUtc(curTime, TimeZoneInfo.Utc);
+				DateTime? nextTime = schedule.Config.GetNextTriggerTimeUtc(curTime, schedulerTimeZone);
 				if (nextTime == null)
 				{
 					break;
diff --git a/Engine/Source/Programs/Horde/Horde.Build/Streams/StreamMessages.cs b/Engine/Source/Programs/Horde/Horde.Build/Streams/StreamMessages.cs
--- a/Engine/Source/Programs/Horde/Horde.Build/Streams/StreamMessages.cs
+++ b/Engine/Source/Programs/Horde/Horde.Build/Streams/StreamMessages.cs
@@ -273,7 +273,8 @@ namespace Horde.Build.Streams
 		/// <param name="templateRef">The template ref</param>
 		/// <param name="template">The actual template</param>
 		/// <param name="stepStates">The template step states</param>
-		public GetTemplateRefResponse(TemplateId id, ITemplateRef templateRef, ITemplate template, List<GetTemplateStepStateResponse>? stepStates)
+		/// <param name="schedulerTimeZone">The scheduler time zone</param>
+		public GetTemplateRefResponse(TemplateId id, ITemplateRef templateRef, ITemplate template, List<GetTemplateStepStateResponse>? stepStates, TimeZoneInfo schedulerTimeZone)
 			: base(template)
 		{
 			Id = id.ToString();
@@ -282,7 +283,7 @@ namespace Horde.Build.Streams
 			ShowUgsAlerts = templateRef.Config.ShowUgsAlerts;
 			NotificationChannel = templateRef.Config.NotificationChannel;
 			NotificationChannelFilter = templateRef.Config.NotificationChannelFilter;
-			Schedule = (templateRef.Schedule != null) ? new GetScheduleResponse(templateRef.Schedule) : null;
+			Schedule = (templateRef.Schedule != null) ? new GetScheduleResponse(templateRef.Schedule, schedulerTimeZone) : null;
 			ChainedJobs = (templateRef.Config.ChainedJobs != null && templateRef.Config.ChainedJobs.Count > 0) ? templateRef.Config.ChainedJobs : null;
 			StepStates = stepStates;
 			DefaultChange = templateRef.Config.DefaultChange;
diff --git a/Engine/Source/Programs/Horde/Horde.Build/Streams/StreamsController.cs b/Engine/Source/Programs/Horde/Horde.Build/Streams/StreamsController.cs
--- a/Engine/Source/Programs/Horde/Horde.Build/Streams/StreamsController.cs
+++ b/Engine/Source/Programs/Horde/Horde.Build/Streams/StreamsController.cs
@@ -18,6 +18,7 @@ using Microsoft.AspNetCore.Mvc;
 using Microsoft.Extensions.Options;
 using OpenTracing;
 using OpenTracing.Util;
+using TimeZoneConverter;
 
 namespace Horde.Build.Streams
 {
@@ -39,6 +40,7 @@ namespace Horde.Build.Streams
 		private readonly IJobStepRefCollection _jobStepRefCollection;
 		private readonly IUserCollection _userCollection;
 		private readonly IOptionsSnapshot<GlobalConfig> _globalConfig;
+		private readonly TimeZoneInfo _timeZone;
 
 		/// <summary>
 		/// Constructor
@@ -51,6 +53,9 @@ namespace Horde.Build.Streams
 			_jobStepRefCollection = jobStepRefCollection;
 			_userCollection = userCollection;
 			_globalConfig = globalConfig;
+
+			string? timeZoneName = _globalConfig.Value.ServerSettings.ScheduleTimeZone;
+			_timeZone = (timeZoneName == null) ? TimeZoneInfo.Local : TZConvert.GetTimeZoneInfo(timeZoneName);
 		}
 
 		/// <summary>
@@ -180,7 +185,7 @@ namespace Horde.Build.Streams
 		{
 			using IScope scope = GlobalTracer.Instance.BuildSpan("CreateGetStreamResponse").StartActive();
 			scope.Span.SetTag("streamId", stream.Id);
-			
+
 			List<GetTemplateRefResponse> apiTemplateRefs = new List<GetTemplateRefResponse>();
 			foreach (KeyValuePair<TemplateId, ITemplateRef> pair in stream.Templates)
 			{
@@ -205,7 +210,7 @@ namespace Horde.Build.Streams
 					}
 
 					ITemplate? template = await _templateCollection.GetOrAddAsync(pair.Value.Config);
-					apiTemplateRefs.Add(new GetTemplateRefResponse(pair.Key, pair.Value, template, stepStates));
+					apiTemplateRefs.Add(new GetTemplateRefResponse(pair.Key, pair.Value, template, stepStates, _timeZone));
 				}
 			}
 
