Horde: Add logging for JWT auth

diff --git a/Engine/Source/Programs/Horde/Horde.Build/Authentication/HordeJwtBearerHandler.cs b/Engine/Source/Programs/Horde/Horde.Build/Authentication/HordeJwtBearerHandler.cs
--- a/Engine/Source/Programs/Horde/Horde.Build/Authentication/HordeJwtBearerHandler.cs
+++ b/Engine/Source/Programs/Horde/Horde.Build/Authentication/HordeJwtBearerHandler.cs
@@ -15,6 +15,7 @@ using Horde.Build.Utilities;
 using Microsoft.AspNetCore.Authentication;
 using Microsoft.AspNetCore.Authentication.JwtBearer;
 using Microsoft.Extensions.DependencyInjection;
+using Microsoft.Extensions.Logging;
 using Microsoft.IdentityModel.Tokens;
 
 namespace Horde.Build.Authentication;
@@ -72,29 +73,31 @@ public class HordeJwtBearerHandler
 	/// <param name="context">Token validation context</param>
 	private async Task OnTokenValidated(TokenValidatedContext context)
 	{
+		ILogger<HordeJwtBearerHandler> logger = context.HttpContext.RequestServices.GetRequiredService<ILogger<HordeJwtBearerHandler>>();
+		
 		if (context.Principal == null)
 		{
-			context.Fail("Principal not set in context");
+			ReportError(logger, context, "Principal not set in context");
 			return;
 		}
 
 		JwtSecurityToken? accessToken = context.SecurityToken as JwtSecurityToken;
 		if (accessToken == null)
 		{
-			context.Fail("Unable to read access token");
+			ReportError(logger, context, "Unable to read access token");
 			return;
 		}
 
 		ClaimsIdentity? identity = (ClaimsIdentity?)context.Principal?.Identity;
 		if (identity == null)
 		{
-			context.Fail("No identity specified");
+			ReportError(logger, context, "No identity specified");
 			return;
 		}
 		
 		if (_settings.OidcAuthority == null)
 		{
-			context.Fail("OidcAuthority not set in settings");
+			ReportError(logger, context, "OidcAuthority not set in settings");
 			return;
 		}
 
@@ -112,7 +115,7 @@ public class HordeJwtBearerHandler
 			HttpResponseMessage response = await httpClient.SendAsync(request, cancellationToken);
 			if (!response.IsSuccessStatusCode)
 			{
-				context.Fail("Bad status code from OpenID authority when fetching user info: " + response.StatusCode);
+				ReportError(logger, context, "Bad status code from OpenID authority when fetching user info: " + response.StatusCode);
 				return;
 			}
 			
@@ -131,25 +134,25 @@ public class HordeJwtBearerHandler
 			}
 			else
 			{
-				context.Fail("Unknown response type: " + contentType?.MediaType);
+				ReportError(logger, context, "Unknown response type: " + contentType?.MediaType);
 				return;
 			}
 
 			if (!userInfo.TryGetStringProperty("preferred_username", out string? login))
 			{
-				context.Fail("Unable to read property 'preferred_username' from /userinfo");
+				ReportError(logger, context, "Unable to read property 'preferred_username' from /userinfo");
 				return;
 			}
 
 			if (!userInfo.TryGetStringProperty("name", out string? name))
 			{
-				context.Fail("Unable to read property 'name' from /userinfo");
+				ReportError(logger, context, "Unable to read property 'name' from /userinfo");
 				return;
 			}
 
 			if (!userInfo.TryGetStringProperty("email", out string? email))
 			{
-				context.Fail("Unable to read property 'email' from /userinfo");
+				ReportError(logger, context, "Unable to read property 'email' from /userinfo");
 				return;
 			}
 
@@ -164,7 +167,7 @@ public class HordeJwtBearerHandler
 
 		if (!_subToUserInfo.TryGetValue(accessToken.Subject, out JsonElement cachedUserInfo))
 		{
-			context.Fail("Cached user info not found");
+			ReportError(logger, context, "Cached user info not found");
 			return;
 		}
 
@@ -172,6 +175,14 @@ public class HordeJwtBearerHandler
 		HordeOpenIdConnectHandler.AddUserInfoClaims(_settings, cachedUserInfo, identity);
 	}
 
+	private static void ReportError(ILogger<HordeJwtBearerHandler> logger, ResultContext<JwtBearerOptions> context, string message)
+	{
+#pragma warning disable CA2254 // Template should be a static expression
+		logger.LogError(message);
+#pragma warning restore CA2254
+		context.Fail(message);
+	}
+
 	/// <summary>
 	/// Registers this instance as a JWT handler
 	/// </summary>
