Horde: Readding customized react charting components, pinned version
#rnx
#skipci

diff --git a/Engine/Source/Programs/Horde/HordeDashboard/package.json b/Engine/Source/Programs/Horde/HordeDashboard/package.json
--- a/Engine/Source/Programs/Horde/HordeDashboard/package.json
+++ b/Engine/Source/Programs/Horde/HordeDashboard/package.json
@@ -8,7 +8,7 @@
     "@blueprintjs/icons": "^3.13.0",
     "@datadog/browser-logs": "^4.30.0",
     "@fluentui/react": "^8.104.5",
-    "@fluentui/react-charting": "^5.14.36",
+    "@fluentui/react-charting": "5.2.0",
     "crypto-js": "^4.1.1",
     "d3": "^7.6.1",
     "idb": "^7.1.1",
diff --git a/Engine/Source/Programs/Horde/HordeDashboard/src/base/components/HordeFluentCharting/HCartesianChart/HCartesianChart.base.tsx b/Engine/Source/Programs/Horde/HordeDashboard/src/base/components/HordeFluentCharting/HCartesianChart/HCartesianChart.base.tsx
new file mode 100644
--- /dev/null
+++ b/Engine/Source/Programs/Horde/HordeDashboard/src/base/components/HordeFluentCharting/HCartesianChart/HCartesianChart.base.tsx
@@ -0,0 +1,441 @@
+import * as React from 'react';
+import { IProcessedStyleSet } from '@fluentui/react/lib/Styling';
+import { classNamesFunction, getId, getRTL } from '@fluentui/react/lib/Utilities';
+import { Callout } from '@fluentui/react/lib/Callout';
+import { FocusZone, FocusZoneDirection } from '@fluentui/react-focus';
+import { IHCartesianChartStyleProps, IHModifiedCartesianChartProps } from './HCartesianChart.types';
+import { ChartHoverCard, ICartesianChartStyles, IMargins, IModifiedCartesianChartProps, IYValueHover, LegendShape, Shape } from '@fluentui/react-charting';
+import { createDateXAxis, createStringXAxis, createStringYAxis, createWrapOfXLabels, getDomainNRangeValues, getMinMaxOfYAxis, Points, pointTypes, XAxisTypes, YAxisType } from '@fluentui/react-charting/lib/utilities/utilities';
+import { hCreateNumericXAxis, hCreateYAxis } from '../Hutilities';
+
+const getClassNames = classNamesFunction<IHCartesianChartStyleProps, ICartesianChartStyles>();
+
+export interface ICartesianChartState {
+  containerWidth: number;
+  containerHeight: number;
+  _width: number;
+  _height: number;
+  /* To update this values using setState in render method.
+   * To avoid multiple re renders, Only first time setting the value.
+   */
+  isRemoveValCalculated?: boolean;
+  /* Used for when WrapXAxisLabels props appeared.
+   * To display the total word (space separated words), Need to have more space than usual.
+   * This height will get total height need to disaply total word.
+   * These value need to be removed from actual svg height/graph height.
+   * Defalut value is 0. And this values calculted when 'wrapXAxisLables' or 'showXAxisLablesTooltip' is true.
+   */
+  _removalValueForTextTuncate?: number;
+}
+
+/**
+ * Cartesian chart used for
+ * 1.draw X and Y axis of the chart
+ * 2.Callout
+ * 3.Fit parent Continer
+ */
+export class HCartesianChartBase extends React.Component<IHModifiedCartesianChartProps, ICartesianChartState> {
+  private _classNames: IProcessedStyleSet<ICartesianChartStyles> | undefined;
+  private chartContainer: HTMLDivElement | undefined;
+  private legendContainer: HTMLDivElement | undefined;
+  private minLegendContainerHeight: number = 32;
+  private xAxisElement: SVGElement | null | undefined;
+  private yAxisElement: SVGElement | null | undefined;
+  private margins: IMargins;
+  private idForGraph: string;
+  private _reqID: number | undefined;
+  private _isRtl: boolean = getRTL();
+
+  constructor(props: IHModifiedCartesianChartProps) {
+    super(props);
+    this.state = {
+      containerHeight: 0,
+      containerWidth: 0,
+      _width: this.props.width || 600,
+      _height: this.props.height || 350,
+      _removalValueForTextTuncate: 0,
+      isRemoveValCalculated: true,
+    };
+    this.idForGraph = getId('chart_');
+    /**
+     * In RTL mode, Only graph will be rendered left/right. We need to provide left and right margins manually.
+     * So that, in RTL, left margins becomes right margins and viceversa.
+     * As graph needs to be drawn perfecty, these values consider as default values.
+     * Same margins using for all other cartesian charts. Can be accessible through 'getMargins' call back method.
+     */
+    this.margins = {
+      top: this.props.margins?.top ?? 20,
+      bottom: this.props.margins?.bottom ?? 35,
+      right: this._isRtl ? this.props.margins?.left ?? 40 : this.props.margins?.right ?? 20,
+      left: this._isRtl ? this.props.margins?.right ?? 20 : this.props.margins?.left ?? 40,
+    };
+  }
+
+  public componentDidMount(): void {
+    this._fitParentContainer();
+  }
+
+  public componentWillUnmount(): void {
+    cancelAnimationFrame(this._reqID!);
+  }
+
+  public componentDidUpdate(prevProps: IModifiedCartesianChartProps): void {
+    if (prevProps.height !== this.props.height || prevProps.width !== this.props.width) {
+      this._fitParentContainer();
+    }
+  }
+
+  public render(): JSX.Element {
+    const { calloutProps, points, chartType, chartHoverProps, svgFocusZoneProps, svgProps } = this.props;
+    if (this.props.parentRef) {
+      this._fitParentContainer();
+    }
+    // Callback for margins to the chart
+    if(this.props.getmargins) {
+      let newMargins = this.props.getmargins();
+      if(newMargins) {
+        this.margins = newMargins;
+      }
+    }
+
+    const XAxisParams = {
+      domainNRangeValues: getDomainNRangeValues(
+        points,
+        this.margins,
+        this.state.containerWidth,
+        chartType,
+        this._isRtl,
+        this.props.xAxisType,
+        this.props.barwidth!,
+        this.props.tickValues!,
+      ),
+      xAxisElement: this.xAxisElement!,
+      showRoundOffXTickValues: true,
+      xAxisCount: this.props.xAxisTickCount,
+      xAxistickSize: this.props.xAxistickSize,
+      tickPadding: this.props.tickPadding || this.props.showXAxisLablesTooltip ? 5 : 10,
+      xAxisPadding: this.props.xAxisPadding,
+    };
+
+    const YAxisParams = {
+      margins: this.margins,
+      containerWidth: this.state.containerWidth,
+      containerHeight: this.state.containerHeight - this.state._removalValueForTextTuncate!,
+      yAxisElement: this.yAxisElement,
+      yAxisTickFormat: this.props.yAxisTickFormat!,
+      yAxisTickCount: this.props.yAxisTickCount!,
+      yMinValue: this.props.yMinValue || 0,
+      yMaxValue: this.props.yMaxValue || 0,
+      tickPadding: 10,
+      maxOfYVal: this.props.maxOfYVal,
+      yMinMaxValues: getMinMaxOfYAxis(points, chartType),
+      yAxisPadding: this.props.yAxisPadding,
+    };
+
+    /**
+     * These scales used for 2 purposes.
+     * 1. To create x and y axis
+     * 2. To draw the graph.
+     * For area/line chart using same scales. For other charts, creating their own scales to draw the graph.
+     */
+    // eslint-disable-next-line @typescript-eslint/no-explicit-any
+    let xScale: any;
+    switch (this.props.xAxisType!) {
+      case XAxisTypes.NumericAxis:
+        xScale = hCreateNumericXAxis(XAxisParams, this.props.tickParams);
+        break;
+      case XAxisTypes.DateAxis:
+        xScale = createDateXAxis(XAxisParams, this.props.tickParams!);
+        break;
+      case XAxisTypes.StringAxis:
+        xScale = createStringXAxis(XAxisParams, this.props.tickParams!, this.props.datasetForXAxisDomain!);
+        break;
+      default:
+        xScale = hCreateNumericXAxis(XAxisParams, this.props.tickParams);
+    }
+
+    if (this.props.wrapXAxisLables || this.props.showXAxisLablesTooltip) {
+      const wrapLabelProps = {
+        node: this.xAxisElement ?? null,
+        xAxis: xScale,
+        showXAxisLablesTooltip: this.props.showXAxisLablesTooltip || false,
+        noOfCharsToTruncate: this.props.noOfCharsToTruncate || 4,
+      };
+      const temp = xScale && (createWrapOfXLabels(wrapLabelProps) as number);
+      // this value need to be updated for draw graph updated. So instead of using private value, using set state.
+      if (this.state.isRemoveValCalculated && this.state._removalValueForTextTuncate !== temp) {
+        this.setState({ _removalValueForTextTuncate: temp, isRemoveValCalculated: false });
+      }
+    }
+
+    /**
+     * These scales used for 2 purposes.
+     * 1. To create x and y axis
+     * 2. To draw the graph.
+     * For area/line chart using same scales. For other charts, creating their own scales to draw the graph.
+     */
+    // eslint-disable-next-line @typescript-eslint/no-explicit-any
+    let yScale: any;
+    if (this.props.yAxisType && this.props.yAxisType === YAxisType.StringAxis) {
+      yScale = createStringYAxis(YAxisParams, this.props.stringDatasetForYAxisDomain!, this._isRtl);
+    } else {
+      yScale = hCreateYAxis(YAxisParams, this._isRtl);
+    }
+
+    // Callback function for chart, returns axis
+    this._getData(xScale, yScale);
+
+    this._classNames = getClassNames(this.props.styles!, {
+      theme: this.props.theme!,
+      width: this.state._width,
+      height: this.state._height,
+      className: this.props.className,
+      isRtl: this._isRtl,
+    });
+    const svgDimensions = {
+      width: this.state.containerWidth,
+      height: this.state.containerHeight,
+    };
+    const children = this.props.children({
+      ...this.state,
+      xScale,
+      yScale,
+    });
+    let focusDirection;
+    if (this.props.focusZoneDirection === FocusZoneDirection.vertical) {
+      focusDirection = this.props.focusZoneDirection;
+    } else if (this.props.focusZoneDirection) {
+      focusDirection = this.props.focusZoneDirection;
+    } else {
+      focusDirection = FocusZoneDirection.horizontal;
+    }
+    return (
+      <div
+        id={this.idForGraph}
+        className={this._classNames.root}
+        role={'presentation'}
+        ref={(rootElem: HTMLDivElement) => (this.chartContainer = rootElem)}
+      >
+        <FocusZone direction={focusDirection} {...svgFocusZoneProps}>
+          <svg width={svgDimensions.width} height={svgDimensions.height} style={{ display: 'block' }} {...svgProps}>
+            <g
+              ref={(e: SVGElement | null) => {
+                this.xAxisElement = e;
+              }}
+              id={`xAxisGElement${this.idForGraph}`}
+              // To add wrap of x axis lables feature, need to remove word height from svg height.
+              transform={`translate(0, ${svgDimensions.height -
+                this.margins.bottom! -
+                this.state._removalValueForTextTuncate!})`}
+              className={this._classNames.xAxis}
+            />
+            <g
+              ref={(e: SVGElement | null) => {
+                this.yAxisElement = e;
+              }}
+              id={`yAxisGElement${this.idForGraph}`}
+              transform={`translate(${
+                this._isRtl ? svgDimensions.width - this.margins.right! : this.margins.left!
+              }, 0)`}
+              className={this._classNames.yAxis}
+            />
+            {children}
+          </svg>
+        </FocusZone>
+        {!this.props.hideLegend && (
+          <div ref={(e: HTMLDivElement) => (this.legendContainer = e)} className={this._classNames.legendContainer}>
+            {this.props.legendBars}
+          </div>
+        )}
+        {!this.props.hideTooltip && calloutProps!.isCalloutVisible && (
+          <Callout {...calloutProps}>
+            {/** Given custom callout, then it will render */}
+            {this.props.customizedCallout && this.props.customizedCallout}
+            {/** single x point its corresponding y points of all the bars/lines in chart will render in callout */}
+            {!this.props.customizedCallout && this.props.isCalloutForStack && this._multiValueCallout(calloutProps)}
+            {/** single x point its corresponding y point of single line/bar in the chart will render in callout */}
+            {!this.props.customizedCallout && !this.props.isCalloutForStack && (
+              <ChartHoverCard
+                XValue={calloutProps.XValue}
+                Legend={calloutProps.legend!}
+                YValue={calloutProps.YValue!}
+                color={calloutProps.color!}
+                {...chartHoverProps}
+              />
+            )}
+          </Callout>
+        )}
+      </div>
+    );
+  }
+
+  // TO DO: Write a common funtional component for Multi value callout and divide sub count method
+  // eslint-disable-next-line @typescript-eslint/no-explicit-any
+  private _multiValueCallout = (calloutProps: any) => {
+    const yValueHoverSubCountsExists: boolean = this._yValueHoverSubCountsExists(calloutProps.YValueHover);
+    return (
+      <div className={this._classNames!.calloutContentRoot}>
+        <div
+          className={this._classNames!.calloutDateTimeContainer}
+          style={yValueHoverSubCountsExists ? { marginBottom: '11px' } : {}}
+        >
+          <div className={this._classNames!.calloutContentX}>{calloutProps!.hoverXValue} </div>
+        </div>
+        <div
+          className={this._classNames!.calloutInfoContainer}
+          style={yValueHoverSubCountsExists ? { display: 'flex' } : {}}
+        >
+          {calloutProps!.YValueHover &&
+            calloutProps!.YValueHover.map((yValue: IYValueHover, index: number, yValues: IYValueHover[]) => {
+              const isLast: boolean = index + 1 === yValues.length;
+              const { shouldDrawBorderBottom = false } = yValue;
+              return (
+                <div
+                  key={`callout-content-${index}`}
+                  style={
+                    yValueHoverSubCountsExists
+                      ? {
+                          display: 'inline-block',
+                          ...(shouldDrawBorderBottom && {
+                            borderBottom: `1px solid ${this.props.theme!.semanticColors.menuDivider}`,
+                            paddingBottom: '10px',
+                          }),
+                        }
+                      : {
+                          ...(shouldDrawBorderBottom && {
+                            borderBottom: `1px solid ${this.props.theme!.semanticColors.menuDivider}`,
+                            paddingBottom: '10px',
+                          }),
+                        }
+                  }
+                >
+                  {this._getCalloutContent(yValue, index, yValueHoverSubCountsExists, isLast)}
+                </div>
+              );
+            })}
+        </div>
+      </div>
+    );
+  };
+
+  private _yValueHoverSubCountsExists(yValueHover?: IYValueHover[]) {
+    if (yValueHover) {
+      return yValueHover.some(
+        (yValue: {
+          legend?: string;
+          y?: number;
+          color?: string;
+          yAxisCalloutData?: string | { [id: string]: number };
+        }) => yValue.yAxisCalloutData && typeof yValue.yAxisCalloutData !== 'string',
+      );
+    }
+    return false;
+  }
+
+  private _getCalloutContent(
+    xValue: IYValueHover,
+    index: number,
+    yValueHoverSubCountsExists: boolean,
+    isLast: boolean,
+  ): React.ReactNode {
+    const marginStyle: React.CSSProperties = isLast ? {} : { marginRight: '16px' };
+    const toDrawShape = xValue.index !== undefined && xValue.index !== -1;
+    const _classNames = getClassNames(this.props.styles!, {
+      theme: this.props.theme!,
+      width: this.state._width,
+      height: this.state._height,
+      className: this.props.className,
+      isRtl: this._isRtl,
+      lineColor: xValue.color,
+      toDrawShape,
+    });
+
+    if (!xValue.yAxisCalloutData || typeof xValue.yAxisCalloutData === 'string') {
+      return (
+        <div style={yValueHoverSubCountsExists ? marginStyle : {}}>
+          {yValueHoverSubCountsExists && (
+            <div className="ms-fontWeight-semibold" style={{ fontSize: '12pt' }}>
+              {xValue.legend!} ({xValue.y})
+            </div>
+          )}
+          <div id={`${index}_${xValue.y}`} className={_classNames.calloutBlockContainer}>
+            {toDrawShape && (
+              <Shape
+                svgProps={{
+                  className: _classNames.shapeStyles,
+                }}
+                pathProps={{ fill: xValue.color }}
+                shape={Points[xValue.index! % Object.keys(pointTypes).length] as LegendShape}
+              />
+            )}
+            <div>
+              <div className={_classNames.calloutlegendText}> {xValue.legend}</div>
+              <div className={_classNames.calloutContentY}>
+                {xValue.yAxisCalloutData ? xValue.yAxisCalloutData : xValue.y || xValue.data}
+              </div>
+            </div>
+          </div>
+        </div>
+      );
+    } else {
+      const subcounts: { [id: string]: number } = xValue.yAxisCalloutData as { [id: string]: number };
+      return (
+        <div style={marginStyle}>
+          <div className="ms-fontWeight-semibold" style={{ fontSize: '12pt' }}>
+            {xValue.legend!} ({xValue.y})
+          </div>
+          {Object.keys(subcounts).map((subcountName: string) => {
+            return (
+              <div key={subcountName} className={_classNames.calloutBlockContainer}>
+                <div className={_classNames.calloutlegendText}> {subcountName}</div>
+                <div className={_classNames.calloutContentY}>{subcounts[subcountName]}</div>
+              </div>
+            );
+          })}
+        </div>
+      );
+    }
+  }
+
+  private _fitParentContainer(): void {
+    const { containerWidth, containerHeight } = this.state;
+    this._reqID = requestAnimationFrame(() => {
+      let legendContainerHeight;
+      if (this.props.hideLegend) {
+        legendContainerHeight = 0;
+      } else {
+        const legendContainerComputedStyles = getComputedStyle(this.legendContainer!);
+        legendContainerHeight =
+          (this.legendContainer!.getBoundingClientRect().height || this.minLegendContainerHeight) +
+          parseFloat(legendContainerComputedStyles.marginTop || '0') +
+          parseFloat(legendContainerComputedStyles.marginBottom || '0');
+      }
+      const container = this.props.parentRef ? this.props.parentRef : this.chartContainer;
+      const currentContainerWidth = container!.getBoundingClientRect().width;
+      const currentContainerHeight =
+        container!.getBoundingClientRect().height > legendContainerHeight
+          ? container!.getBoundingClientRect().height
+          : 350;
+      const shouldResize =
+        containerWidth !== currentContainerWidth || containerHeight !== currentContainerHeight - legendContainerHeight;
+      if (shouldResize) {
+        this.setState({
+          containerWidth: currentContainerWidth,
+          containerHeight: currentContainerHeight - legendContainerHeight,
+        });
+      }
+    });
+  }
+
+  // eslint-disable-next-line @typescript-eslint/no-explicit-any
+  private _getData = (xScale: any, yScale: any) => {
+    this.props.getGraphData &&
+      this.props.getGraphData(
+        xScale,
+        yScale,
+        this.state.containerHeight - this.state._removalValueForTextTuncate!,
+        this.state.containerWidth,
+        this.xAxisElement,
+      );
+  };
+}
diff --git a/Engine/Source/Programs/Horde/HordeDashboard/src/base/components/HordeFluentCharting/HCartesianChart/HCartesianChart.tsx b/Engine/Source/Programs/Horde/HordeDashboard/src/base/components/HordeFluentCharting/HCartesianChart/HCartesianChart.tsx
new file mode 100644
--- /dev/null
+++ b/Engine/Source/Programs/Horde/HordeDashboard/src/base/components/HordeFluentCharting/HCartesianChart/HCartesianChart.tsx
@@ -0,0 +1,158 @@
+import * as React from 'react';
+import { styled } from '@fluentui/react/lib/Utilities';
+import { IHCartesianChartStyleProps, IHCartesianChartStyles, IHModifiedCartesianChartProps } from './HCartesianChart.types';
+import { HCartesianChartBase } from './HCartesianChart.base';
+import { HighContrastSelectorBlack, isIE11 } from '@fluentui/react';
+
+const isIE11Var: boolean = isIE11();
+
+export const getStyles = (props: IHCartesianChartStyleProps): IHCartesianChartStyles => {
+  const { className, theme, isRtl, shouldHighlight, href, lineColor = 'transparent', toDrawShape } = props;
+  const { fonts } = theme!;
+  return {
+    root: [
+      theme.fonts.medium,
+      {
+        display: 'flex',
+        width: '100%',
+        height: '100%',
+        flexDirection: 'column',
+        overflow: 'hidden',
+      },
+      className,
+    ],
+    xAxis: {
+      selectors: {
+        text: [
+          theme.fonts.tiny,
+          {
+            fill: theme.semanticColors.bodyText,
+            selectors: {
+              [HighContrastSelectorBlack]: {
+                fill: 'rgb(179, 179, 179)',
+              },
+            },
+          },
+        ],
+        line: {
+          opacity: 0.2,
+          stroke: theme.semanticColors.bodyText,
+          width: '1px',
+          selectors: {
+            [HighContrastSelectorBlack]: {
+              opacity: 0.1,
+              stroke: 'rgb(179, 179, 179)',
+            },
+          },
+        },
+        path: {
+          display: 'none',
+        },
+      },
+    },
+    yAxis: {
+      selectors: {
+        text: [
+          theme.fonts.tiny,
+          {
+            fill: theme.semanticColors.bodyText,
+            selectors: {
+              [HighContrastSelectorBlack]: {
+                fill: 'rgb(179, 179, 179)',
+              },
+            },
+          },
+        ],
+        path: {
+          display: 'none',
+        },
+        line: {
+          opacity: 0.2,
+          stroke: theme.semanticColors.bodyText,
+          selectors: {
+            [HighContrastSelectorBlack]: {
+              opacity: 0.1,
+              stroke: 'rgb(179, 179, 179)',
+            },
+          },
+        },
+        g: [
+          isRtl &&
+            !isIE11Var && {
+              textAnchor: 'end',
+            },
+        ],
+      },
+    },
+
+    opacityChangeOnHover: {
+      opacity: shouldHighlight ? '' : '0.1',
+      cursor: href ? 'pointer' : 'default',
+    },
+
+    legendContainer: [
+      {
+        marginTop: '8px',
+        marginLeft: '20px',
+      },
+    ],
+    calloutContentRoot: [
+      {
+        display: 'grid',
+        overflow: 'hidden',
+        padding: '11px 16px 10px 16px',
+        backgroundColor: theme.semanticColors.bodyBackground,
+        backgroundBlendMode: 'normal, luminosity',
+      },
+    ],
+    calloutDateTimeContainer: {
+      display: 'flex',
+      flexDirection: 'row',
+      justifyContent: 'space-between',
+    },
+    calloutContentX: [
+      {
+        ...fonts.small,
+        lineHeight: '16px',
+        opacity: '0.8',
+        color: theme.semanticColors.bodySubtext,
+      },
+    ],
+    calloutBlockContainer: [
+      theme.fonts.mediumPlus,
+      {
+        marginTop: '13px',
+        color: theme.semanticColors.bodyText,
+      },
+      !toDrawShape && {
+        borderLeft: `4px solid ${lineColor}`,
+        paddingLeft: '8px',
+      },
+      toDrawShape && {
+        display: 'flex',
+      },
+    ],
+    shapeStyles: {
+      marginRight: '8px',
+    },
+    calloutlegendText: {
+      ...fonts.small,
+      lineHeight: '16px',
+      color: theme.semanticColors.bodySubtext,
+    },
+    calloutContentY: [
+      {
+        ...fonts.mediumPlus,
+        fontWeight: 'bold',
+        lineHeight: '22px',
+      },
+    ],
+  };
+};
+
+export const HCartesianChart: React.FunctionComponent<IHModifiedCartesianChartProps> = styled<
+  IHModifiedCartesianChartProps,
+  IHCartesianChartStyleProps,
+  IHCartesianChartStyles
+>(HCartesianChartBase, getStyles);
+
diff --git a/Engine/Source/Programs/Horde/HordeDashboard/src/base/components/HordeFluentCharting/HCartesianChart/HCartesianChart.types.ts b/Engine/Source/Programs/Horde/HordeDashboard/src/base/components/HordeFluentCharting/HCartesianChart/HCartesianChart.types.ts
new file mode 100644
--- /dev/null
+++ b/Engine/Source/Programs/Horde/HordeDashboard/src/base/components/HordeFluentCharting/HCartesianChart/HCartesianChart.types.ts
@@ -0,0 +1,446 @@
+import * as React from 'react';
+import { IStyleFunctionOrObject } from '@fluentui/react/lib/Utilities';
+import { ITheme, IStyle } from '@fluentui/react/lib/Styling';
+import { IOverflowSetProps } from '@fluentui/react/lib/OverflowSet';
+import { IFocusZoneProps, FocusZoneDirection } from '@fluentui/react-focus';
+import { ICalloutProps } from '@fluentui/react/lib/Callout';
+import { ICartesianChartStyleProps, ICartesianChartStyles, IChartHoverCardProps, IChildProps, ILegendsProps, IMargins, IYValueHover } from '@fluentui/react-charting';
+import { ChartTypes, XAxisTypes, YAxisType } from '@fluentui/react-charting/lib/utilities/utilities';
+
+export interface IHCartesianChartStyleProps {
+  /**
+   * Theme (provided through customization.)
+   */
+  theme: ITheme;
+
+  /**
+   * Additional CSS class(es) to apply to the Chart.
+   */
+  className?: string;
+
+  /**
+   * Width of the chart.
+   */
+  width?: number;
+
+  /**
+   * Height of the chart.
+   */
+  height?: number;
+
+  /**
+   * Color of the chart.
+   */
+  color?: string;
+
+  /**
+   * Link to redirect if click action for graph
+   */
+  href?: string;
+
+  /**
+   * prop to check if the chart is selcted or hovered upon to determine opacity
+   */
+  shouldHighlight?: boolean;
+
+  /**
+   * prop to check if the Page is in Rtl
+   */
+  isRtl?: boolean;
+
+  /**
+   * color of the line
+   */
+  lineColor?: string;
+
+  /**
+   * boolean flag which determines if shape is drawn in callout
+   */
+  toDrawShape?: boolean;
+}
+
+export interface IHCartesianChartStyles {
+  /**
+   *  Style for the root element.
+   */
+  root?: IStyle;
+
+  /**
+   * Style for the element containing the x-axis.
+   */
+  xAxis?: IStyle;
+
+  /**
+   * Style for the element containing the y-axis.
+   */
+  yAxis?: IStyle;
+
+  /**
+   * Style for legend container
+   */
+  legendContainer?: IStyle;
+
+  /**
+   * line hover box css
+   */
+  hover?: IStyle;
+
+  /**
+   * styles for callout root-content
+   */
+  calloutContentRoot?: IStyle;
+
+  /**
+   * styles for callout x-content
+   */
+  calloutContentX?: IStyle;
+
+  /**
+   * styles for callout y-content
+   */
+  calloutContentY?: IStyle;
+
+  /**
+   * styles for callout Date time container
+   */
+  calloutDateTimeContainer?: IStyle;
+
+  /**
+   * styles for callout Date time container
+   */
+  calloutInfoContainer?: IStyle;
+
+  /**
+   * styles for callout Date time container
+   */
+  calloutBlockContainer?: IStyle;
+
+  /**
+   * styles for callout y-content
+   */
+  calloutlegendText?: IStyle;
+
+  tooltip?: IStyle;
+
+  /**
+   * Style for the chart Title.
+   */
+  chartTitle?: IStyle;
+
+  /**
+   * Style to change the opacity of bars in dataviz when we hover on a single bar or legends
+   */
+  opacityChangeOnHover?: IStyle;
+
+  /**
+   * styles set for the shape object in the callout
+   */
+  shapeStyles?: IStyle;
+}
+
+export interface IHCartesianChartProps {
+  /**
+   * Below height used for resizing of the chart
+   * Wrap chart in your container and send the updated height and width to these props.
+   * These values decide wheather chart re render or not. Please check examples for reference
+   */
+  height?: number;
+
+  /**
+   * Below width used for resizing of the chart
+   * Wrap chart in your container and send the updated height and width to these props.
+   * These values decide wheather chart re render or not. Please check examples for reference
+   */
+  width?: number;
+
+  /**
+   * this prop takes its parent as a HTML element to define the width and height of the chart
+   */
+  parentRef?: HTMLElement | null;
+
+  /**
+   * Additional CSS class(es) to apply to the Chart.
+   */
+  className?: string;
+
+  /**
+   * Theme (provided through customization.)
+   */
+  theme?: ITheme;
+
+  /**
+   * Margins for the chart
+   * @default `{ top: 20, bottom: 35, left: 40, right: 20 }`
+   * To avoid edge cuttings to the chart, we recommend you use default values or greater then default values
+   */
+  margins?: IMargins;
+
+  /** decides wether to show/hide legends
+   * @defaultvalue false
+   */
+  hideLegend?: boolean;
+
+  /**
+   * Do not show tooltips in chart
+   * @default false
+   */
+  hideTooltip?: boolean;
+
+  /**
+   * this prop takes values that you want the chart to render on x-axis
+   * This is a optional parameter if not specified D3 will decide which values appear on the x-axis for you
+   * Please look at [REDACTED_PROFILE_URL]/d3-scale for more information on how D3 decides what data to appear on the axis of chart
+   */
+  tickValues?: number[] | Date[];
+
+  /**
+   * the format in for the data on x-axis. For date object this can be specified to your requirement. Eg: '%m/%d', '%d'
+   * Please look at https://www.npmjs.com/package/d3-time-format for all the formats supported
+   */
+  tickFormat?: any;
+
+  /**
+   * Width of line stroke
+   */
+  strokeWidth?: number;
+
+  /**
+   * x Axis labels tick padding
+   * @default 10
+   */
+  xAxisTickPadding?: number;
+
+  /**
+   * the format in for the data on y-axis. For data object this can be specified to your requirement.
+   *  Eg: d3.format(".0%")(0.123),d3.format("+20")(42);
+   * Please look at [REDACTED_PROFILE_URL]/d3-format for all the formats supported
+   */
+  // eslint-disable-next-line @typescript-eslint/no-explicit-any
+  yAxisTickFormat?: any;
+
+  /**
+   * minimum  data value point in y-axis
+   */
+  yMinValue?: number;
+
+  /**
+   * maximum data value point in y-axis
+   */
+  yMaxValue?: number;
+
+  /**
+   * Number of ticks on the y-axis.
+   * Tick count should be factor of difference between (yMinValue, yMaxValue)
+   * @default 4
+   */
+  yAxisTickCount?: number;
+
+  /**
+   * defines the number of ticks on the x-axis
+   * @default 6
+   */
+  xAxisTickCount?: number;
+
+  /**
+   * define the size of the tick on the x-axis
+   * @default 10
+   */
+  xAxistickSize?: number;
+
+  /**
+   * define the space between the tick and the data point
+   * @default 10
+   */
+  tickPadding?: number;
+
+  /**
+   * Url that the data-viz needs to redirect to upon clicking on it
+   */
+  href?: string;
+
+  /**
+   * Label to apply to the whole chart.
+   * @deprecated - Use your chart label for the chart.
+   */
+  chartLabel?: string;
+
+  /**
+   * overflow props for the legends
+   */
+  legendsOverflowProps?: Partial<IOverflowSetProps>;
+
+  // eslint-disable-next-line @typescript-eslint/no-explicit-any
+  legendsOverflowText?: any;
+
+  /**
+   * focus zone props in hover card for legends
+   */
+  focusZonePropsForLegendsInHoverCard?: IFocusZoneProps;
+
+  /**
+   * Enable the legends to wrap lines if there is not enough space to show all legends on a single line
+   */
+  enabledLegendsWrapLines?: boolean;
+
+  /*
+   * props for the legends in the chart
+   */
+  legendProps?: Partial<ILegendsProps>;
+
+  /**
+   *@default false
+   *Used for to elipse x axis labes and show tooltip on x axis labels
+   */
+  showXAxisLablesTooltip?: boolean;
+
+  /**
+   * @default 4
+   * Used for X axis labels
+   * While Giving showXAxisLablesTooltip prop, need to define after how many chars, we need to truncate the word.
+   */
+  noOfCharsToTruncate?: number;
+
+  /**
+   * @default false
+   * Used to display x axis labels values (whole value)
+   */
+  wrapXAxisLables?: boolean;
+
+  /**
+   * Call to provide customized styling that will layer on top of the variant rules.
+   */
+  styles?: IStyleFunctionOrObject<ICartesianChartStyleProps, ICartesianChartStyles>;
+
+  /**
+   * Callout customization props
+   */
+  calloutProps?: Partial<ICalloutProps>;
+
+  /**
+   * props for the svg; use this to include aria-* or other attributes on the tag
+   */
+  svgProps?: React.SVGProps<SVGSVGElement>;
+}
+
+// Only used for Cartesian chart base
+export interface IHModifiedCartesianChartProps extends IHCartesianChartProps {
+  /**
+   * Only used for Area chart
+   * Value used to draw y axis of that chart.
+   */
+  maxOfYVal?: number;
+
+  /**
+   * Data of the chart
+   */
+  // eslint-disable-next-line @typescript-eslint/no-explicit-any
+  points: any;
+
+  /**
+   * Define type of the chart
+   */
+  chartType: ChartTypes;
+
+  /** X axis type */
+  xAxisType: XAxisTypes;
+
+  /** Y axis type */
+  yAxisType?: YAxisType;
+
+  /**
+   * Legeds of the chart.
+   */
+  legendBars: JSX.Element;
+
+  /**
+   * Callout props
+   */
+  calloutProps: Partial<ICalloutProps> & {
+    isCalloutVisible: boolean;
+    id: string;
+    YValueHover?: IYValueHover[];
+    hoverXValue?: string | number | null;
+    legend?: string;
+    color?: string;
+    YValue?: string | number;
+    XValue?: string;
+  };
+
+  /**
+   * Callback method used for to get margins to the chart.
+   */
+  getmargins?: () => IMargins | undefined;
+
+  /**
+   * This is a call back method to the chart from cartesian chart.
+   * params are xScale, yScale, containerHeight, containerWidth. These values were used to draw the graph.
+   * It also contians an optional param xAxisElement - defines as x axis scale element.
+   * This param used to enable feature word wrap of Xaxis.
+   */
+  // eslint-disable-next-line @typescript-eslint/no-explicit-any
+  getGraphData?: any;
+
+  /**
+   * Used for bar chart graphs.
+   * To define width of the bar
+   */
+  barwidth?: number;
+
+  /**
+   * Used for tick styles of the x axis of the chart
+   */
+  tickParams?: {
+    tickValues?: number[] | Date[];
+    tickFormat?: any;
+  };
+
+  /**
+   * it's padding between bar's or lines in the graph
+   */
+  xAxisPadding?: number;
+
+  /**
+   * it's padding between bar's or lines in the graph
+   */
+  yAxisPadding?: number;
+
+  children(props: IChildProps): React.ReactNode;
+
+  /**
+   * To enable callout for bar. Using for only Vertical stacked bar chart.
+   * @default false
+   * @type {boolean}
+   */
+  isCalloutForStack?: boolean;
+
+  /** dataset values to find out domain of the String axis
+   * Present using for only vertical stacked bar chart and grouped vertical bar chart
+   */
+  datasetForXAxisDomain?: string[];
+
+  /** Own callout design */
+  // eslint-disable-next-line @typescript-eslint/no-explicit-any
+  customizedCallout?: any;
+
+  /**
+   * if the data points for the y-axis is of type string, then we need to give this
+   * prop to construct the y-axis
+   */
+  stringDatasetForYAxisDomain?: string[];
+
+  /**
+   * Focus zone direction to the chart
+   * @default FocusZoneDirection.horizontal
+   */
+  focusZoneDirection?: FocusZoneDirection;
+
+  /**
+   * props to send into the chart hover card
+   */
+  chartHoverProps?: IChartHoverCardProps;
+
+  /**
+   * props to send to the focuszone
+   */
+  svgFocusZoneProps?: IFocusZoneProps;
+}
\ No newline at end of file
diff --git a/Engine/Source/Programs/Horde/HordeDashboard/src/base/components/HordeFluentCharting/HHoverCard/HHoverCard.base.tsx b/Engine/Source/Programs/Horde/HordeDashboard/src/base/components/HordeFluentCharting/HHoverCard/HHoverCard.base.tsx
new file mode 100644
--- /dev/null
+++ b/Engine/Source/Programs/Horde/HordeDashboard/src/base/components/HordeFluentCharting/HHoverCard/HHoverCard.base.tsx
@@ -0,0 +1,83 @@
+import * as React from 'react';
+import { IHChartHoverCardStyles, IHChartHoverCardStyleProps, IHChartHoverCardProps } from './HHoverCard.types';
+import { classNamesFunction, IProcessedStyleSet } from '@fluentui/react';
+
+const getClassNames = classNamesFunction<IHChartHoverCardStyleProps, IHChartHoverCardStyles>();
+export class HChartHoverCardBase extends React.Component<IHChartHoverCardProps, {}> {
+   private _classNames: IProcessedStyleSet<IHChartHoverCardStyles> | undefined;
+   public render(): React.ReactNode {
+      const { props1, dataPoint1, props2, dataPoint2, styles, theme, ratio, descriptionMessage } = this.props;
+      this._classNames = getClassNames(styles!, {
+         theme: theme!,
+         color: dataPoint1.color,
+         XValue: dataPoint1.xAxisCalloutData,
+         X2Value: dataPoint2?.xAxisCalloutData,
+         isRatioPresent: !!ratio,
+      });
+
+      let calloutIsObject = typeof dataPoint1.yAxisCalloutData !== 'string';
+      let xValue = `${props1.hoverXValue!}${props2 ? ` => ${props2.hoverXValue}` : ''}`;
+
+      let legendObjects: any[] = [];
+      if (!calloutIsObject) {
+         legendObjects.push(
+            <div className={this._classNames.calloutBlockContainer}>
+               <div className={this._classNames.calloutlegendText}>
+                  {dataPoint1.legend}
+               </div>
+               <div className={this._classNames.calloutContentY}>
+                  {dataPoint1.yAxisCalloutData as string}
+                  {!!props2 && !!dataPoint2 && ` => ${dataPoint2.yAxisCalloutData}`}
+               </div>
+            </div>
+         );
+      }
+      else {
+         for (const [key, val] of Object.entries(dataPoint1.yAxisCalloutData as { [id: string]: number })) {
+            let yAxisCalloutPoint2 = undefined;
+            if (props2 && dataPoint2) {
+               yAxisCalloutPoint2 = (dataPoint2.yAxisCalloutData as { [id: string]: number })[key];
+            }
+            legendObjects.push(
+               <div className={this._classNames.calloutBlockContainer}>
+                  <div className={this._classNames.calloutlegendText}>
+                     {key}
+                  </div>
+                  <div className={this._classNames.calloutContentY}>
+                     {val}
+                     {!!yAxisCalloutPoint2 && ` => ${yAxisCalloutPoint2}`}
+                  </div>
+               </div>
+            );
+         }
+      }
+
+      return (
+         <div className={this._classNames.calloutContentRoot}>
+            <div className={this._classNames.calloutDateTimeContainer}>
+               <div className={this._classNames.calloutContentX}>
+                  {xValue}
+               </div>
+            </div>
+            <div className={this._classNames.calloutInfoContainer}>
+               {legendObjects.map((object, index) => {
+                  return object;
+               })}
+               {!!ratio && (
+                  <div className={this._classNames.ratio}>
+                     <>
+                        <span className={this._classNames.numerator}>{ratio[0]}</span>/
+                        <span className={this._classNames.denominator}>{ratio[1]}</span>
+                     </>
+                  </div>
+               )}
+            </div>
+            {!!descriptionMessage &&
+               <div className={this._classNames.descriptionMessage}>
+                  {descriptionMessage}
+               </div>
+            }
+         </div>
+      );
+   }
+}
diff --git a/Engine/Source/Programs/Horde/HordeDashboard/src/base/components/HordeFluentCharting/HHoverCard/HHoverCard.styles.ts b/Engine/Source/Programs/Horde/HordeDashboard/src/base/components/HordeFluentCharting/HHoverCard/HHoverCard.styles.ts
new file mode 100644
--- /dev/null
+++ b/Engine/Source/Programs/Horde/HordeDashboard/src/base/components/HordeFluentCharting/HHoverCard/HHoverCard.styles.ts
@@ -0,0 +1,83 @@
+import { FontWeights } from '@fluentui/react/lib/Styling';
+import { IHChartHoverCardStyleProps, IHChartHoverCardStyles } from './HHoverCard.types';
+
+export const getChartHoverCardStyles = (props: IHChartHoverCardStyleProps): IHChartHoverCardStyles => {
+  const { color, XValue, theme, isRatioPresent = false } = props;
+  return {
+    calloutContentRoot: [
+      {
+        display: 'grid',
+        overflow: 'hidden',
+        padding: '11px 16px 10px 16px',
+        backgroundColor: theme.semanticColors.bodyBackground,
+        backgroundBlendMode: 'normal, luminosity',
+      },
+    ],
+    calloutDateTimeContainer: {
+      display: 'flex',
+      flexDirection: 'row',
+      justifyContent: 'space-between',
+    },
+    calloutContentX: [
+      theme.fonts.small,
+      {
+        lineHeight: '16px',
+        opacity: '0.8',
+        color: theme.semanticColors.bodySubtext,
+      },
+    ],
+    calloutBlockContainer: [
+      theme.fonts.xxLarge,
+      {
+        marginTop: XValue ? '13px' : 'unset',
+        paddingLeft: '8px',
+        lineHeight: '22px',
+        color: theme.semanticColors.bodyText,
+        borderLeft: `4px solid ${color}`,
+      },
+    ],
+    calloutlegendText: [
+      theme.fonts.small,
+      {
+        lineHeight: '16px',
+        color: theme.semanticColors.bodyText,
+      },
+    ],
+    calloutContentY: [
+      theme.fonts.xxLarge,
+      {
+        color: color ? color : theme.semanticColors.bodyText,
+        fontWeight: 'bold',
+        lineHeight: '36px',
+      },
+    ],
+    calloutInfoContainer: [
+      isRatioPresent && {
+        display: 'flex',
+        alignItems: 'flex-end',
+      },
+    ],
+    ratio: [
+      theme.fonts.small,
+      {
+        marginLeft: '6px',
+        color: theme.semanticColors.bodyText,
+      },
+    ],
+    numerator: {
+      fontWeight: FontWeights.bold,
+    },
+    denominator: {
+      fontWeight: FontWeights.semibold,
+    },
+    descriptionMessage: [
+      theme.fonts.small,
+      {
+        color: theme.semanticColors.bodyText,
+        marginTop: '10px',
+        paddingTop: '10px',
+        borderTop: `1px solid ${theme.semanticColors.menuDivider}`,
+      },
+    ],
+  };
+};
diff --git a/Engine/Source/Programs/Horde/HordeDashboard/src/base/components/HordeFluentCharting/HHoverCard/HHoverCard.tsx b/Engine/Source/Programs/Horde/HordeDashboard/src/base/components/HordeFluentCharting/HHoverCard/HHoverCard.tsx
new file mode 100644
--- /dev/null
+++ b/Engine/Source/Programs/Horde/HordeDashboard/src/base/components/HordeFluentCharting/HHoverCard/HHoverCard.tsx
@@ -0,0 +1,12 @@
+import * as React from 'react';
+import { HChartHoverCardBase } from './HHoverCard.base';
+import { IHChartHoverCardStyles, IHChartHoverCardStyleProps, IHChartHoverCardProps } from './HHoverCard.types';
+import { styled } from '@fluentui/react/lib/Utilities';
+import { getChartHoverCardStyles } from './HHoverCard.styles';
+
+//  Create a ChartHoverCard variant which uses these default styles and this styled subcomponent.
+export const HChartHoverCard: React.FunctionComponent<IHChartHoverCardProps> = styled<
+  IHChartHoverCardProps,
+  IHChartHoverCardStyleProps,
+  IHChartHoverCardStyles
+>(HChartHoverCardBase, getChartHoverCardStyles);
diff --git a/Engine/Source/Programs/Horde/HordeDashboard/src/base/components/HordeFluentCharting/HHoverCard/HHoverCard.types.ts b/Engine/Source/Programs/Horde/HordeDashboard/src/base/components/HordeFluentCharting/HHoverCard/HHoverCard.types.ts
new file mode 100644
--- /dev/null
+++ b/Engine/Source/Programs/Horde/HordeDashboard/src/base/components/HordeFluentCharting/HHoverCard/HHoverCard.types.ts
@@ -0,0 +1,109 @@
+import { ITheme, IStyle, IStyleFunctionOrObject } from '@fluentui/react';
+import { ICustomizedCalloutDataPoint } from '@fluentui/react-charting';
+import { IHCustomizedCalloutData } from '../HLineChart/HLineChart.types';
+export interface IHChartHoverCardProps {
+  props1: IHCustomizedCalloutData;
+  dataPoint1: ICustomizedCalloutDataPoint;
+
+  props2: IHCustomizedCalloutData | undefined;
+  dataPoint2: ICustomizedCalloutDataPoint | undefined;
+
+  includeWait?: boolean;
+  /**
+   * ratio to show
+   * first number is numerator
+   * and second number is denominator
+   */
+  ratio?: [number, number];
+  /**
+   * Theme (provided through customization.)
+   */
+  theme?: ITheme;
+  /**
+   * description message in the callout
+   */
+  descriptionMessage?: string;
+  /**
+   * Call to provide customized styling that will layer on top of the variant rules.
+   */
+  styles?: IStyleFunctionOrObject<IHChartHoverCardStyleProps, IHChartHoverCardStyles>;
+}
+
+export interface IHChartHoverCardStyles {
+  /**
+   * styles for callout root-content
+   */
+  calloutContentRoot?: IStyle;
+
+  /**
+   * styles for callout Date time container
+   */
+  calloutDateTimeContainer?: IStyle;
+
+  /**
+   * styles for callout Date time container
+   */
+  calloutInfoContainer?: IStyle;
+
+  /**
+   * styles for callout Date time container
+   */
+  calloutBlockContainer?: IStyle;
+
+  /**
+   * styles for callout y-content
+   */
+  calloutlegendText?: IStyle;
+
+  /**
+   * styles for callout x-content
+   */
+  calloutContentX?: IStyle;
+  /**
+   * styles for callout y-content
+   */
+  calloutContentY?: IStyle;
+  /**
+   * styles for denomination
+   */
+  ratio?: IStyle;
+  /**
+   * styles for numerator
+   */
+  numerator?: IStyle;
+  /**
+   * styles for denominator
+   */
+  denominator?: IStyle;
+  /**
+   * styles for the description
+   */
+  descriptionMessage?: IStyle;
+}
+
+export interface IHChartHoverCardStyleProps {
+  /**
+   * Theme (provided through customization.)
+   */
+  theme: ITheme;
+
+  /**
+   * color for hover card
+   */
+  color?: string;
+
+  /**
+   * X  value for hover card
+   */
+  XValue?: string;
+
+  /**
+   * X2  value for hover card
+   */
+   X2Value?: string;
+
+  /**
+   * indicate if denomination is present
+   */
+  isRatioPresent?: boolean;
+}
diff --git a/Engine/Source/Programs/Horde/HordeDashboard/src/base/components/HordeFluentCharting/HLineChart/HLineChart.base.tsx b/Engine/Source/Programs/Horde/HordeDashboard/src/base/components/HordeFluentCharting/HLineChart/HLineChart.base.tsx
new file mode 100644
--- /dev/null
+++ b/Engine/Source/Programs/Horde/HordeDashboard/src/base/components/HordeFluentCharting/HLineChart/HLineChart.base.tsx
@@ -0,0 +1,953 @@
+import * as React from 'react';
+import { select as d3Select } from 'd3-selection';
+import { classNamesFunction, getId, find } from '@fluentui/react/lib/Utilities';
+import { DirectionalHint } from '@fluentui/react/lib/Callout';
+import { IBasestate, IChildProps, IColorFillBarsProps, ILegend, ILineChartPoints, ILineChartStyleProps, ILineChartStyles, IRefArrayData, Legends } from '@fluentui/react-charting';
+import { calloutData, ChartTypes, getXAxisType, NumericAxis, Points, pointTypes, tooltipOfXAxislabels, XAxisTypes } from '@fluentui/react-charting/lib/utilities/utilities';
+import { IHCustomizedCalloutData, IHLineChartProps } from './HLineChart.types';
+import { HCartesianChart } from '../HCartesianChart/HCartesianChart';
+import { IMargins } from '@fluentui/react-charting/lib/LineChart';
+
+const getClassNames = classNamesFunction<ILineChartStyleProps, ILineChartStyles>();
+
+enum PointSize {
+  hoverSize = 11,
+  invisibleSize = 1,
+  normalVisibleSize = 6,
+}
+
+/**
+ *
+ * @param x units from origin
+ * @param y units from origin
+ * @param w is the legnth of the each side of a shape
+ * @param index index to get the shape path
+ */
+const _getPointPath = (x: number, y: number, w: number, index: number): string => {
+  const allPointPaths = [
+    // circle path
+    `M${x - w / 2} ${y}
+     A${w / 2} ${w / 2} 0 1 0 ${x + w / 2} ${y}
+     M${x - w / 2} ${y}
+     A ${w / 2} ${w / 2} 0 1 1 ${x + w / 2} ${y}
+     `,
+    //square
+    `M${x - w / 2} ${y - w / 2}
+     L${x + w / 2} ${y - w / 2}
+     L${x + w / 2} ${y + w / 2}
+     L${x - w / 2} ${y + w / 2}
+     Z`,
+    //triangle
+    `M${x - w / 2} ${y - 0.2886 * w}
+     H ${x + w / 2}
+     L${x} ${y + 0.5774 * w} Z`,
+    //diamond
+    `M${x} ${y - w / 2}
+     L${x + w / 2} ${y}
+     L${x} ${y + w / 2}
+     L${x - w / 2} ${y}
+     Z`,
+    //pyramid
+    `M${x} ${y - 0.5774 * w}
+     L${x + w / 2} ${y + 0.2886 * w}
+     L${x - w / 2} ${y + 0.2886 * w} Z`,
+    //hexagon
+    `M${x - 0.5 * w} ${y - 0.866 * w}
+     L${x + 0.5 * w} ${y - 0.866 * w}
+     L${x + w} ${y}
+     L${x + 0.5 * w} ${y + 0.866 * w}
+     L${x - 0.5 * w} ${y + 0.866 * w}
+     L${x - w} ${y}
+     Z`,
+    //pentagon
+    `M${x} ${y - 0.851 * w}
+     L${x + 0.6884 * w} ${y - 0.2633 * w}
+     L${x + 0.5001 * w} ${y + 0.6884 * w}
+     L${x - 0.5001 * w} ${y + 0.6884 * w}
+     L${x - 0.6884 * w} ${y - 0.2633 * w}
+     Z`,
+    //octagon
+    `M${x - 0.5001 * w} ${y - 1.207 * w}
+     L${x + 0.5001 * w} ${y - 1.207 * w}
+     L${x + 1.207 * w} ${y - 0.5001 * w}
+     L${x + 1.207 * w} ${y + 0.5001 * w}
+     L${x + 0.5001 * w} ${y + 1.207 * w}
+     L${x - 0.5001 * w} ${y + 1.207 * w}
+     L${x - 1.207 * w} ${y + 0.5001 * w}
+     L${x - 1.207 * w} ${y - 0.5001 * w}
+     Z`,
+  ];
+  return allPointPaths[index];
+};
+
+type LineChartDataWithIndex = ILineChartPoints & { index: number };
+
+export interface IHLineChartState extends IBasestate {
+  // This array contains data of selected legends for points
+  selectedLegendPoints: LineChartDataWithIndex[];
+  // This array contains data of selected legends for color bars
+  selectedColorBarLegend: IColorFillBarsProps[];
+  // This is a boolean value which is set to true
+  // when at least one legend is selected
+  isSelectedLegend: boolean;
+  // This value will be used as customized callout props - point callout.
+  dataPointCalloutProps?: IHCustomizedCalloutData;
+  dataPoint2CalloutProps?: IHCustomizedCalloutData;
+  // This value will be used as Customized callout props - For stack callout.
+  stackCalloutProps?: IHCustomizedCalloutData;
+  stack2CalloutProps?: IHCustomizedCalloutData;
+  
+  // active or hoverd point
+  activePoint?: string;
+
+  hoverX2Value?: string | null;
+  activePoint2?: string;
+  YValueHover2?: {
+    legend?: string;
+    y?: number;
+    color?: string;
+  }[];
+}
+
+export class HLineChartBase extends React.Component<IHLineChartProps, IHLineChartState> {
+  private _points: LineChartDataWithIndex[];
+  // eslint-disable-next-line @typescript-eslint/no-explicit-any
+  private _calloutPoints: any[];
+  // eslint-disable-next-line @typescript-eslint/no-explicit-any
+  private _xAxisScale: any = '';
+  // eslint-disable-next-line @typescript-eslint/no-explicit-any
+  private _yAxisScale: any = '';
+  private _circleId: string;
+  private _lineId: string;
+  private _verticalLine: string;
+  private _verticalLine2: string;
+  private _colorFillBarPatternId: string;
+  private _uniqueCallOutID: string | undefined;
+  private _refArray: IRefArrayData[];
+  private margins: IMargins | undefined;
+  private eventLabelHeight: number = 36;
+  private lines: JSX.Element[] | undefined;
+  private _renderedColorFillBars: JSX.Element[] | undefined;
+  private _colorFillBars: IColorFillBarsProps[];
+  private _colorFillBarsOpacity: number;
+  private _tooltipId: string;
+
+  constructor(props: IHLineChartProps) {
+    super(props);
+    this.state = {
+      hoverXValue: '',
+      hoverX2Value: '',
+      activeLegend: '',
+      YValueHover: [],
+      YValueHover2: [],
+      refSelected: '',
+      selectedLegend: '',
+      isCalloutVisible: false,
+      selectedLegendPoints: [],
+      selectedColorBarLegend: [],
+      isSelectedLegend: false,
+      activePoint: '',
+    };
+    this._refArray = [];
+    this._points = this._injectIndexPropertyInLineChartData(this.props.data.lineChartData);
+    this._colorFillBars = [];
+    this._colorFillBarsOpacity = 0.4;
+    this._calloutPoints = calloutData(this._points) || [];
+    this._circleId = getId('circle');
+    this._lineId = getId('lineID');
+    this._verticalLine = getId('verticalLine');
+    this._verticalLine2 = getId('verticalLine2');
+    this._colorFillBarPatternId = getId('colorFillBarPattern');
+    this._tooltipId = getId('LineChartTooltipId_');
+    props.eventAnnotationProps &&
+      props.eventAnnotationProps.labelHeight &&
+      (this.eventLabelHeight = props.eventAnnotationProps.labelHeight);
+  }
+
+  public componentDidUpdate(prevProps: IHLineChartProps): void {
+    /** note that height and width are not used to resize or set as dimesions of the chart,
+     * fitParentContainer is responisble for setting the height and width or resizing of the svg/chart
+     */
+    if (
+      prevProps.height !== this.props.height ||
+      prevProps.width !== this.props.width ||
+      prevProps.data !== this.props.data
+    ) {
+      this._points = this._injectIndexPropertyInLineChartData(this.props.data.lineChartData);
+      this._calloutPoints = calloutData(this._points) || [];
+    }
+  }
+
+  public render(): JSX.Element {
+    const { tickValues, tickFormat, legendProps, margins } = this.props;
+    this._points = this._injectIndexPropertyInLineChartData(this.props.data.lineChartData);
+
+    const isXAxisDateType = getXAxisType(this._points);
+    let points = this._points;
+    if (legendProps && !!legendProps.canSelectMultipleLegends) {
+      points = this.state.selectedLegendPoints.length >= 1 ? this.state.selectedLegendPoints : this._points;
+      this._calloutPoints = calloutData(points);
+    }
+
+    if(margins) {
+      this.margins = margins;
+    }
+
+    const legendBars = this._createLegends(this._points!);
+    const calloutProps = {
+      isCalloutVisible: this.state.isCalloutVisible,
+      directionalHint: DirectionalHint.topAutoEdge,
+      YValueHover: this.state.YValueHover,
+      hoverXValue: this.state.hoverXValue,
+      id: `toolTip${this._uniqueCallOutID}`,
+      target: this.state.refSelected,
+      isBeakVisible: false,
+      gapSpace: 15,
+      onDismiss: this._closeCallout,
+      preventDismissOnEvent: () => true,
+      hidden: !(!this.props.hideTooltip && this.state.isCalloutVisible),
+      ...this.props.calloutProps,
+    };
+    const tickParams = {
+      tickValues: tickValues,
+      tickFormat: tickFormat,
+    };
+
+    return (
+      <HCartesianChart
+        {...this.props}
+        points={points}
+        chartType={ChartTypes.LineChart}
+        isCalloutForStack
+        calloutProps={calloutProps}
+        tickParams={tickParams}
+        legendBars={legendBars}
+        getmargins={this._getMargins}
+        getGraphData={this._initializeLineChartData}
+        xAxisType={isXAxisDateType ? XAxisTypes.DateAxis : XAxisTypes.NumericAxis}
+        customizedCallout={this._getCustomizedCallout()}
+        /* eslint-disable react/jsx-no-bind */
+        // eslint-disable-next-line react/no-children-prop
+        children={(props: IChildProps) => {
+          this._xAxisScale = props.xScale!;
+          this._yAxisScale = props.yScale!;
+          return (
+            <>
+              <g>
+                <line
+                  x1={0}
+                  y1={0}
+                  x2={0}
+                  y2={props.containerHeight}
+                  stroke={'steelblue'}
+                  id={this._verticalLine}
+                  visibility={'hidden'}
+                  strokeDasharray={'5,5'}
+                />
+                <line
+                  x1={0}
+                  y1={0}
+                  x2={0}
+                  y2={props.containerHeight}
+                  stroke={'steelblue'}
+                  id={this._verticalLine2}
+                  visibility={'hidden'}
+                  strokeDasharray={'5,5'}
+                />
+                <g>
+                  {this._renderedColorFillBars}
+                  {this.lines}
+                </g>
+                {/* {eventAnnotationProps && (
+                  <EventsAnnotation
+                    {...eventAnnotationProps}
+                    scale={props.xScale!}
+                    chartYTop={this.margins.top! + this.eventLabelHeight}
+                    chartYBottom={props.containerHeight! - 35}
+                  />
+                )} */}
+              </g>
+            </>
+          );
+        }}
+      />
+    );
+  }
+
+  private _injectIndexPropertyInLineChartData = (lineChartData?: ILineChartPoints[]): LineChartDataWithIndex[] | [] => {
+    const { allowMultipleShapesForPoints = false } = this.props;
+    return lineChartData
+      ? lineChartData.map((item: ILineChartPoints, index: number) => ({
+          ...item,
+          index: allowMultipleShapesForPoints ? index : -1,
+        }))
+      : [];
+  };
+
+  private _getCustomizedCallout = () => {
+    return this.props.onRenderCalloutPerStack
+      ? this.props.onRenderCalloutPerStack(
+          [this.state.stackCalloutProps, this.state.stack2CalloutProps],
+        )
+      : this.props.onRenderCalloutPerDataPoint
+      ? this.props.onRenderCalloutPerDataPoint(
+          [this.state.dataPointCalloutProps, this.state.dataPoint2CalloutProps]
+        )
+      : null;
+  };
+
+  private _getMargins = () => {
+    return this.margins;
+  };
+
+  private _initializeLineChartData = (
+    xScale: NumericAxis,
+    yScale: NumericAxis,
+    containerHeight: number,
+    containerWidth: number,
+    xElement: SVGElement | null,
+  ) => {
+    this._xAxisScale = xScale;
+    this._yAxisScale = yScale;
+    this._renderedColorFillBars = this.props.colorFillBars ? this._createColorFillBars(containerHeight) : [];
+    this.lines = this._createLines(xElement!);
+  };
+
+  private _handleSingleLegendSelectionAction = (lineChartItem: LineChartDataWithIndex | IColorFillBarsProps) => {
+    if (this.state.selectedLegend === lineChartItem.legend) {
+      this.setState({ selectedLegend: '', activeLegend: lineChartItem.legend });
+      this._handleLegendClick(lineChartItem, null);
+    } else {
+      this.setState({
+        selectedLegend: lineChartItem.legend,
+        activeLegend: lineChartItem.legend,
+      });
+      this._handleLegendClick(lineChartItem, lineChartItem.legend);
+    }
+  };
+
+  private _onHoverCardHide = () => {
+    this.setState({
+      selectedLegendPoints: [],
+      selectedColorBarLegend: [],
+      isSelectedLegend: false,
+    });
+  };
+
+  private _createLegends(data: LineChartDataWithIndex[]): JSX.Element {
+    const { legendProps, allowMultipleShapesForPoints = false } = this.props;
+    const isLegendMultiSelectEnabled = !!(legendProps && !!legendProps.canSelectMultipleLegends);
+    const legendDataItems = data.map((point: LineChartDataWithIndex) => {
+      const color: string = point.color;
+      // mapping data to the format Legends component needs
+      const legend: ILegend = {
+        title: point.legend!,
+        color: color,
+        action: () => {
+          if (isLegendMultiSelectEnabled) {
+            this._handleMultipleLineLegendSelectionAction(point);
+          } else {
+            this._handleSingleLegendSelectionAction(point);
+          }
+        },
+        onMouseOutAction: () => {
+          this.setState({ activeLegend: '' });
+        },
+        hoverAction: () => {
+          this.setState({ activeLegend: point.legend });
+        },
+        ...(allowMultipleShapesForPoints && {
+          shape: Points[point.index % Object.keys(pointTypes).length] as ILegend['shape'],
+        }),
+      };
+      return legend;
+    });
+
+    const colorFillBarsLegendDataItems = this.props.colorFillBars
+      ? this.props.colorFillBars.map((colorFillBar: IColorFillBarsProps, index: number) => {
+          const title = colorFillBar.legend;
+          const legend: ILegend = {
+            title,
+            color: colorFillBar.color,
+            action: () => {
+              if (isLegendMultiSelectEnabled) {
+                this._handleMultipleColorFillBarLegendSelectionAction(colorFillBar);
+              } else {
+                this._handleSingleLegendSelectionAction(colorFillBar);
+              }
+            },
+            onMouseOutAction: () => {
+              this.setState({ activeLegend: '' });
+            },
+            hoverAction: () => {
+              this.setState({ activeLegend: title });
+            },
+            opacity: this._colorFillBarsOpacity,
+            stripePattern: colorFillBar.applyPattern,
+          };
+          return legend;
+        })
+      : [];
+
+    const legends = (
+      <Legends
+        legends={[...legendDataItems, ...colorFillBarsLegendDataItems]}
+        enabledWrapLines={this.props.enabledLegendsWrapLines}
+        overflowProps={this.props.legendsOverflowProps}
+        focusZonePropsInHoverCard={this.props.focusZonePropsForLegendsInHoverCard}
+        overflowText={this.props.legendsOverflowText}
+        {...(isLegendMultiSelectEnabled && { onLegendHoverCardLeave: this._onHoverCardHide })}
+        {...this.props.legendProps}
+      />
+    );
+    return legends;
+  }
+
+  private _closeCallout = () => {
+    this.setState({
+      isCalloutVisible: false,
+    });
+  };
+
+  private _getBoxWidthOfShape = (pointId: string, pointIndex: number, isLastPoint: boolean) => {
+    const { allowMultipleShapesForPoints = false } = this.props;
+    const { activePoint, activePoint2 } = this.state;
+    if (allowMultipleShapesForPoints) {
+      if (activePoint === pointId || activePoint2 === pointId) {
+        return PointSize.hoverSize;
+      } else if (pointIndex === 1 || isLastPoint) {
+        return PointSize.normalVisibleSize;
+      } else {
+        return PointSize.normalVisibleSize;
+      }
+    } else {
+      if (activePoint === pointId || activePoint2 === pointId) {
+        return PointSize.hoverSize;
+      } else {
+        return PointSize.normalVisibleSize;
+      }
+    }
+  };
+
+  private _getPath = (
+    xPos: number,
+    yPos: number,
+    pointId: string,
+    pointIndex: number,
+    isLastPoint: boolean,
+    pointOftheLine: number,
+  ): string => {
+    const { allowMultipleShapesForPoints = false } = this.props;
+    let w = this._getBoxWidthOfShape(pointId, pointIndex, isLastPoint);
+    const index: number = allowMultipleShapesForPoints ? pointOftheLine % Object.keys(pointTypes).length : 0;
+    const widthRatio = pointTypes[index].widthRatio;
+    w = widthRatio > 1 ? w / widthRatio : w;
+
+    return _getPointPath(xPos, yPos, w, index);
+  };
+  private _getPointFill = (lineColor: string, pointId: string, pointIndex: number, isLastPoint: boolean) => {
+    const { activePoint, activePoint2 } = this.state;
+    const { theme, allowMultipleShapesForPoints = false } = this.props;
+    if (allowMultipleShapesForPoints) {
+      if (pointIndex === 1 || isLastPoint) {
+        if (activePoint === pointId || activePoint2 === pointId) {
+          return theme!.palette.white;
+        } else {
+          return lineColor;
+        }
+      } else {
+        if (activePoint === pointId || activePoint2 === pointId) {
+          return theme!.palette.white;
+        } else {
+          return lineColor;
+        }
+      }
+    } else {
+      if (activePoint === pointId || activePoint2 === pointId) {
+        return theme!.palette.white;
+      } else {
+        return lineColor;
+      }
+    }
+  };
+  private _createLines(xElement: SVGElement): JSX.Element[] {
+    const lines = [];
+    if (this.state.isSelectedLegend) {
+      this._points = this.state.selectedLegendPoints;
+    } else {
+      this._points = this._injectIndexPropertyInLineChartData(this.props.data.lineChartData);
+    }
+    for (let i = 0; i < this._points.length; i++) {
+      const legendVal: string = this._points[i].legend;
+      const lineColor: string = this._points[i].color;
+      const { activePoint, /*activePoint2*/ } = this.state;
+      const { theme } = this.props;
+      if (this._points[i].data.length === 1) {
+        const x1 = this._points[i].data[0].x;
+        const y1 = this._points[i].data[0].y;
+        const xAxisCalloutData = this._points[i].data[0].xAxisCalloutData;
+        const circleId = `${this._circleId}${i}`;
+        lines.push(
+          <circle
+            id={`${this._circleId}${i}`}
+            key={`${this._circleId}${i}`}
+            r={activePoint === circleId ? 5.5 : 3.5}
+            cx={this._xAxisScale(x1)}
+            cy={this._yAxisScale(y1)}
+            fill={activePoint === circleId ? theme!.palette.white : lineColor}
+            onMouseOver={this._handleHover.bind(this, x1, y1, xAxisCalloutData, circleId)}
+            onMouseMove={this._handleHover.bind(this, x1, y1, xAxisCalloutData, circleId)}
+            onMouseOut={this._handleMouseOut}
+            strokeWidth={activePoint === circleId ? 2 : 0}
+            stroke={activePoint === circleId ? lineColor : ''}
+          />,
+        );
+      }
+      for (let j = 1; j < this._points[i].data.length; j++) {
+        const lineId = `${this._lineId}${i}_${j}`;
+        const circleId = `${this._circleId}${i}_${j-1}`;
+        let circleId2 = `${this._circleId}${i}_${j}`;
+        if(j+1 === this._points[i].data.length) {
+          circleId2 = `${this._circleId}${i}_${j-1}${j}L`;
+        }
+        const x1 = this._points[i].data[j - 1].x;
+        const y1 = this._points[i].data[j - 1].y;
+        const x2 = this._points[i].data[j].x;
+        const y2 = this._points[i].data[j].y;
+        const xAxisCalloutData = this._points[i].data[j - 1].xAxisCalloutData;
+        const xAxisCalloutData2 = this._points[i].data[j].xAxisCalloutData;
+        let path = this._getPath(this._xAxisScale(x1), this._yAxisScale(y1), circleId, j, false, this._points[i].index);
+        if (this.state.activeLegend === legendVal || this.state.activeLegend === '' || this.state.isSelectedLegend) {
+          lines.push(
+            <line
+              id={lineId}
+              key={lineId}
+              x1={this._xAxisScale(x1)}
+              y1={this._yAxisScale(y1)}
+              x2={this._xAxisScale(x2)}
+              y2={this._yAxisScale(y2)}
+              strokeWidth={this.props.strokeWidth || 4}
+              ref={(e: SVGLineElement | null) => {
+                this._refCallback(e!, lineId);
+              }}
+              onMouseOver={(ev) => this._handleLineHover(x1, y1, x2, y2, xAxisCalloutData, xAxisCalloutData2, circleId, circleId2, ev)}
+              onMouseMove={(ev) => this._handleLineHover(x1, y1, x2, y2, xAxisCalloutData, xAxisCalloutData2, circleId, circleId2, ev)}
+              onMouseOut={this._handleMouseOut}
+              stroke={lineColor}
+              strokeLinecap={'round'}
+              opacity={1}
+              onClick={this._onLineClick.bind(this, this._points[i].onLineClick!)}
+            />,
+          );
+          lines.push(
+            <path
+              id={circleId}
+              key={circleId}
+              d={path}
+              data-is-focusable={i === 0 ? true : false}
+              onMouseOver={this._handleHover.bind(this, x1, y1, xAxisCalloutData, circleId)}
+              onMouseMove={this._handleHover.bind(this, x1, y1, xAxisCalloutData, circleId)}
+              onMouseOut={this._handleMouseOut}
+              onFocus={() => this._handleFocus(lineId, x1, y1, xAxisCalloutData, circleId)}
+              onBlur={this._handleMouseOut}
+              onClick={this._onDataPointClick.bind(this, this._points[i].data[j - 1].onDataPointClick!)}
+              opacity={1}
+              fill={this._getPointFill(lineColor, circleId, j, false)}
+              stroke={lineColor}
+              strokeWidth={2}
+            />,
+          );
+          if (j + 1 === this._points[i].data.length) {
+            const lastCircleId = `${circleId}${j}L`;
+            path = this._getPath(
+              this._xAxisScale(x2),
+              this._yAxisScale(y2),
+              lastCircleId,
+              j,
+              true,
+              this._points[i].index,
+            );
+            const lastCirlceXCallout = this._points[i].data[j].xAxisCalloutData;
+            lines.push(
+              <path
+                id={lastCircleId}
+                key={lastCircleId}
+                d={path}
+                data-is-focusable={i === 0 ? true : false}
+                onMouseOver={this._handleHover.bind(this, x2, y2, lastCirlceXCallout, lastCircleId)}
+                onMouseMove={this._handleHover.bind(this, x2, y2, lastCirlceXCallout, lastCircleId)}
+                onMouseOut={this._handleMouseOut}
+                onFocus={() => this._handleFocus(lineId, x2, y2, lastCirlceXCallout, lastCircleId)}
+                onBlur={this._handleMouseOut}
+                onClick={this._onDataPointClick.bind(this, this._points[i].data[j].onDataPointClick!)}
+                opacity={1}
+                fill={this._getPointFill(lineColor, lastCircleId, j, true)}
+                stroke={lineColor}
+                strokeWidth={2}
+              />,
+            );
+            /* eslint-enable react/jsx-no-bind */
+          }
+        } else {
+          lines.push(
+            <line
+              id={lineId}
+              key={lineId}
+              x1={this._xAxisScale(x1)}
+              y1={this._yAxisScale(y1)}
+              x2={this._xAxisScale(x2)}
+              y2={this._yAxisScale(y2)}
+              strokeWidth={this.props.strokeWidth || 4}
+              stroke={lineColor}
+              strokeLinecap={'round'}
+              opacity={0.1}
+            />,
+          );
+        }
+      }
+    }
+    const classNames = getClassNames(this.props.styles!, {
+      theme: this.props.theme!,
+    });
+    // Removing un wanted tooltip div from DOM, when prop not provided.
+    if (!this.props.showXAxisLablesTooltip) {
+      try {
+        document.getElementById(this._tooltipId) && document.getElementById(this._tooltipId)!.remove();
+        // eslint-disable-next-line no-empty
+      } catch (e) {}
+    }
+    // Used to display tooltip at x axis labels.
+    if (!this.props.wrapXAxisLables && this.props.showXAxisLablesTooltip) {
+      const xAxisElement = d3Select(xElement).call(this._xAxisScale);
+      try {
+        document.getElementById(this._tooltipId) && document.getElementById(this._tooltipId)!.remove();
+        // eslint-disable-next-line no-empty
+      } catch (e) {}
+      const tooltipProps = {
+        tooltipCls: classNames.tooltip!,
+        id: this._tooltipId,
+        xAxis: xAxisElement,
+      };
+      xAxisElement && tooltipOfXAxislabels(tooltipProps);
+    }
+    return lines;
+  }
+
+  private _createColorFillBars = (containerHeight: number) => {
+    const colorFillBars: JSX.Element[] = [];
+    if (this.state.isSelectedLegend) {
+      this._colorFillBars = this.state.selectedColorBarLegend;
+    } else {
+      this._colorFillBars = this.props.colorFillBars!;
+    }
+    for (let i = 0; i < this._colorFillBars.length; i++) {
+      const colorFillBar = this._colorFillBars[i];
+      const colorFillBarId = getId(colorFillBar.legend.replace(/\W/g, ''));
+
+      if (colorFillBar.applyPattern) {
+        // Using a pattern element because CSS was unable to render diagonal stripes for rect elements
+        colorFillBars.push(this._getStripePattern(colorFillBar.color, i));
+      }
+
+      for (let j = 0; j < colorFillBar.data.length; j++) {
+        const startX = colorFillBar.data[j].startX;
+        const endX = colorFillBar.data[j].endX;
+        const opacity =
+          this.state.activeLegend === colorFillBar.legend ||
+          this.state.activeLegend === '' ||
+          this.state.isSelectedLegend
+            ? this._colorFillBarsOpacity
+            : 0.1;
+        colorFillBars.push(
+          <rect
+            fill={colorFillBar.applyPattern ? `url(#${this._colorFillBarPatternId}${i})` : colorFillBar.color}
+            fillOpacity={opacity}
+            x={this._xAxisScale(startX)}
+            y={this._yAxisScale(0) - containerHeight}
+            width={Math.abs(this._xAxisScale(endX) - this._xAxisScale(startX))}
+            height={containerHeight}
+            key={`${colorFillBarId}${j}`}
+          />,
+        );
+      }
+    }
+    return colorFillBars;
+  };
+
+  private _getStripePattern = (color: string, id: number) => {
+    // This describes a tile pattern that resembles diagonal stripes
+    // For more information: https://developer.mozilla.org/en-US/docs/Web/SVG/Attribute/d
+    const stripePath = 'M-4,4 l8,-8 M0,16 l16,-16 M12,20 l8,-8';
+    return (
+      <pattern
+        id={`${this._colorFillBarPatternId}${id}`}
+        width={16}
+        height={16}
+        key={`${this._colorFillBarPatternId}${id}`}
+        patternUnits={'userSpaceOnUse'}
+      >
+        <path d={stripePath} stroke={color} strokeWidth={2} />
+      </pattern>
+    );
+  };
+
+  private _refCallback(element: SVGGElement, legendTitle: string): void {
+    this._refArray.push({ index: legendTitle, refElement: element });
+  }
+
+  private _handleFocus = (
+    lineId: string,
+    x: number | Date,
+    y: number,
+    xAxisCalloutData: string | undefined,
+    circleId: string,
+  ) => {
+    this._uniqueCallOutID = circleId;
+    const formattedData = x instanceof Date ? x.toLocaleDateString() : x;
+    const xVal = x instanceof Date ? x.getTime() : x;
+    const found = find(this._calloutPoints, (element: { x: string | number }) => element.x === xVal) as any;
+    const _this = this;
+    d3Select('#' + circleId).attr('aria-labelledby', `toolTip${this._uniqueCallOutID}`);
+    d3Select(`#${this._verticalLine}`)
+      .attr('transform', () => `translate(${_this._xAxisScale(x)}, 0)`)
+      .attr('visibility', 'visibility');
+    this._refArray.forEach((obj: IRefArrayData) => {
+      if (obj.index === lineId) {
+        let hoverXValue = xAxisCalloutData ? xAxisCalloutData : '' + formattedData;
+        found.hoverXValue = hoverXValue;
+        found.hoverYValue = y;
+        this.setState({
+          isCalloutVisible: true,
+          refSelected: obj.refElement,
+          hoverXValue: hoverXValue,
+          hoverX2Value: undefined,
+          hoverYValue: y,
+          YValueHover: found.values,
+          YValueHover2: undefined,
+          stackCalloutProps: found!,
+          stack2CalloutProps: undefined,
+          dataPointCalloutProps: found!,
+          dataPoint2CalloutProps: undefined,
+          activePoint: circleId,
+        });
+      }
+    });
+  };
+
+  private _handleHover = (
+    x: number | Date,
+    y: number,
+    xAxisCalloutData: string | undefined,
+    circleId: string,
+    mouseEvent: React.MouseEvent<SVGElement>,
+  ) => {
+    mouseEvent.persist();
+    this._uniqueCallOutID = circleId;
+    const formattedData = x instanceof Date ? x.toLocaleDateString() : x;
+    const xVal = x instanceof Date ? x.getTime() : x;
+    const _this = this;
+    d3Select(`#${this._verticalLine}`)
+      .attr('transform', () => `translate(${_this._xAxisScale(x)}, 0)`)
+      .attr('visibility', 'visibility');
+    const found = find(this._calloutPoints, (element: { x: string | number }) => element.x === xVal) as any;
+    let hoverXValue = xAxisCalloutData ? xAxisCalloutData : '' + formattedData;
+    found.hoverXValue = hoverXValue;
+    found.hoverYValue = y;
+    this.setState({
+      isCalloutVisible: true,
+      refSelected: mouseEvent,
+      hoverXValue: hoverXValue,
+      hoverX2Value: undefined,
+      hoverYValue: y,
+      YValueHover: found.values,
+      YValueHover2: undefined,
+      stackCalloutProps: found!,
+      stack2CalloutProps: undefined,
+      dataPointCalloutProps: found!,
+      dataPoint2CalloutProps: undefined,
+      activePoint: circleId,
+    });
+  };
+
+  private _handleLineHover = (
+    x1: number | Date,
+    y1: number,
+    x2: number | Date,
+    y2: number,
+    xAxisCalloutData: string | undefined,
+    xAxisCalloutData2: string | undefined,
+    circleId1: string,
+    circleId2: string,
+    mouseEvent: React.MouseEvent<SVGElement>,
+  ) => {
+    mouseEvent.persist();
+    this._uniqueCallOutID = `${circleId1}-${circleId2}`;
+    const formattedX1 = x1 instanceof Date ? x1.toLocaleDateString() : x1;
+    const x1Val = x1 instanceof Date ? x1.getTime() : x1;
+    const formattedX2 = x2 instanceof Date ? x2.toLocaleDateString() : x2;
+    const x2Val = x2 instanceof Date ? x2.getTime() : x2;
+
+    const _this = this;
+    d3Select(`#${this._verticalLine}`)
+      .attr('transform', () => `translate(${_this._xAxisScale(x1)}, 0)`)
+      .attr('visibility', 'visibility');
+    d3Select(`#${this._verticalLine2}`)
+      .attr('transform', () => `translate(${_this._xAxisScale(x2)}, 0)`)
+      .attr('visibility', 'visibility');
+    const found1 = find(this._calloutPoints, (element: { x: string | number }) => element.x === x1Val) as any;
+    const found2 = find(this._calloutPoints, (element: { x: string | number }) => element.x === x2Val) as any;
+    
+    let hoverXValue = xAxisCalloutData ? xAxisCalloutData : '' + formattedX1;
+    let hoverX2Value = xAxisCalloutData2 ? xAxisCalloutData2 : '' + formattedX2;
+    found1.hoverXValue = hoverXValue;
+    found1.hoverYValue = y1;
+    
+
+    found2.hoverXValue = hoverX2Value;
+    found2.hoverYValue = y2;
+
+    this.setState({
+      isCalloutVisible: true,
+      refSelected: mouseEvent,
+      hoverXValue: hoverXValue,
+      hoverX2Value: hoverX2Value,
+      hoverYValue: y1,
+      YValueHover: found1.values,
+      YValueHover2: found2.values,
+      stackCalloutProps: found1!,
+      dataPointCalloutProps: found1!,
+      stack2CalloutProps: found2!,
+      dataPoint2CalloutProps: found2!,
+      activePoint: circleId1,
+      activePoint2: circleId2,
+    });
+  };
+
+  private _onLineClick = (func: () => void) => {
+    if (func) {
+      func();
+    }
+  };
+
+  private _onDataPointClick = (func: () => void) => {
+    if (func) {
+      func();
+    }
+  };
+
+  private _handleMouseOut = () => {
+    d3Select(`#${this._verticalLine}`).attr('visibility', 'hidden');
+    d3Select(`#${this._verticalLine2}`).attr('visibility', 'hidden');
+    this.setState({
+      isCalloutVisible: false,
+      activePoint: '',
+      activePoint2: '',
+    });
+  };
+
+  private _handleLegendClick = (
+    lineChartItem: LineChartDataWithIndex | IColorFillBarsProps,
+    selectedLegend: string | null | string[],
+  ): void => {
+    if (lineChartItem.onLegendClick) {
+      lineChartItem.onLegendClick(selectedLegend);
+    }
+  };
+
+  private _handleMultipleLineLegendSelectionAction = (selectedLine: LineChartDataWithIndex) => {
+    const selectedLineIndex = this.state.selectedLegendPoints.reduce((acc, line, index) => {
+      if (acc > -1 || line.legend !== selectedLine.legend) {
+        return acc;
+      } else {
+        return index;
+      }
+    }, -1);
+
+    let selectedLines: LineChartDataWithIndex[];
+    if (selectedLineIndex === -1) {
+      selectedLines = [...this.state.selectedLegendPoints, selectedLine];
+    } else {
+      selectedLines = this.state.selectedLegendPoints
+        .slice(0, selectedLineIndex)
+        .concat(this.state.selectedLegendPoints.slice(selectedLineIndex + 1));
+    }
+
+    const areAllLineLegendsSelected = this.props.data && selectedLines.length === this.props.data.lineChartData!.length;
+
+    if (
+      areAllLineLegendsSelected &&
+      ((this.props.colorFillBars && this.props.colorFillBars.length === this.state.selectedColorBarLegend.length) ||
+        !this.props.colorFillBars)
+    ) {
+      // Clear all legends if all legends including color fill bar legends are selected
+      // Or clear all legends if all legends are selected and there are no color fill bars
+      this._clearMultipleLegendSelections();
+    } else if (!selectedLines.length && !this.state.selectedColorBarLegend.length) {
+      // Clear all legends if no legends including color fill bar legends are selected
+      this._clearMultipleLegendSelections();
+    } else {
+      // Otherwise, set state when one or more legends are selected, including color fill bar legends
+      this.setState({
+        selectedLegendPoints: selectedLines,
+        isSelectedLegend: true,
+      });
+    }
+
+    const selectedLegendTitlesToPass = selectedLines.map((line: LineChartDataWithIndex) => line.legend);
+    this._handleLegendClick(selectedLine, selectedLegendTitlesToPass);
+  };
+
+  private _handleMultipleColorFillBarLegendSelectionAction = (selectedColorFillBar: IColorFillBarsProps) => {
+    const selectedColorFillBarIndex = this.state.selectedColorBarLegend.reduce((acc, colorFillBar, index) => {
+      if (acc > -1 || colorFillBar.legend !== selectedColorFillBar.legend) {
+        return acc;
+      } else {
+        return index;
+      }
+    }, -1);
+
+    let selectedColorFillBars: IColorFillBarsProps[];
+    if (selectedColorFillBarIndex === -1) {
+      selectedColorFillBars = [...this.state.selectedColorBarLegend, selectedColorFillBar];
+    } else {
+      selectedColorFillBars = this.state.selectedColorBarLegend
+        .slice(0, selectedColorFillBarIndex)
+        .concat(this.state.selectedColorBarLegend.slice(selectedColorFillBarIndex + 1));
+    }
+
+    const areAllColorFillBarLegendsSelected =
+      selectedColorFillBars.length === (this.props.colorFillBars && this.props.colorFillBars!.length);
+
+    if (
+      areAllColorFillBarLegendsSelected &&
+      ((this.props.data && this.props.data.lineChartData!.length === this.state.selectedLegendPoints.length) ||
+        !this.props.data)
+    ) {
+      // Clear all legends if all legends, including line legends, are selected
+      // Or clear all legends if all legends are selected and there is no line data
+      this._clearMultipleLegendSelections();
+    } else if (!selectedColorFillBars.length && !this.state.selectedLegendPoints.length) {
+      // Clear all legends if no legends are selected, including line legends
+      this._clearMultipleLegendSelections();
+    } else {
+      // set state when one or more legends are selected, including line legends
+      this.setState({
+        selectedColorBarLegend: selectedColorFillBars,
+        isSelectedLegend: true,
+      });
+    }
+
+    const selectedLegendTitlesToPass = selectedColorFillBars.map(
+      (colorFillBar: IColorFillBarsProps) => colorFillBar.legend,
+    );
+    this._handleLegendClick(selectedColorFillBar, selectedLegendTitlesToPass);
+  };
+
+  private _clearMultipleLegendSelections = () => {
+    this.setState({
+      selectedColorBarLegend: [],
+      selectedLegendPoints: [],
+      isSelectedLegend: false,
+    });
+  };
+}
diff --git a/Engine/Source/Programs/Horde/HordeDashboard/src/base/components/HordeFluentCharting/HLineChart/HLineChart.tsx b/Engine/Source/Programs/Horde/HordeDashboard/src/base/components/HordeFluentCharting/HLineChart/HLineChart.tsx
new file mode 100644
--- /dev/null
+++ b/Engine/Source/Programs/Horde/HordeDashboard/src/base/components/HordeFluentCharting/HLineChart/HLineChart.tsx
@@ -0,0 +1,30 @@
+import * as React from 'react';
+import { styled } from '@fluentui/react/lib/Utilities';
+import { IHLineChartProps } from './HLineChart.types';
+import { HLineChartBase } from './HLineChart.base';
+import { ILineChartStyleProps, ILineChartStyles } from '@fluentui/react-charting';
+
+export const getStyles = (props: ILineChartStyleProps): ILineChartStyles => {
+  return {
+    tooltip: {
+      ...props.theme!.fonts.medium,
+      display: 'flex',
+      flexDirection: 'column',
+      padding: '8px',
+      position: 'absolute',
+      textAlign: 'center',
+      top: '0px',
+      background: props.theme!.semanticColors.bodyBackground,
+      borderRadius: '2px',
+      pointerEvents: 'none',
+    },
+  };
+};
+
+
+// Create a LineChart variant which uses these default styles and this styled subcomponent.
+export const HLineChart: React.FunctionComponent<IHLineChartProps> = styled<
+  IHLineChartProps,
+  ILineChartStyleProps,
+  ILineChartStyles
+>(HLineChartBase, getStyles);
diff --git a/Engine/Source/Programs/Horde/HordeDashboard/src/base/components/HordeFluentCharting/HLineChart/HLineChart.types.ts b/Engine/Source/Programs/Horde/HordeDashboard/src/base/components/HordeFluentCharting/HLineChart/HLineChart.types.ts
new file mode 100644
--- /dev/null
+++ b/Engine/Source/Programs/Horde/HordeDashboard/src/base/components/HordeFluentCharting/HLineChart/HLineChart.types.ts
@@ -0,0 +1,49 @@
+import { IRenderFunction, IStyleFunctionOrObject } from "@fluentui/react";
+import { ICartesianChartProps, IChartProps, IColorFillBarsProps, ICustomizedCalloutDataPoint, IEventsAnnotationProps, ILineChartStyleProps, ILineChartStyles } from "@fluentui/react-charting";
+
+/**
+ * Used for custom callout data interface. As Area chart callout data will be prepared from given props.data,
+ * Those required data passing to onRenderCalloutPerDataPoint and onRenderCalloutPerStack.
+ */
+ export interface IHCustomizedCalloutData {
+    x: number | string | Date;
+    values: ICustomizedCalloutDataPoint[];
+    hoverXValue?: string | number | null | undefined;
+    hoverYValue?: string | number | null | undefined;
+}
+
+export interface IHLineChartProps extends ICartesianChartProps {
+    /**
+     * Data to render in the chart.
+     */
+    data: IChartProps;
+    /**
+     * Call to provide customized styling that will layer on top of the variant rules.
+     */
+    styles?: IStyleFunctionOrObject<ILineChartStyleProps, ILineChartStyles>;
+    /**
+     * Show event annotation
+     */
+    eventAnnotationProps?: IEventsAnnotationProps;
+    /**
+     * Define a custom callout renderer for a data point
+     */
+    onRenderCalloutPerDataPoint?: IRenderFunction<(IHCustomizedCalloutData | undefined)[]>;
+    /**
+     * Define a custom callout renderer for a stack; default is to render per data point
+     */
+    onRenderCalloutPerStack?: IRenderFunction<(IHCustomizedCalloutData | undefined)[]>;
+    colorFillBars?: IColorFillBarsProps[];
+    /**
+     * if this is set to true, then for each line there will be a unique shape assigned to the point,
+     * there are total 8 shapes which are as follow circle, square, triangele, diamond, pyramid,
+     *  hexagon, pentagon and octagon, which will get assigned as respectively, if there are more
+     * than 8 lines in the line chart then it will again start from cicle to octagon.
+     * setting this flag to true will also change the behavior of the points, like for a
+     * line, last point shape and first point shape will be visible all the times, and all
+     * other points will get enlarge only when hovered over them
+     * if set to false default shape will be circle, with the existing behavior
+     * @default false
+     */
+    allowMultipleShapesForPoints?: boolean;
+}
\ No newline at end of file
diff --git a/Engine/Source/Programs/Horde/HordeDashboard/src/base/components/HordeFluentCharting/Hutilities.ts b/Engine/Source/Programs/Horde/HordeDashboard/src/base/components/HordeFluentCharting/Hutilities.ts
new file mode 100644
--- /dev/null
+++ b/Engine/Source/Programs/Horde/HordeDashboard/src/base/components/HordeFluentCharting/Hutilities.ts
@@ -0,0 +1,85 @@
+import { IXAxisParams, IYAxisParams, prepareDatapoints } from "@fluentui/react-charting/lib/utilities/utilities";
+import { axisBottom as d3AxisBottom, axisRight as d3AxisRight, axisLeft as d3AxisLeft } from 'd3-axis';
+import { scaleLinear as d3ScaleLinear } from 'd3-scale';
+import { select as d3Select } from 'd3-selection';
+import { format as d3Format } from 'd3-format';
+
+export interface IHTickParams {
+	tickValues?: Date[] | number[];
+	tickFormat?: any;
+  }
+
+  /**
+ * Create Numeric X axis
+ * @export
+ * @param {IXAxisParams} xAxisParams
+ */
+export function hCreateNumericXAxis(xAxisParams: IXAxisParams, tickParams?: IHTickParams) {
+	const {
+	  domainNRangeValues,
+	  showRoundOffXTickValues = false,
+	  xAxistickSize = 6,
+	  tickPadding = 10,
+	  xAxisCount = 6,
+	  xAxisElement,
+	} = xAxisParams;
+	const xAxisScale = d3ScaleLinear()
+	  .domain([domainNRangeValues.dStartValue, domainNRangeValues.dEndValue])
+	  .range([domainNRangeValues.rStartValue, domainNRangeValues.rEndValue]);
+	showRoundOffXTickValues && xAxisScale.nice();
+  
+	const xAxis = d3AxisBottom(xAxisScale)
+	  .tickSize(xAxistickSize)
+	  .tickPadding(tickPadding)
+	  .ticks(xAxisCount)
+	  .tickFormat(tickParams?.tickFormat || null)
+	  .tickSizeOuter(0);
+	if (xAxisElement) {
+	  d3Select(xAxisElement)
+		.call(xAxis as any)
+		.selectAll('text')
+		.attr('aria-hidden', 'true');
+	}
+	return xAxisScale;
+  }
+
+  export function hCreateYAxis(yAxisParams: IYAxisParams, isRtl: boolean) {
+	const {
+	  yMinMaxValues = { startValue: 0, endValue: 0 },
+	  yAxisElement = null,
+	  yMaxValue = 0,
+	  yMinValue = 0,
+	  containerHeight,
+	  containerWidth,
+	  margins,
+	  tickPadding = 12,
+	  maxOfYVal = 0,
+	  yAxisTickFormat,
+	  yAxisTickCount = 4,
+	  eventAnnotationProps,
+	  eventLabelHeight,
+	} = yAxisParams;
+  
+	// maxOfYVal coming from only area chart and Grouped vertical bar chart(Calculation done at base file)
+	const tempVal = maxOfYVal || yMinMaxValues.endValue;
+	const finalYmax = tempVal > yMaxValue ? tempVal : yMaxValue!;
+	const finalYmin = yMinMaxValues.startValue < yMinValue ? 0 : yMinValue!;
+	const domainValues = prepareDatapoints(finalYmax, finalYmin, yAxisTickCount);
+	const yAxisScale = d3ScaleLinear()
+	  .domain([finalYmin, domainValues[domainValues.length - 1]])
+	  .range([containerHeight - margins.bottom!, margins.top! + (eventAnnotationProps! ? eventLabelHeight! : 0)]);
+	const axis = isRtl ? d3AxisRight(yAxisScale) : d3AxisLeft(yAxisScale);
+	const yAxis = axis
+	  .tickPadding(tickPadding)
+	  .tickValues(domainValues)
+	  .tickSizeInner(-(containerWidth - margins.left! - margins.right!));
+	yAxisTickFormat ? yAxis.tickFormat(yAxisTickFormat) : yAxis.tickFormat(d3Format('.2~s'));
+	// eslint-disable-next-line @typescript-eslint/no-unused-expressions
+	yAxisElement
+	  ? d3Select(yAxisElement)
+		  .call(yAxis as any)
+		  .selectAll('text')
+		  .attr('aria-hidden', 'true')
+	  : '';
+	return yAxisScale;
+  }
\ No newline at end of file
diff --git a/Engine/Source/Programs/Horde/HordeDashboard/src/components/jobDetailsV2/JobDetailStepTrends.tsx b/Engine/Source/Programs/Horde/HordeDashboard/src/components/jobDetailsV2/JobDetailStepTrends.tsx
--- a/Engine/Source/Programs/Horde/HordeDashboard/src/components/jobDetailsV2/JobDetailStepTrends.tsx
+++ b/Engine/Source/Programs/Horde/HordeDashboard/src/components/jobDetailsV2/JobDetailStepTrends.tsx
@@ -1,7 +1,7 @@
 // Copyright Epic Games, Inc. All Rights Reserved.
 
 import { ContextualMenu, Dialog, DialogType, IconButton, IIconProps, Spinner, SpinnerSize, Stack, Text, Toggle } from '@fluentui/react';
-import { IChartProps, ICustomizedCalloutData, ILineChartDataPoint, LineChart } from '@fluentui/react-charting';
+import { IChartProps, ILineChartDataPoint } from '@fluentui/react-charting';
 import { observer } from 'mobx-react-lite';
 import moment from 'moment';
 import React, { useCallback, useEffect, useState } from 'react';
@@ -9,6 +9,9 @@ import backend from '../../backend';
 import { GetBatchResponse, GetJobTimingResponse, GetLabelResponse, GetLabelStateResponse, GetLabelTimingInfoResponse, GetNodeResponse, GetStepResponse, GetStepTimingInfoResponse, GetTemplateRefResponse, JobData, JobQuery, JobTimingsQuery, StepData } from '../../backend/Api';
 import { JobLabel } from '../../backend/JobDetails';
 import TemplateCache from '../../backend/TemplateCache';
+import { HChartHoverCard } from '../../base/components/HordeFluentCharting/HHoverCard/HHoverCard';
+import { HLineChart } from '../../base/components/HordeFluentCharting/HLineChart/HLineChart';
+import { IHCustomizedCalloutData } from '../../base/components/HordeFluentCharting/HLineChart/HLineChart.types';
 import { ISideRailLink } from '../../base/components/SideRail';
 import { msecToElapsed } from '../../base/utilities/timeUtils';
 import { hordeClasses } from '../../styles/Styles';
@@ -406,7 +409,7 @@ function preventChartScroll(this: Document, e: WheelEvent) {
    e.returnValue = false
 }
 
-export const StepTrendsPanel: React.FC<{ jobDetails: JobDetailsV2 }> = observer(({ jobDetails }) => {
+export const StepTrendsPanel: React.FC<{ jobDetails: JobDetailsV2}> = observer(({ jobDetails}) => {
    const query = useQuery();
 
    const stepId = query.get("step") ? query.get("step")! : undefined;
@@ -419,15 +422,15 @@ export const StepTrendsPanel: React.FC<{ jobDetails: JobDetailsV2 }> = observer(
          dataView?.clear();
       };
    }, [dataView]);
-
+   
    dataView.subscribe();
 
    if (jobDetails?.jobData) {
-      dataView.initialize(jobDetails.jobData?.preflightChange ? undefined : [sideRail]);
+      dataView.initialize(jobDetails.jobData?.preflightChange ? undefined : [sideRail]);       
    } else {
       jobDetails?.subscribe();
    }
-
+   
    if (jobDetails.jobData == null || !jobDetails?.viewsReady || jobDetails.jobData.preflightChange) {
       return null;
    }
@@ -1216,7 +1219,7 @@ const StepTrendsPanelInner: React.FC<{ jobDetails: JobDetailsV2; dataView: Trend
    let titleText = 'Trends';
    if (filteredChartState && filteredChartState?.chartProps) {
       if (filteredChartState.chartProps?.lineChartData?.length) {
-         placeholderDiv = <LineChart
+         placeholderDiv = <HLineChart
             data={filteredChartState.chartProps}
             hideLegend={false}
             legendProps={{ canSelectMultipleLegends: true, allowFocusOnLegends: false, enabledWrapLines: true }}
@@ -1225,16 +1228,21 @@ const StepTrendsPanelInner: React.FC<{ jobDetails: JobDetailsV2; dataView: Trend
             margins={{ top: 20, bottom: 35, left: filteredChartState.leftMargin, right: 45 }}
             yAxisTickCount={5}
             yMaxValue={filteredChartState.yMaxValue}
-            onRenderCalloutPerDataPoint={(props?: ICustomizedCalloutData) => {
-               // This is currently disabled, due to package updates and incompatibility with 
-               // custom LineChart component hoisted out
-               // of react charting, which had some customizations for the callout
-               // namely comparison with next data point, and more importantly a way to 
-               // tell which y data point the user is hovering over... without that
-               // all matching x's are rendered and no way to customize (a failing of the fluent charting package)
-               // Trends data and rendering need to be redone, and already spent too much time trying to figure out
-               // how to make this work without the customization on the component, so here we are:
-               return <div />
+            onRenderCalloutPerDataPoint={(props: (IHCustomizedCalloutData | undefined)[] | undefined) => {
+               if (!props || (!props[0] && !props[1])) {
+                  return null;
+               }
+
+               let hoverPoint1 = props[0];
+               let hoverPoint2 = props[1];
+
+               let dataPoint1 = hoverPoint1 ? hoverPoint1.values.find((point: { legend: string | number; y: string | number }) => point.y === hoverPoint1!.hoverYValue) : undefined;
+               let dataPoint2 = hoverPoint2 ? hoverPoint2.values.find((point: { legend: string | number; y: string | number }) => point.y === hoverPoint2!.hoverYValue) : undefined;
+
+               if (hoverPoint1 && dataPoint1) {
+                  return <HChartHoverCard props1={hoverPoint1} dataPoint1={dataPoint1} props2={hoverPoint2} dataPoint2={dataPoint2} />
+               }
+               return null;
             }}
          />
       }
