Horde: Horde fix dashboard version on modal with note
#rnx
#skipci

diff --git a/Engine/Source/Programs/Horde/HordeDashboard/src/components/VersionModal.tsx b/Engine/Source/Programs/Horde/HordeDashboard/src/components/VersionModal.tsx
--- a/Engine/Source/Programs/Horde/HordeDashboard/src/components/VersionModal.tsx
+++ b/Engine/Source/Programs/Horde/HordeDashboard/src/components/VersionModal.tsx
@@ -52,9 +52,12 @@ export const VersionModal: React.FC<{ show: boolean, onClose: () => void }> = ({
    let dashboardVersion: string | undefined;
 
    try {
-      dashboardVersion = process?.env?.REACT_APP_VERSION_INFO;
-   } catch {
-      
+      // note: this must be exactly `process.env.REACT_APP_VERSION_INFO`
+      // as webpack does a simple find and replace to the value in the .env
+      // so process?.env?.REACT_APP_VERSION_INFO for example is invalid
+      dashboardVersion = process.env.REACT_APP_VERSION_INFO;
+   } catch (reason) {
+      console.log("Process env error:", reason);
    }
    
    if (dashboardVersion) {
