Horde: Remove 'Z' suffix from Tableau telemetry events.

diff --git a/Engine/Source/Programs/Horde/Horde.Build/Telemetry/EpicTelemetrySink.cs b/Engine/Source/Programs/Horde/Horde.Build/Telemetry/EpicTelemetrySink.cs
--- a/Engine/Source/Programs/Horde/Horde.Build/Telemetry/EpicTelemetrySink.cs
+++ b/Engine/Source/Programs/Horde/Horde.Build/Telemetry/EpicTelemetrySink.cs
@@ -6,6 +6,7 @@ using System.Net.Http;
 using System.Net.Http.Headers;
 using System.Text;
 using System.Text.Json;
+using System.Text.Json.Serialization;
 using System.Threading;
 using System.Threading.Tasks;
 using EpicGames.Core;
@@ -38,6 +39,22 @@ namespace Horde.Build.Telemetry
 	/// </summary>
 	public sealed class EpicTelemetrySink : IHostedService, ITelemetrySink, IDisposable
 	{
+		// Converter for datetime formats that Tableau can ingest
+		class TableauDateTimeConverter : JsonConverter<DateTime>
+		{
+			const string Format = "yyyy-MM-dd HH:mm:ss";
+
+			public override void Write(Utf8JsonWriter writer, DateTime date, JsonSerializerOptions options)
+			{
+				writer.WriteStringValue(date.ToString(Format));
+			}
+
+			public override DateTime Read(ref Utf8JsonReader reader, Type typeToConvert, JsonSerializerOptions options)
+			{
+				return DateTime.ParseExact(reader.GetString() ?? String.Empty, Format, null);
+			}
+		}
+
 		/// <summary>
 		/// Name of the HTTP client for writing telemetry data
 		/// </summary>
@@ -74,6 +91,7 @@ namespace Horde.Build.Telemetry
 			_jsonOptions = new JsonSerializerOptions();
 			Startup.ConfigureJsonSerializer(_jsonOptions);
 			_jsonOptions.PropertyNamingPolicy = null;
+			_jsonOptions.Converters.Insert(0, new TableauDateTimeConverter());
 
 			_packetMemoryWriter = new ArrayMemoryWriter(65536);
 			_packetWriter = new Utf8JsonWriter(_packetMemoryWriter);
