Horde: Pass the Horde environment to telemetry events.

diff --git a/Engine/Source/Programs/Horde/Horde.Build/Telemetry/EpicTelemetrySink.cs b/Engine/Source/Programs/Horde/Horde.Build/Telemetry/EpicTelemetrySink.cs
--- a/Engine/Source/Programs/Horde/Horde.Build/Telemetry/EpicTelemetrySink.cs
+++ b/Engine/Source/Programs/Horde/Horde.Build/Telemetry/EpicTelemetrySink.cs
@@ -110,8 +110,8 @@ namespace Horde.Build.Telemetry
 					["SessionID"] = Guid.NewGuid().ToString(),
 					["AppID"] = config.AppId,
 					["AppVersion"] = Program.Version.ToString(),
-					["UploadType"] = "eteventstream",
-					["AppEnvironment"] = "devlatest"
+					["AppEnvironment"] = Program.DeploymentEnvironment,
+					["UploadType"] = "eteventstream"
 				};
 				_uri = new Uri(QueryHelpers.AddQueryString(config.Url.ToString(), queryParams));
 			}
