Horde: Additional logging for config file updates, showing files being included.

diff --git a/Engine/Source/Programs/Horde/Horde.Build.Tests/ConfigTests.cs b/Engine/Source/Programs/Horde/Horde.Build.Tests/ConfigTests.cs
--- a/Engine/Source/Programs/Horde/Horde.Build.Tests/ConfigTests.cs
+++ b/Engine/Source/Programs/Horde/Horde.Build.Tests/ConfigTests.cs
@@ -73,7 +73,7 @@ namespace Horde.Build.Tests
 			Dictionary<string, IConfigSource> sources = new Dictionary<string, IConfigSource>();
 			sources["memory"] = source;
 
-			ConfigContext options = new ConfigContext(_jsonOptions, sources);
+			ConfigContext options = new ConfigContext(_jsonOptions, sources, NullLogger.Instance);
 
 			ConfigObject result = await ConfigType.ReadAsync<ConfigObject>(fooUri, options, cancellationToken);
 			Assert.AreEqual(result.TestString, "hello");
@@ -107,7 +107,7 @@ namespace Horde.Build.Tests
 			Dictionary<string, IConfigSource> sources = new Dictionary<string, IConfigSource>();
 			sources["file"] = source;
 
-			ConfigContext context = new ConfigContext(_jsonOptions, sources);
+			ConfigContext context = new ConfigContext(_jsonOptions, sources, NullLogger.Instance);
 
 			ConfigObject result = await ConfigType.ReadAsync<ConfigObject>(fooUri, context, cancellationToken);
 			Assert.AreEqual(result.TestString, "hello");
diff --git a/Engine/Source/Programs/Horde/Horde.Build/Configuration/ConfigContext.cs b/Engine/Source/Programs/Horde/Horde.Build/Configuration/ConfigContext.cs
--- a/Engine/Source/Programs/Horde/Horde.Build/Configuration/ConfigContext.cs
+++ b/Engine/Source/Programs/Horde/Horde.Build/Configuration/ConfigContext.cs
@@ -5,6 +5,7 @@ using System.Collections.Generic;
 using System.Diagnostics;
 using System.Diagnostics.CodeAnalysis;
 using System.Text.Json;
+using Microsoft.Extensions.Logging;
 
 namespace Horde.Build.Configuration
 {
@@ -44,6 +45,11 @@ namespace Horde.Build.Configuration
 		/// </summary>
 		public Dictionary<Uri, IConfigFile> Files { get; } = new Dictionary<Uri, IConfigFile>();
 
+		/// <summary>
+		/// Logger for config messages
+		/// </summary>
+		public ILogger Logger { get; }
+
 		/// <summary>
 		/// Uri of the current file
 		/// </summary>
@@ -57,11 +63,12 @@ namespace Horde.Build.Configuration
 		/// <summary>
 		/// Constructor
 		/// </summary>
-		public ConfigContext(JsonSerializerOptions jsonOptions, IReadOnlyDictionary<string, IConfigSource> sources)
+		public ConfigContext(JsonSerializerOptions jsonOptions, IReadOnlyDictionary<string, IConfigSource> sources, ILogger logger)
 		{
 			JsonOptions = jsonOptions;
 			Sources = sources;
 			ScopeStack.Push("$");
+			Logger = logger;
 		}
 
 		/// <summary>
diff --git a/Engine/Source/Programs/Horde/Horde.Build/Configuration/ConfigService.cs b/Engine/Source/Programs/Horde/Horde.Build/Configuration/ConfigService.cs
--- a/Engine/Source/Programs/Horde/Horde.Build/Configuration/ConfigService.cs
+++ b/Engine/Source/Programs/Horde/Horde.Build/Configuration/ConfigService.cs
@@ -323,7 +323,7 @@ namespace Horde.Build.Configuration
 			}
 
 			// Read the config files
-			ConfigContext context = new ConfigContext(_jsonOptions, _sources);
+			ConfigContext context = new ConfigContext(_jsonOptions, _sources, _logger);
 			try
 			{
 				ConfigSnapshot snapshot = new ConfigSnapshot();
diff --git a/Engine/Source/Programs/Horde/Horde.Build/Configuration/ConfigType.cs b/Engine/Source/Programs/Horde/Horde.Build/Configuration/ConfigType.cs
--- a/Engine/Source/Programs/Horde/Horde.Build/Configuration/ConfigType.cs
+++ b/Engine/Source/Programs/Horde/Horde.Build/Configuration/ConfigType.cs
@@ -12,6 +12,7 @@ using System.Text.Json.Serialization;
 using System.Threading;
 using System.Threading.Tasks;
 using EpicGames.Core;
+using Microsoft.Extensions.Logging;
 
 using JsonObject = System.Text.Json.Nodes.JsonObject;
 
@@ -219,6 +220,7 @@ namespace Horde.Build.Configuration
 				string? path = (string?)jsonNode;
 
 				Uri uri = ConfigType.CombinePaths(context.CurrentFile, path!);
+				context.Logger.LogInformation("Including config file {Uri}", uri);
 				IConfigFile file = await ReadFileAsync(uri, context, cancellationToken);
 
 				context.IncludeStack.Push(file);
