Horde: Automation view improvements
#rnx
#skipci

diff --git a/Engine/Source/Programs/Horde/HordeDashboard/src/components/AutomationView.tsx b/Engine/Source/Programs/Horde/HordeDashboard/src/components/AutomationView.tsx
--- a/Engine/Source/Programs/Horde/HordeDashboard/src/components/AutomationView.tsx
+++ b/Engine/Source/Programs/Horde/HordeDashboard/src/components/AutomationView.tsx
@@ -1,13 +1,13 @@
 // Copyright Epic Games, Inc. All Rights Reserved.
 
 import { isNumber } from '@datadog/browser-core';
-import { Checkbox, DefaultButton, Dropdown, FocusZone, FocusZoneDirection, FontIcon, IContextualMenuItem, IContextualMenuProps, IDropdownOption, Label, PrimaryButton, ScrollablePane, ScrollbarVisibility, Spinner, SpinnerSize, Stack, Text } from '@fluentui/react';
+import { ComboBox, DefaultButton, Dropdown, FocusZone, FocusZoneDirection, FontIcon, IComboBox, IComboBoxOption, IComboBoxStyles, IContextualMenuItem, IContextualMenuProps, IDropdownOption, Label, PrimaryButton, ScrollablePane, ScrollbarVisibility, SelectableOptionMenuItemType, Spinner, SpinnerSize, Stack, Text } from '@fluentui/react';
 import * as d3 from "d3";
 import { action, makeObservable, observable } from 'mobx';
 import { observer } from 'mobx-react-lite';
 import moment from 'moment';
 import React, { useState } from 'react';
-import { useNavigate, useLocation } from 'react-router-dom';
+import { useLocation, useNavigate } from 'react-router-dom';
 import backend from '../backend';
 import { GetTestDataRefResponse, GetTestMetaResponse, GetTestResponse, GetTestStreamResponse, GetTestSuiteResponse, TestOutcome } from '../backend/Api';
 import dashboard, { StatusColor } from '../backend/Dashboard';
@@ -452,6 +452,10 @@ class TestDataHandler {
 
       const state = this.state;
 
+      if (!state.targets?.length || !state.platforms?.length || !state.configurations?.length || !state.rhi?.length || !state.variation?.length) {
+         return [];
+      }
+
       if (!state.automation || !state.streams?.length) {
          return [];
       }
@@ -673,7 +677,7 @@ class TestDataHandler {
       });
 
       state.variation?.forEach(v => {
-         search.append("variation", v);
+         search.append("var", v);
       });
 
       state.tests?.forEach(r => {
@@ -1137,6 +1141,24 @@ class TestDataHandler {
       this.variation = Array.from(variation).sort((a, b) => a.localeCompare(b));
 
       this.loaded = true;
+
+      const state = this.state;
+      if (!state.platforms?.length && this.platforms.length) {
+         this.platforms.forEach(p => this.addPlatform(p));
+      }
+      if (!state.configurations?.length && this.configurations.length) {
+         this.configurations.forEach(c => this.addConfiguration(c));
+      }
+      if (!state.targets?.length && this.targets.length) {
+         this.targets.forEach(t => this.addTarget(t));
+      }
+      if (!state.rhi?.length && this.rhi.length) {
+         this.rhi.forEach(r => this.addRHI(r));
+      }
+      if (!state.variation?.length && this.variation.length) {
+         this.variation.forEach(v => this.addVariation(v));
+      }
+
       this.setUpdated();
 
       this.query()
@@ -1356,7 +1378,7 @@ const StreamChooser: React.FC<{ handler: TestDataHandler }> = observer(({ handle
       items: options,
    };
 
-   return <DefaultButton style={{ width: 220, textAlign: "left" }} menuProps={menuProps} text={anySelected ? "Select" : "None Selected"} />
+   return <DefaultButton style={{ width: 270, textAlign: "left" }} menuProps={menuProps} text={anySelected ? "Select" : "None Selected"} />
 });
 
 
@@ -1379,7 +1401,62 @@ const AutomationChooser: React.FC<{ handler: TestDataHandler }> = observer(({ ha
       items: options,
    };
 
-   return <DefaultButton style={{ width: 220, textAlign: "left" }} menuProps={menuProps} text={automation ? automation : "Select"} />
+   return <DefaultButton style={{ width: 270, textAlign: "left" }} menuProps={menuProps} text={automation ? automation : "Select"} />
+});
+
+
+const MultiOptionChooser: React.FC<{ options: IComboBoxOption[], initialKeysIn: string[], updateKeys: (selectedKeys: string[]) => void }> = observer(({ options, initialKeysIn, updateKeys }) => {
+
+   let initialKeys = [...initialKeysIn];
+
+   const [selectedKeys, setSelectedKeys] = React.useState<string[]>(initialKeys);
+
+   if (options.length === initialKeys.length) {
+      selectedKeys.push('selectAll');
+   }
+
+   options.unshift({ key: 'selectAll', text: 'Select All', itemType: SelectableOptionMenuItemType.SelectAll });
+
+   const selectableOptions = options.filter(
+      option =>
+         (option.itemType === SelectableOptionMenuItemType.Normal || option.itemType === undefined) && !option.disabled,
+   );
+
+   const onChange = (
+      event: React.FormEvent<IComboBox>,
+      option?: IComboBoxOption,
+      index?: number,
+      value?: string,
+   ): void => {
+      const selected = option?.selected;
+      const currentSelectedOptionKeys = selectedKeys.filter(key => key !== 'selectAll');
+      const selectAllState = currentSelectedOptionKeys.length === selectableOptions.length;
+
+      let updatedKeys: string[] = [];
+
+      if (option) {
+         if (option?.itemType === SelectableOptionMenuItemType.SelectAll) {
+            updatedKeys = selectAllState ? [] : ['selectAll', ...selectableOptions.map(o => o.key as string)];
+            selectAllState
+               ? setSelectedKeys(updatedKeys)
+               : setSelectedKeys(updatedKeys);
+         } else {
+            updatedKeys = selected
+               ? [...currentSelectedOptionKeys, option!.key as string]
+               : currentSelectedOptionKeys.filter(k => k !== option.key);
+            if (updatedKeys.length === selectableOptions.length) {
+               updatedKeys.push('selectAll');
+            }
+            setSelectedKeys(updatedKeys);
+         }
+
+         updateKeys([...updatedKeys.filter(k => k !== 'selectAll')]);
+      }
+   };
+
+   const comboBoxStyles: Partial<IComboBoxStyles> = { root: { width: 270 } };
+
+   return <ComboBox placeholder="None" defaultSelectedKey={initialKeys} multiSelect options={options} onChange={onChange} styles={comboBoxStyles} />
 });
 
 const TestChooser: React.FC<{ handler: TestDataHandler }> = observer(({ handler }) => {
@@ -1388,34 +1465,36 @@ const TestChooser: React.FC<{ handler: TestDataHandler }> = observer(({ handler
    if (handler.updated) { }
 
    const ctests: Set<string> = new Set(handler.state.tests ?? []);
-   const metaTests = handler.metaTests;
-   const metaTestIds: Set<string> = new Set(metaTests.map(t => t.id));
+   //const metaTests = handler.metaTests;
+   //const metaTestIds: Set<string> = new Set(metaTests.map(t => t.id));
    const streamTests = handler.streamTests;
 
    if (!streamTests.length) {
       return null;
    }
 
-   const tests = streamTests.map(test => {
-      const t = test.name;
-      const valid = metaTestIds.has(test.id);
-      return <Stack style={{ textDecoration: valid ? undefined : "line-through" }}>
-         <Checkbox label={t} checked={!!ctests.has(t)} onChange={(ev, checked) => {
-            if (checked) {
-               handler.addTest(t);
-            } else {
-               handler.removeTest(t);
-            }
-         }} />
-      </Stack>
+   const options: IComboBoxOption[] = [];
+   streamTests.forEach(t => {
+      options.push({ key: t.name, text: t.name });
    });
 
-   return <Stack>
-      <Stack style={{ paddingTop: 12, paddingBottom: 4 }}>
+   return <Stack style={{ paddingTop: 12, paddingBottom: 4 }}>
+      <Stack style={{ paddingTop: 0, paddingBottom: 4 }}>
          <Label>Tests</Label>
       </Stack>
       <Stack tokens={{ childrenGap: 6 }} style={{ paddingLeft: 12 }}>
-         {tests}
+         <MultiOptionChooser options={options} initialKeysIn={handler.state.tests ?? []} updateKeys={(keys) => {
+
+            handler.streamTests.forEach(t => {
+               const selected = keys.find(k => k === t.name);
+               if (!selected && ctests.has(t.name)) {
+                  handler.removeTest(t.name);
+               }
+               else if (selected && !ctests.has(t.name)) {
+                  handler.addTest(t.name);
+               }
+            });
+         }} />
       </Stack>
    </Stack>
 });
@@ -1426,34 +1505,36 @@ const SuiteChooser: React.FC<{ handler: TestDataHandler }> = observer(({ handler
    if (handler.updated) { }
 
    const csuites: Set<string> = new Set(handler.state.suites ?? []);
-   const metaSuites = handler.metaSuites;
-   const metaSuiteIds: Set<string> = new Set(metaSuites.map(s => s.id));
+   //const metaSuites = handler.metaSuites;
+   //const metaSuiteIds: Set<string> = new Set(metaSuites.map(s => s.id));
    const streamSuites = handler.streamSuites;
 
    if (!streamSuites.length) {
       return null;
    }
 
-   const suites = streamSuites.map(suite => {
-      const s = suite.name;
-      const valid = metaSuiteIds.has(suite.id);
-      return <Stack style={{ textDecoration: valid ? undefined : "line-through" }}>
-         <Checkbox label={s} checked={!!csuites.has(s)} onChange={(ev, checked) => {
-            if (checked) {
-               handler.addSuite(s);
-            } else {
-               handler.removeSuite(s);
-            }
-         }} />
-      </Stack>
+   const options: IComboBoxOption[] = [];
+   streamSuites.forEach(t => {
+      options.push({ key: t.name, text: t.name });
    });
 
-   return <Stack>
-      <Stack style={{ paddingTop: 12, paddingBottom: 4 }}>
+   return <Stack style={{ paddingTop: 12, paddingBottom: 4 }}>
+      <Stack style={{ paddingTop: 0, paddingBottom: 4 }}>
          <Label>Suites</Label>
       </Stack>
       <Stack tokens={{ childrenGap: 6 }} style={{ paddingLeft: 12 }}>
-         {suites}
+         <MultiOptionChooser options={options} initialKeysIn={handler.state.suites ?? []} updateKeys={(keys) => {
+
+            handler.streamSuites.forEach(t => {
+               const selected = keys.find(k => k === t.name);
+               if (!selected && csuites.has(t.name)) {
+                  handler.removeSuite(t.name);
+               }
+               else if (selected && !csuites.has(t.name)) {
+                  handler.addSuite(t.name);
+               }
+            });
+         }} />
       </Stack>
    </Stack>
 });
@@ -1466,24 +1547,29 @@ const PlatformChooser: React.FC<{ handler: TestDataHandler }> = observer(({ hand
 
    const cplatforms: Set<string> = new Set(handler.state.platforms ?? []);
 
-   const platforms = handler.platforms.map(p => {
-      return <Stack>
-         <Checkbox label={p} checked={!!cplatforms.has(p)} onChange={(ev, checked) => {
-            if (checked) {
-               handler.addPlatform(p);
-            } else {
-               handler.removePlatform(p);
-            }
-         }} />
-      </Stack>
+   const options: IComboBoxOption[] = [];
+
+   handler.platforms.forEach(p => {
+      options.push({ key: p, text: p });
    });
 
-   return <Stack>
+   return <Stack style={{ paddingTop: 12, paddingBottom: 4 }}>
       <Stack style={{ paddingTop: 0, paddingBottom: 4 }}>
          <Label>Platforms</Label>
       </Stack>
       <Stack tokens={{ childrenGap: 6 }} style={{ paddingLeft: 12 }}>
-         {platforms}
+         <MultiOptionChooser options={options} initialKeysIn={handler.state.platforms ?? []} updateKeys={(keys) => {
+
+            handler.platforms.forEach(p => {
+               const selected = keys.find(k => k === p);
+               if (!selected && cplatforms.has(p)) {
+                  handler.removePlatform(p);
+               }
+               else if (selected && !cplatforms.has(p)) {
+                  handler.addPlatform(p);
+               }
+            });
+         }} />
       </Stack>
    </Stack>
 });
@@ -1495,26 +1581,32 @@ const ConfigChooser: React.FC<{ handler: TestDataHandler }> = observer(({ handle
 
    const cconfig: Set<string> = new Set(handler.state.configurations ?? []);
 
-   const configurations = handler.configurations.map(c => {
-      return <Stack>
-         <Checkbox label={c} checked={!!cconfig.has(c)} onChange={(ev, checked) => {
-            if (checked) {
-               handler.addConfiguration(c);
-            } else {
-               handler.removeConfiguration(c);
-            }
-         }} />
-      </Stack>
+   const options: IComboBoxOption[] = [];
+
+   handler.configurations.forEach(p => {
+      options.push({ key: p, text: p });
    });
 
-   return <Stack>
-      <Stack style={{ paddingTop: 12, paddingBottom: 4 }}>
+   return <Stack style={{ paddingTop: 12, paddingBottom: 4 }}>
+      <Stack style={{ paddingTop: 0, paddingBottom: 4 }}>
          <Label>Configurations</Label>
       </Stack>
       <Stack tokens={{ childrenGap: 6 }} style={{ paddingLeft: 12 }}>
-         {configurations}
+         <MultiOptionChooser options={options} initialKeysIn={handler.state.configurations ?? []} updateKeys={(keys) => {
+
+            handler.configurations.forEach(p => {
+               const selected = keys.find(k => k === p);
+               if (!selected && cconfig.has(p)) {
+                  handler.removeConfiguration(p);
+               }
+               else if (selected && !cconfig.has(p)) {
+                  handler.addConfiguration(p);
+               }
+            });
+         }} />
       </Stack>
    </Stack>
+
 });
 
 const TargetChooser: React.FC<{ handler: TestDataHandler }> = observer(({ handler }) => {
@@ -1524,24 +1616,29 @@ const TargetChooser: React.FC<{ handler: TestDataHandler }> = observer(({ handle
 
    const ctargets: Set<string> = new Set(handler.state.targets ?? []);
 
-   const targets = handler.targets.map(t => {
-      return <Stack>
-         <Checkbox label={t} checked={!!ctargets.has(t)} onChange={(ev, checked) => {
-            if (checked) {
-               handler.addTarget(t);
-            } else {
-               handler.removeTarget(t);
-            }
-         }} />
-      </Stack>
+   const options: IComboBoxOption[] = [];
+
+   handler.targets.forEach(p => {
+      options.push({ key: p, text: p });
    });
 
-   return <Stack>
-      <Stack style={{ paddingTop: 12, paddingBottom: 4 }}>
+   return <Stack style={{ paddingTop: 12, paddingBottom: 4 }}>
+      <Stack style={{ paddingTop: 0, paddingBottom: 4 }}>
          <Label>Targets</Label>
       </Stack>
       <Stack tokens={{ childrenGap: 6 }} style={{ paddingLeft: 12 }}>
-         {targets}
+         <MultiOptionChooser options={options} initialKeysIn={handler.state.targets ?? []} updateKeys={(keys) => {
+
+            handler.targets.forEach(p => {
+               const selected = keys.find(k => k === p);
+               if (!selected && ctargets.has(p)) {
+                  handler.removeTarget(p);
+               }
+               else if (selected && !ctargets.has(p)) {
+                  handler.addTarget(p);
+               }
+            });
+         }} />
       </Stack>
    </Stack>
 });
@@ -1553,26 +1650,31 @@ const RHIChooser: React.FC<{ handler: TestDataHandler }> = observer(({ handler }
 
    const crhi: Set<string> = new Set(handler.state.rhi ?? []);
 
-   const rhi = handler.rhi.map(r => {
-      return <Stack>
-         <Checkbox label={r === "default" ? "Default" : r.toLocaleUpperCase()} checked={!!crhi.has(r)} onChange={(ev, checked) => {
-            if (checked) {
-               handler.addRHI(r);
-            } else {
-               handler.removeRHI(r);
-            }
-         }} />
-      </Stack>
+   const options: IComboBoxOption[] = [];
+
+   handler.rhi.forEach(p => {
+      options.push({ key: p, text: p === "default" ? "Default" : p.toUpperCase() });
    });
 
-   return <Stack>
-      <Stack style={{ paddingTop: 12, paddingBottom: 4 }}>
+   return <Stack style={{ paddingTop: 12, paddingBottom: 4 }}>
+      <Stack style={{ paddingTop: 0, paddingBottom: 4 }}>
          <Label>RHI</Label>
       </Stack>
       <Stack tokens={{ childrenGap: 6 }} style={{ paddingLeft: 12 }}>
-         {rhi}
+         <MultiOptionChooser options={options} initialKeysIn={handler.state.rhi ?? []} updateKeys={(keys) => {
+            handler.rhi.forEach(p => {
+               const selected = keys.find(k => k === p);
+               if (!selected && crhi.has(p)) {
+                  handler.removeRHI(p);
+               }
+               else if (selected && !crhi.has(p)) {
+                  handler.addRHI(p);
+               }
+            });
+         }} />
       </Stack>
    </Stack>
+
 });
 
 const VariationChooser: React.FC<{ handler: TestDataHandler }> = observer(({ handler }) => {
@@ -1582,31 +1684,37 @@ const VariationChooser: React.FC<{ handler: TestDataHandler }> = observer(({ han
 
    const cvariation: Set<string> = new Set(handler.state.variation ?? []);
 
-   const variation = handler.variation.map(v => {
-      return <Stack>
-         <Checkbox label={v === "default" ? "Default" : v.toLocaleUpperCase()} checked={!!cvariation.has(v)} onChange={(ev, checked) => {
-            if (checked) {
-               handler.addVariation(v);
-            } else {
-               handler.removeVariation(v);
-            }
-         }} />
-      </Stack>
+   const options: IComboBoxOption[] = [];
+
+   handler.variation.forEach(p => {
+      options.push({ key: p, text: p === "default" ? "Default" : p.toUpperCase() });
    });
 
-   return <Stack>
-      <Stack style={{ paddingTop: 12, paddingBottom: 4 }}>
+   return <Stack style={{ paddingTop: 12, paddingBottom: 4 }}>
+      <Stack style={{ paddingTop: 0, paddingBottom: 4 }}>
          <Label>Variation</Label>
       </Stack>
       <Stack tokens={{ childrenGap: 6 }} style={{ paddingLeft: 12 }}>
-         {variation}
+         <MultiOptionChooser options={options} initialKeysIn={handler.state.variation ?? []} updateKeys={(keys) => {
+
+            handler.variation.forEach(p => {
+               const selected = keys.find(k => k === p);
+               if (!selected && cvariation.has(p)) {
+                  handler.removeVariation(p);
+               }
+               else if (selected && !cvariation.has(p)) {
+                  handler.addVariation(p);
+               }
+            });
+         }} />
       </Stack>
    </Stack>
+
 });
 
 const AutomationSidebarLeft: React.FC<{ handler: TestDataHandler }> = ({ handler }) => {
 
-   return <Stack style={{ paddingRight: 18 }}>
+   return <Stack style={{ width: 300, paddingRight: 18 }}>
       <Stack className={hordeClasses.modal}>
          <Stack>
             <Stack style={{ paddingTop: 0, paddingBottom: 4 }}>
@@ -1630,14 +1738,6 @@ const AutomationSidebarLeft: React.FC<{ handler: TestDataHandler }> = ({ handler
          <Stack>
             <SuiteChooser handler={handler} />
          </Stack>
-      </Stack>
-   </Stack>
-}
-
-const AutomationSidebarRight: React.FC<{ handler: TestDataHandler }> = ({ handler }) => {
-
-   return <Stack className={hordeClasses.modal}>
-      <Stack>
          <Stack>
             <PlatformChooser handler={handler} />
          </Stack>
@@ -1753,7 +1853,7 @@ export const AutomationView: React.FC = observer(() => {
 
    if (state.search !== csearch) {
       location.search = csearch;
-      navigate(location, {replace: true});
+      navigate(location, { replace: true });
       setState({ ...state, search: csearch });
       return null;
    }
@@ -1789,9 +1889,6 @@ export const AutomationView: React.FC = observer(() => {
                         </Stack>
                         <AutomationCenter handler={handler} />
                      </Stack>
-                     <Stack>
-                        <AutomationSidebarRight handler={handler} />
-                     </Stack>
                   </Stack>
                </Stack>}
             </Stack>
@@ -1991,7 +2088,7 @@ class AutomationGraph {
 
       const handler = this.handler;
       const refs = this.refs.sort((a, b) => handler.metaNames.get(a.metaId)!.localeCompare(handler.metaNames.get(b.metaId)!));
-      const width = 1254
+      const width = 1328
 
       const scolors = dashboard.getStatusColors();
       const colors: Record<string, string> = {
