Horde: Dashboard schedule changes, use server calculations and fix sort
#rnx
#skipci

diff --git a/Engine/Source/Programs/Horde/HordeDashboard/src/backend/Api.ts b/Engine/Source/Programs/Horde/HordeDashboard/src/backend/Api.ts
--- a/Engine/Source/Programs/Horde/HordeDashboard/src/backend/Api.ts
+++ b/Engine/Source/Programs/Horde/HordeDashboard/src/backend/Api.ts
@@ -141,7 +141,7 @@ export enum JobStepBatchError {
 	ExecutionError = "ExecutionError",
 
 	/** The change that the job is running against is invalid. */
-	UnknownShelf  = "UnknownShelf"
+	UnknownShelf = "UnknownShelf"
 
 }
 
@@ -1725,6 +1725,9 @@ export type GetScheduleResponse = {
 	/** Last time that the schedule was triggered */
 	lastTriggerTime: Date | string;
 
+	/// Next trigger times for schedule
+	nextTriggerTimesUTC: Date[];
+
 	/** List of active jobs */
 	activeJobs: string[];
 
@@ -4405,16 +4408,16 @@ export type GetSuiteTestDataResponse = {
 
 /// Test details
 export type GetTestDataDetailsResponse = {
-	
+
 	/// The corresponding test ref	
 	id: string;
-	
+
 	/// The test documents for this ref	
 	testDataIds: string[];
-	
+
 	/// Suite test data		
 	suiteTests?: GetSuiteTestDataResponse[];
-	
+
 }
 
 /// Testt data ref
diff --git a/Engine/Source/Programs/Horde/HordeDashboard/src/backend/Backend.ts b/Engine/Source/Programs/Horde/HordeDashboard/src/backend/Backend.ts
--- a/Engine/Source/Programs/Horde/HordeDashboard/src/backend/Backend.ts
+++ b/Engine/Source/Programs/Horde/HordeDashboard/src/backend/Backend.ts
@@ -734,6 +734,10 @@ export class Backend {
         return new Promise<ScheduleData>((resolve, reject) => {
 
             this.backend.get(`/api/v1/schedules/${id}`).then((value) => {
+                const data = value.data as ScheduleData;
+                data.nextTriggerTimesUTC.forEach((d, index) => {
+                    data.nextTriggerTimesUTC[index] = new Date(d as any as string);
+                });
                 resolve(value.data as ScheduleData);
             }).catch(reason => {
                 reject(reason);
diff --git a/Engine/Source/Programs/Horde/HordeDashboard/src/components/SchedulePane.tsx b/Engine/Source/Programs/Horde/HordeDashboard/src/components/SchedulePane.tsx
--- a/Engine/Source/Programs/Horde/HordeDashboard/src/components/SchedulePane.tsx
+++ b/Engine/Source/Programs/Horde/HordeDashboard/src/components/SchedulePane.tsx
@@ -1,239 +1,187 @@
 // Copyright Epic Games, Inc. All Rights Reserved.
-import moment, { Moment } from 'moment-timezone';
 import { mergeStyleSets, Stack } from '@fluentui/react';
 import { DetailsList, DetailsListLayoutMode, DetailsRow, IColumn, IDetailsListProps, SelectionMode } from '@fluentui/react/lib/DetailsList';
+import moment from 'moment-timezone';
 import React, { useState } from 'react';
 import { Link } from 'react-router-dom';
 import backend from '../backend';
-import { GetJobResponse, GetSchedulePatternResponse, GetTemplateRefResponse, JobState } from '../backend/Api';
+import { GetJobResponse, GetTemplateRefResponse, JobState } from '../backend/Api';
 import dashboard from '../backend/Dashboard';
 import { displayTimeZone } from "../base/utilities/timeUtils";
 
 const classes = mergeStyleSets({
-    detailsRow: {
-        selectors: {
-            '.ms-DetailsRow-cell': {
-                overflow: "hidden",
-                whiteSpace: "nowrap",
-                padding: 8
-            },
-            '.ms-List-cell:nth-child(odd)': {
-                background: "rgb(240, 239, 239)",
-            },
-            '.ms-List-cell:nth-child(even)': {
-                background: "#FFFFFF",
-            }
-
-        }
-    }
+   detailsRow: {
+      selectors: {
+         '.ms-DetailsRow-cell': {
+            overflow: "hidden",
+            whiteSpace: "nowrap",
+            padding: 8
+         },
+         '.ms-List-cell:nth-child(odd)': {
+            background: "rgb(240, 239, 239)",
+         },
+         '.ms-List-cell:nth-child(even)': {
+            background: "#FFFFFF",
+         }
+
+      }
+   }
 });
 
 type ScheduleItem = {
-    template: GetTemplateRefResponse
+   template: GetTemplateRefResponse
 };
 
 const scheduledJobs = new Map<string, GetJobResponse | null>();
 
-const getNextTriggerTime = (pattern: GetSchedulePatternResponse, lastTimeIn: Date | string): Moment | undefined => {
+export const SchedulePane: React.FC<{ templates: GetTemplateRefResponse[] }> = ({ templates }) => {
 
-    // Convert last time into the correct timezone for running the scheule
-    const lastTime = moment(lastTimeIn);
+   let [updated, setUpdated] = useState(new Date().getTime());
 
-    // Get the base time (ie. the start of this day) for anchoring the schedule
-    let baseTime = moment().year(lastTime.year()).month(lastTime.month()).date(lastTime.day()).hour(0).minute(0).second(0).millisecond(0);
+   const ids = new Set<string>();
 
-    for (; ;) {
+   if (new Date().getTime() - updated > 10000) {
 
-        const day = baseTime.format('dddd');
+      scheduledJobs.forEach((job, key) => {
+         if (job?.state !== JobState.Complete) {
+            ids.add(key);
+         }
+      });
 
-        if (!pattern.daysOfWeek || !!pattern.daysOfWeek.find(d => {
-            return d.toLowerCase() === day.toLowerCase();
-        })) {
+   }
 
-            // Get the last time in minutes from the start of this day
-            const lastTimeMinutes = lastTime.diff(baseTime, "minutes");
+   templates.forEach(t => {
+      if (t.schedule && t.schedule.enabled && t.schedule.maxActive === 1 && t.schedule.activeJobs.length) {
+         const jid = t.schedule.activeJobs[0];
+         if (!scheduledJobs.has(jid)) {
+            scheduledJobs.set(jid, null);
+            ids.add(jid);
+         }
+      }
+   })
 
-            if (lastTimeMinutes < pattern.minTime) {
-                return baseTime.add(pattern.minTime, "minutes");
-            }
+   if (ids.size) {
+      backend.getJobsByIds(Array.from(ids.values()), { filter: "id,createTime,state" }, false).then(response => {
+         response.forEach(j => {
+            scheduledJobs.set(j.id, j);
+         });
 
-            // Otherwise, get the time for the last trigger in the day.
-            if (pattern.interval) {
-                const actualMaxTime = pattern.maxTime ? pattern.maxTime : (24 * 60);
-                if (lastTimeMinutes < actualMaxTime) {
-                    const lastIndex = (lastTimeMinutes - pattern.minTime) / pattern.interval;
-                    const nextIndex = lastIndex + 1;
-
-                    const nextTimeMinutes = pattern.minTime + (nextIndex * pattern.interval);
-                    if (nextTimeMinutes < actualMaxTime) {
-                        return baseTime.add(nextTimeMinutes, "minutes");
-                    }
-                }
-            }
-        }
+         setUpdated(new Date().getTime());
+      });
+   }
 
-        baseTime = baseTime.add(1, "days");
+   const items = templates.filter(t => !!t.schedule?.enabled).map(t => {
+      return { template: t };
+   });
 
-    }
+   const columns = [
+      { key: 'schedule_column1', name: 'job_name', minWidth: 320, maxWidth: 320, isResizable: false },
+      { key: 'schedule_column3', name: 'time_0', minWidth: 140, maxWidth: 140, isResizable: false },
+      { key: 'schedule_column4', name: 'time_1', minWidth: 100, maxWidth: 100, isResizable: false },
+      { key: 'schedule_column5', name: 'time_2', minWidth: 100, maxWidth: 100, isResizable: false },
+      { key: 'schedule_column6', name: 'time_3', minWidth: 100, maxWidth: 100, isResizable: false },
+      { key: 'schedule_column7', name: 'time_4', minWidth: 100, maxWidth: 100, isResizable: false },
+      { key: 'schedule_column8', name: 'time_5', minWidth: 100, maxWidth: 100, isResizable: false },
+      { key: 'schedule_column9', name: 'time_6', minWidth: 100, maxWidth: 100, isResizable: false },
+      { key: 'schedule_column10', name: 'space', minWidth: 2, isResizable: false }
 
-}
+   ];
 
-export const SchedulePane: React.FC<{ templates: GetTemplateRefResponse[] }> = ({ templates }) => {
+   const renderItem = (item: ScheduleItem, index?: number, column?: IColumn) => {
 
-    let [updated, setUpdated] = useState(new Date().getTime());
+      const name = column!.name;
+      const template = item.template;
+      const schedule = template.schedule!;
 
-    const ids = new Set<string>();
+      if (name === "job_name") {
 
-    if (new Date().getTime() - updated > 10000 ) {
+         return <Stack horizontalAlign="center" verticalFill verticalAlign="center">{item.template.name}</Stack>;
+      }
 
-        scheduledJobs.forEach((job, key) => {
-            if (job?.state !== JobState.Complete ) {
-                ids.add(key);
-            }
-        });
+      if (name.startsWith("time_")) {
 
-    }
+         const index = parseInt(name.slice(5));
 
-    templates.forEach(t => {
-        if (t.schedule && t.schedule.enabled && t.schedule.maxActive === 1 && t.schedule.activeJobs.length) {
-            const jid = t.schedule.activeJobs[0];
-            if (!scheduledJobs.has(jid)) {
-                scheduledJobs.set(jid, null);
-                ids.add(jid);
-            }
-        }
-    })
-
-    if (ids.size) {
-        backend.getJobsByIds(Array.from(ids.values()), { filter: "id,createTime,state" }, false).then(response => {
-            response.forEach(j => {
-                scheduledJobs.set(j.id, j);
-            });
-
-            setUpdated(new Date().getTime());
-        });
-    }
-
-    const items = templates.filter(t => !!t.schedule?.enabled).map(t => {
-        return { template: t };
-    });
-
-    const columns = [
-        { key: 'schedule_column1', name: 'job_name', minWidth: 320, maxWidth: 320, isResizable: false },
-        { key: 'schedule_column3', name: 'time_0', minWidth: 140, maxWidth: 140, isResizable: false },
-        { key: 'schedule_column4', name: 'time_1', minWidth: 100, maxWidth: 100, isResizable: false },
-        { key: 'schedule_column5', name: 'time_2', minWidth: 100, maxWidth: 100, isResizable: false },
-        { key: 'schedule_column6', name: 'time_3', minWidth: 100, maxWidth: 100, isResizable: false },
-        { key: 'schedule_column7', name: 'time_4', minWidth: 100, maxWidth: 100, isResizable: false },
-        { key: 'schedule_column8', name: 'time_5', minWidth: 100, maxWidth: 100, isResizable: false },
-        { key: 'schedule_column9', name: 'time_6', minWidth: 100, maxWidth: 100, isResizable: false },
-        { key: 'schedule_column10', name: 'space', minWidth: 2, isResizable: false }
-
-    ];
-
-    const renderItem = (item: ScheduleItem, index?: number, column?: IColumn) => {
-
-        const name = column!.name;
-        const template = item.template;
-        const schedule = template.schedule!;
-
-        if (name === "job_name") {
-
-            return <Stack horizontalAlign="center" verticalFill verticalAlign="center">{item.template.name}</Stack>;            
-        }
-
-        if (name.startsWith("time_")) {
-
-            const index = parseInt(name.slice(5));
-
-            let times: Moment[] = [];
-            let triggerTime = moment(schedule.lastTriggerTime);
-            times.push(triggerTime);
-
-            schedule.patterns.forEach(p => {
-                triggerTime = getNextTriggerTime(p, schedule.lastTriggerTime)!;
-                times.push(triggerTime);
-                for (let i = 0; i < 5; i++) {
-                    triggerTime = getNextTriggerTime(p, moment.utc(triggerTime).toISOString())!;
-                    times.push(triggerTime);
-                }
-            })
-
-            times = times.sort((a, b) => {
-                return moment.utc(a).milliseconds() - moment.utc(b).milliseconds();
-            })
-
-            if (index >= times.length) {
-                return null;
-            }
+         const triggerTimes = schedule.nextTriggerTimesUTC.slice(0, 5);
 
-            const displayTime = times[index].tz(displayTimeZone());
-            const now = moment.utc().tz(displayTimeZone());
-            let timeStr = dashboard.display24HourClock ? displayTime.format('H:mm z') : displayTime.format('h:mma z');
-            if (!(now.isSame(displayTime, 'day'))) {
-                timeStr += displayTime.format(' (ddd)');
-            }
-            
-            if (index === 0) {
-
-                if (schedule?.activeJobs.length) {
-                    const job = scheduledJobs.get(schedule.activeJobs[0]);
-                    if (job) {
-                        const time = moment(job.createTime).tz(displayTimeZone());
-                        const timeStr = dashboard.display24HourClock ? time.format("MMM Do H:mm z") : time.format("MMM Do h:mma z");
-
-                        let state = "Complete";
-
-                        if (job.state === JobState.Waiting || job.state === JobState.Running) {
-                            state = "Running";
-                        }
-
-                        return <Link to={`/job/${job.id}`}>
-                            <Stack horizontalAlign="center" tokens={{childrenGap:4}}>
-                            <Stack >{timeStr}</Stack>
-                            <Stack >{state}</Stack>
-                        </Stack>
-                        </Link>;
-                    }
-                }
+         let times = triggerTimes.map(t => moment(t));
+         times.unshift(moment(schedule.lastTriggerTime));
+
+         times = times.sort((a, b) => {
+            return a.unix() - b.unix();
+         })
+
+         if (index >= times.length) {
+            return null;
+         }
+
+         const displayTime = times[index].tz(displayTimeZone());
+         const now = moment.utc().tz(displayTimeZone());
+         let timeStr = dashboard.display24HourClock ? displayTime.format('H:mm z') : displayTime.format('h:mma z');
+         if (!(now.isSame(displayTime, 'day'))) {
+            timeStr += displayTime.format(' (ddd)');
+         }
+
+         if (index === 0) {
+
+            if (schedule?.activeJobs.length) {
+               const job = scheduledJobs.get(schedule.activeJobs[0]);
+               if (job) {
+                  const time = moment(job.createTime).tz(displayTimeZone());
+                  const timeStr = dashboard.display24HourClock ? time.format("MMM Do H:mm z") : time.format("MMM Do h:mma z");
+
+                  let state = "Complete";
+
+                  if (job.state === JobState.Waiting || job.state === JobState.Running) {
+                     state = "Running";
+                  }
+
+                  return <Link to={`/job/${job.id}`}>
+                     <Stack horizontalAlign="center" tokens={{ childrenGap: 4 }}>
+                        <Stack >{timeStr}</Stack>
+                        <Stack >{state}</Stack>
+                     </Stack>
+                  </Link>;
+               }
             }
+         }
 
-            return <Stack horizontalAlign="center" verticalFill verticalAlign="center">{timeStr}</Stack>;
-        }
+         return <Stack horizontalAlign="center" verticalFill verticalAlign="center">{timeStr}</Stack>;
+      }
 
 
-        return null;
+      return null;
 
-    };
+   };
 
-    const renderRow: IDetailsListProps['onRenderRow'] = (props) => {
+   const renderRow: IDetailsListProps['onRenderRow'] = (props) => {
 
-        if (props) {
+      if (props) {
 
-            const nprops = { ...props, styles: { root: { backgroundColor: "unset", padding: 8 } } };
+         const nprops = { ...props, styles: { root: { backgroundColor: "unset", padding: 8 } } };
 
-            return <DetailsRow {...nprops} />
+         return <DetailsRow {...nprops} />
 
-        }
-        return null;
-    };
+      }
+      return null;
+   };
 
-    return <div className={classes.detailsRow} style={{ overflowY: 'auto', overflowX: 'hidden', maxHeight: "820px" }} data-is-scrollable={true}>
+   return <div className={classes.detailsRow} style={{ overflowY: 'auto', overflowX: 'hidden', maxHeight: "820px" }} data-is-scrollable={true}>
 
-        <DetailsList
-            styles={{root:{overflow:"unset"}}}
-            compact={true}
-            isHeaderVisible={false}
-            indentWidth={0}
-            items={items}
-            columns={columns}
-            setKey="set"
-            selectionMode={SelectionMode.none}
-            layoutMode={DetailsListLayoutMode.justified}
-            onRenderItemColumn={renderItem}
-            onRenderRow={renderRow}
-        />
-    </div>
+      <DetailsList
+         styles={{ root: { overflow: "unset" } }}
+         compact={true}
+         isHeaderVisible={false}
+         indentWidth={0}
+         items={items}
+         columns={columns}
+         setKey="set"
+         selectionMode={SelectionMode.none}
+         layoutMode={DetailsListLayoutMode.justified}
+         onRenderItemColumn={renderItem}
+         onRenderRow={renderRow}
+      />
+   </div>
 
 
 }
