Horde: Fix order of arguments to log message.

diff --git a/Engine/Source/Programs/Horde/Horde.Build/Configuration/ConfigService.cs b/Engine/Source/Programs/Horde/Horde.Build/Configuration/ConfigService.cs
--- a/Engine/Source/Programs/Horde/Horde.Build/Configuration/ConfigService.cs
+++ b/Engine/Source/Programs/Horde/Horde.Build/Configuration/ConfigService.cs
@@ -392,7 +392,7 @@ namespace Horde.Build.Configuration
 				{
 					if (!files[idx].Revision.Equals(pairs[idx].Value, StringComparison.Ordinal))
 					{
-						_logger.LogInformation("Config is out of date (file {Path}: {OldVersion} -> {NewVersion})", files[idx].Uri, files[idx].Revision, pairs[idx].Value);
+						_logger.LogInformation("Config is out of date (file {Path}: {OldVersion} -> {NewVersion})", files[idx].Uri, pairs[idx].Value, files[idx].Revision);
 						return true;
 					}
 				}
