Horde: More specific versioning for config files in logs.

diff --git a/Engine/Source/Programs/Horde/Horde.Build/Configuration/ConfigSource.cs b/Engine/Source/Programs/Horde/Horde.Build/Configuration/ConfigSource.cs
--- a/Engine/Source/Programs/Horde/Horde.Build/Configuration/ConfigSource.cs
+++ b/Engine/Source/Programs/Horde/Horde.Build/Configuration/ConfigSource.cs
@@ -311,8 +311,13 @@ namespace Horde.Build.Configuration
 		async ValueTask<ReadOnlyMemory<byte>> ReadAsync(Uri uri, int change, CancellationToken cancellationToken)
 		{
 			string cacheKey = $"{nameof(PerforceConfigSource)}:data:{uri}@{change}";
-			if (!_cache.TryGetValue(cacheKey, out ReadOnlyMemory<byte> data))
+			if (_cache.TryGetValue(cacheKey, out ReadOnlyMemory<byte> data))
 			{
+				_logger.LogInformation("Read {Uri}@{Change} from cache ({Key})", uri, change, cacheKey);
+			}
+			else
+			{
+				_logger.LogInformation("Reading {Uri}@{Change} from Perforce", uri, change);
 				using (IPerforceConnection perforce = await ConnectAsync(uri.Host, cancellationToken))
 				{
 					PerforceResponse<PrintRecord<byte[]>> response = await perforce.TryPrintDataAsync($"{uri.AbsolutePath}@{change}", cancellationToken);
diff --git a/Engine/Source/Programs/Horde/Horde.Build/Configuration/ConfigType.cs b/Engine/Source/Programs/Horde/Horde.Build/Configuration/ConfigType.cs
--- a/Engine/Source/Programs/Horde/Horde.Build/Configuration/ConfigType.cs
+++ b/Engine/Source/Programs/Horde/Horde.Build/Configuration/ConfigType.cs
@@ -220,7 +220,6 @@ namespace Horde.Build.Configuration
 				string? path = (string?)jsonNode;
 
 				Uri uri = ConfigType.CombinePaths(context.CurrentFile, path!);
-				context.Logger.LogInformation("Including config file {Uri}", uri);
 				IConfigFile file = await ReadFileAsync(uri, context, cancellationToken);
 
 				context.IncludeStack.Push(file);
