Horde: More logging for config updates.

diff --git a/Engine/Source/Programs/Horde/Horde.Build/Configuration/ConfigService.cs b/Engine/Source/Programs/Horde/Horde.Build/Configuration/ConfigService.cs
--- a/Engine/Source/Programs/Horde/Horde.Build/Configuration/ConfigService.cs
+++ b/Engine/Source/Programs/Horde/Horde.Build/Configuration/ConfigService.cs
@@ -253,7 +253,14 @@ namespace Horde.Build.Configuration
 					snapshot = await CreateSnapshotGuardedAsync(cancellationToken);
 					if (snapshot != null)
 					{
-						await WriteSnapshotAsync(snapshot);
+						try
+						{
+							await WriteSnapshotAsync(snapshot);
+						}
+						catch (Exception ex)
+						{
+							_logger.LogError(ex, "Error while updating config: {Message}", ex.Message);
+						}
 					}
 				}
 				await Task.Delay(_tickInterval, cancellationToken);
@@ -403,6 +410,7 @@ namespace Horde.Build.Configuration
 		async Task WriteSnapshotAsync(ConfigSnapshot snapshot)
 		{
 			ReadOnlyMemory<byte> data = SerializeSnapshot(snapshot);
+			_logger.LogInformation("Publishing new config snapshot ({Size})", data.Length);
 			await _redisService.GetDatabase().StringSetAsync(_snapshotKey, data);
 			await _redisService.PublishAsync(_updateChannel, RedisValue.EmptyString);
 			_logger.LogInformation("Published new config snapshot: {@Info}", GetSnapshotInfo(data.Span, snapshot)); 
