Horde: Extra logging for config updates.

diff --git a/Engine/Source/Programs/Horde/Horde.Build/Configuration/ConfigService.cs b/Engine/Source/Programs/Horde/Horde.Build/Configuration/ConfigService.cs
--- a/Engine/Source/Programs/Horde/Horde.Build/Configuration/ConfigService.cs
+++ b/Engine/Source/Programs/Horde/Horde.Build/Configuration/ConfigService.cs
@@ -410,10 +410,15 @@ namespace Horde.Build.Configuration
 		async Task WriteSnapshotAsync(ConfigSnapshot snapshot)
 		{
 			ReadOnlyMemory<byte> data = SerializeSnapshot(snapshot);
-			_logger.LogInformation("Publishing new config snapshot ({Size})", data.Length);
+
+			IoHash hash = IoHash.Compute(data.Span);
+			_logger.LogInformation("Publishing new config snapshot (hash: {Hash}, server: {Server}, size: {Size})", hash.ToString(), Program.Version.ToString(), data.Length);
+			_logger.LogInformation("Dependencies for config revision {Hash}: {@Info}", hash.ToString(), snapshot.Dependencies);
+
 			await _redisService.GetDatabase().StringSetAsync(_snapshotKey, data);
 			await _redisService.PublishAsync(_updateChannel, RedisValue.EmptyString);
-			_logger.LogInformation("Published new config snapshot: {@Info}", GetSnapshotInfo(data.Span, snapshot)); 
+
+			_logger.LogInformation("Published new config snapshot (hash: {Hash}, server: {Server}, size: {Size})", hash.ToString(), Program.Version.ToString(), data.Length);
 		}
 
 		async Task<ReadOnlyMemory<byte>> ReadSnapshotDataAsync()
diff --git a/Engine/Source/Programs/Horde/Horde.Build/Configuration/ConfigSource.cs b/Engine/Source/Programs/Horde/Horde.Build/Configuration/ConfigSource.cs
--- a/Engine/Source/Programs/Horde/Horde.Build/Configuration/ConfigSource.cs
+++ b/Engine/Source/Programs/Horde/Horde.Build/Configuration/ConfigSource.cs
@@ -317,7 +317,7 @@ namespace Horde.Build.Configuration
 			}
 			else
 			{
-				_logger.LogInformation("Reading {Uri}@{Change} from Perforce", uri, change);
+				_logger.LogInformation("Reading {Uri} at CL {Change} from Perforce", uri, change);
 				using (IPerforceConnection perforce = await ConnectAsync(uri.Host, cancellationToken))
 				{
 					PerforceResponse<PrintRecord<byte[]>> response = await perforce.TryPrintDataAsync($"{uri.AbsolutePath}@{change}", cancellationToken);
