Horde: More config logging.

diff --git a/Engine/Source/Programs/Horde/Horde.Build/Configuration/ConfigService.cs b/Engine/Source/Programs/Horde/Horde.Build/Configuration/ConfigService.cs
--- a/Engine/Source/Programs/Horde/Horde.Build/Configuration/ConfigService.cs
+++ b/Engine/Source/Programs/Horde/Horde.Build/Configuration/ConfigService.cs
@@ -114,7 +114,7 @@ namespace Horde.Build.Configuration
 		readonly ILogger _logger;
 
 		readonly ITicker _ticker;
-		readonly TimeSpan _tickInterval = TimeSpan.FromSeconds(10.0);
+		readonly TimeSpan _tickInterval = TimeSpan.FromMinutes(1.0);
 
 		readonly RedisChannel _updateChannel = "config-update";
 		readonly BackgroundTask _updateTask;
@@ -277,10 +277,19 @@ namespace Horde.Build.Configuration
 				{
 					Task updateTask = updateEvent.Task;
 
-					RedisValue value = await _redisService.GetDatabase().StringGetAsync(_snapshotKey);
-					if (!value.IsNullOrEmpty)
+					try
 					{
-						await UpdateConfigObjectAsync(value);
+						RedisValue value = await _redisService.GetDatabase().StringGetAsync(_snapshotKey);
+						if (!value.IsNullOrEmpty)
+						{
+							await UpdateConfigObjectAsync(value);
+						}
+					}
+					catch (Exception ex)
+					{
+						_logger.LogError(ex, "Exception while updating config from Redis: {Message}", ex.Message);
+						await Task.Delay(TimeSpan.FromMinutes(1.0), cancellationToken);
+						continue;
 					}
 
 					await await Task.WhenAny(updateTask, cancellationTask);
@@ -413,7 +422,6 @@ namespace Horde.Build.Configuration
 
 			IoHash hash = IoHash.Compute(data.Span);
 			_logger.LogInformation("Publishing new config snapshot (hash: {Hash}, server: {Server}, size: {Size})", hash.ToString(), Program.Version.ToString(), data.Length);
-			_logger.LogInformation("Dependencies for config revision {Hash}: {@Info}", hash.ToString(), snapshot.Dependencies);
 
 			await _redisService.GetDatabase().StringSetAsync(_snapshotKey, data);
 			await _redisService.PublishAsync(_updateChannel, RedisValue.EmptyString);
@@ -455,7 +463,7 @@ namespace Horde.Build.Configuration
 			IoHash hash = IoHash.Compute(data.Span);
 			if (hash != state.Hash)
 			{
-				// Update the config object
+				_logger.LogInformation("Updating config state from Redis (hash: {OldHash} -> {NewHash}, new size: {Size})", state.Hash, hash, data.Length);
 				GlobalConfig globalConfig = CreateGlobalConfig(data);
 				Set(hash, globalConfig);
 			}
