Horde: Fix issue with non-worker pods failing to update config.

diff --git a/Engine/Source/Programs/Horde/Horde.Build/Configuration/ConfigService.cs b/Engine/Source/Programs/Horde/Horde.Build/Configuration/ConfigService.cs
--- a/Engine/Source/Programs/Horde/Horde.Build/Configuration/ConfigService.cs
+++ b/Engine/Source/Programs/Horde/Horde.Build/Configuration/ConfigService.cs
@@ -160,7 +160,10 @@ namespace Horde.Build.Configuration
 		/// <inheritdoc/>
 		public async Task StartAsync(CancellationToken cancellationToken)
 		{
-			await _ticker.StartAsync();
+			if (_serverSettings.IsRunModeActive(RunMode.Worker))
+			{
+				await _ticker.StartAsync();
+			}
 			_updateTask.Start();
 		}
 
@@ -168,7 +171,10 @@ namespace Horde.Build.Configuration
 		public async Task StopAsync(CancellationToken cancellationToken)
 		{
 			await _updateTask.StopAsync();
-			await _ticker.StopAsync();
+			if (_serverSettings.IsRunModeActive(RunMode.Worker))
+			{
+				await _ticker.StopAsync();
+			}
 		}
 
 		/// <summary>
diff --git a/Engine/Source/Programs/Horde/Horde.Build/Startup.cs b/Engine/Source/Programs/Horde/Horde.Build/Startup.cs
--- a/Engine/Source/Programs/Horde/Horde.Build/Startup.cs
+++ b/Engine/Source/Programs/Horde/Horde.Build/Startup.cs
@@ -372,6 +372,9 @@ namespace Horde.Build
 			services.AddSingleton<IOptionsFactory<GlobalConfig>>(sp => sp.GetRequiredService<ConfigService>());
 			services.AddSingleton<IOptionsChangeTokenSource<GlobalConfig>>(sp => sp.GetRequiredService<ConfigService>());
 
+			// Always run the hosted config service, regardless of the server role, so we receive updates on the latest config values.
+			services.AddHostedService(provider => provider.GetRequiredService<ConfigService>());
+
 			// Auditing
 			services.AddSingleton<IAuditLog<AgentId>>(sp => sp.GetRequiredService<IAuditLogFactory<AgentId>>().Create("Agents.Log", "AgentId"));
 
@@ -669,7 +672,6 @@ namespace Horde.Build
 				{
 					services.AddHostedService(provider => provider.GetRequiredService<SlackNotificationSink>());
 				}
-				services.AddHostedService(provider => provider.GetRequiredService<ConfigService>());
 				services.AddHostedService<TelemetryService>();
 				services.AddHostedService(provider => provider.GetRequiredService<DeviceService>());
 				services.AddHostedService(provider => provider.GetRequiredService<TestDataService>());
