Horde: Additional logging for failed log writes.

diff --git a/Engine/Source/Programs/Horde/Horde.Agent/Utility/JsonRpcLogger.cs b/Engine/Source/Programs/Horde/Horde.Agent/Utility/JsonRpcLogger.cs
--- a/Engine/Source/Programs/Horde/Horde.Agent/Utility/JsonRpcLogger.cs
+++ b/Engine/Source/Programs/Horde/Horde.Agent/Utility/JsonRpcLogger.cs
@@ -496,7 +496,7 @@ namespace Horde.Agent.Parser
 					}
 					catch (Exception ex)
 					{
-						_inner.LogWarning(ex, "Unable to write data to server (log {LogId}, offset {Offset})", _logId, offset);
+						_inner.LogWarning(ex, "Unable to write data to server (log {LogId}, offset {Offset}, length {Length}, lines {StartLine}-{EndLine})", _logId, offset, data.Length, initialLineIndex, lineIndex);
 					}
 					offset += data.Length;
 				}
