Horde: Add support for compute clusters in V2 version of compute service, and allow submitting multiple tasks at once.

diff --git a/Engine/Source/Programs/Horde/Horde.Agent/Commands/Compute/ComputeCommandV2.cs b/Engine/Source/Programs/Horde/Horde.Agent/Commands/Compute/ComputeCommandV2.cs
--- a/Engine/Source/Programs/Horde/Horde.Agent/Commands/Compute/ComputeCommandV2.cs
+++ b/Engine/Source/Programs/Horde/Horde.Agent/Commands/Compute/ComputeCommandV2.cs
@@ -23,10 +23,10 @@ namespace Horde.Agent.Commands
 	[Command("ComputeV2", "Executes a command through the Horde Compute API")]
 	class ComputeCommandV2 : Command
 	{
-		/// <summary>
-		/// Condition for 
-		/// </summary>
-		[CommandLine("-Condition")]
+		[CommandLine("-Cluster=", Description = "Cluster to execute on")]
+		public ClusterId ClusterId { get; set; } = new ClusterId("default");
+
+		[CommandLine("-Condition=", Description = "Match the agent to run on")]
 		public string? Condition { get; set; }
 
 		readonly AgentSettings _settings;
@@ -43,12 +43,12 @@ namespace Horde.Agent.Commands
 		public override async Task<int> ExecuteAsync(ILogger logger)
 		{
 			ServerProfile profile = _settings.GetCurrentServerProfile();
-			logger.LogInformation("Server: {Server}", profile.Name);
+			logger.LogInformation("Connecting to server: {Server}", profile.Name);
 
 			await using ComputeClient client = new ComputeClient(CreateHttpClient, logger);
 			client.Start();
 
-			IComputeRequest<object?> request = await client.AddRequestAsync(null, TestCommandAsync);
+			IComputeRequest<object?> request = await client.AddRequestAsync(ClusterId, null, TestCommandAsync);
 			await request.Result;
 
 			return 0;
diff --git a/Engine/Source/Programs/Horde/Horde.Build.Tests/Compute/ComputeServiceTest.cs b/Engine/Source/Programs/Horde/Horde.Build.Tests/Compute/ComputeServiceTest.cs
--- a/Engine/Source/Programs/Horde/Horde.Build.Tests/Compute/ComputeServiceTest.cs
+++ b/Engine/Source/Programs/Horde/Horde.Build.Tests/Compute/ComputeServiceTest.cs
@@ -8,7 +8,6 @@ using EpicGames.Horde.Compute;
 using EpicGames.Horde.Storage;
 using EpicGames.Serialization;
 using Horde.Build.Agents.Pools;
-using Horde.Build.Projects;
 using Horde.Build.Server;
 using Microsoft.Extensions.DependencyInjection;
 using Microsoft.VisualStudio.TestTools.UnitTesting;
@@ -31,7 +30,7 @@ namespace Horde.Build.Tests.Compute
 		{
 			_clusterConfig = new();
 			_namespaceId = new (_clusterConfig.NamespaceId);
-			_clusterId = new (_clusterConfig.Id);
+			_clusterId = _clusterConfig.Id;
 		}
 
 		[TestInitialize]
diff --git a/Engine/Source/Programs/Horde/Horde.Build/Compute/V1/ComputeService.cs b/Engine/Source/Programs/Horde/Horde.Build/Compute/V1/ComputeService.cs
--- a/Engine/Source/Programs/Horde/Horde.Build/Compute/V1/ComputeService.cs
+++ b/Engine/Source/Programs/Horde/Horde.Build/Compute/V1/ComputeService.cs
@@ -104,7 +104,7 @@ namespace Horde.Build.Compute.V1
 
 			public ClusterInfo(ComputeClusterConfig config)
 			{
-				Id = new ClusterId(config.Id);
+				Id = config.Id;
 				NamespaceId = new NamespaceId(config.NamespaceId);
 				RequestBucketId = new BucketId(config.RequestBucketId);
 				ResponseBucketId = new BucketId(config.ResponseBucketId);
@@ -232,7 +232,7 @@ namespace Horde.Build.Compute.V1
 
 		ComputeClusterConfig? GetCluster(ClusterId clusterId)
 		{
-			return _globalConfig.CurrentValue.Compute.FirstOrDefault(x => new ClusterId(x.Id) == clusterId);
+			return _globalConfig.CurrentValue.Compute.FirstOrDefault(x => x.Id == clusterId);
 		}
 
 		/// <inheritdoc/>
diff --git a/Engine/Source/Programs/Horde/Horde.Build/Compute/V2/ComputeControllerV2.cs b/Engine/Source/Programs/Horde/Horde.Build/Compute/V2/ComputeControllerV2.cs
--- a/Engine/Source/Programs/Horde/Horde.Build/Compute/V2/ComputeControllerV2.cs
+++ b/Engine/Source/Programs/Horde/Horde.Build/Compute/V2/ComputeControllerV2.cs
@@ -3,12 +3,8 @@
 using System;
 using System.Collections.Generic;
 using System.Net;
-using System.Threading;
-using System.Threading.Tasks;
 using EpicGames.Core;
-using EpicGames.Horde.Common;
 using EpicGames.Horde.Compute;
-using EpicGames.Horde.Compute.Impl;
 using Horde.Build.Acls;
 using Horde.Build.Server;
 using Horde.Build.Utilities;
@@ -31,24 +27,45 @@ namespace Horde.Build.Compute.V2
 		/// <summary>
 		/// Port to connect on
 		/// </summary>
-		public int RemotePort { get; set; } = 4000;
+		public int? RemotePort { get; set; }
 
 		/// <summary>
-		/// Base-64 encoded cryptographic nonce to identify the request
+		/// Cryptographic nonce to identify the request, as a hex string
 		/// </summary>
 		public string Nonce { get; set; } = String.Empty;
 
 		/// <summary>
-		/// Base-64 encoded AES key for the channel
+		/// AES key for the channel, as a hex string
 		/// </summary>
 		public string AesKey { get; set; } = String.Empty;
 
 		/// <summary>
-		/// Base-64 encoded AES IV for the channel
+		/// AES IV for the channel, as a hex string
 		/// </summary>
 		public string AesIv { get; set; } = String.Empty;
 	}
 
+	/// <summary>
+	/// Request a machine to execute compute requests
+	/// </summary>
+	public class AddComputeTasksRequest
+	{
+		/// <summary>
+		/// Condition to identify machines that can execute the request
+		/// </summary>
+		public Requirements? Requirements { get; set; }
+
+		/// <summary>
+		/// Port to connect on
+		/// </summary>
+		public int RemotePort { get; set; }
+
+		/// <summary>
+		/// List of tasks to add
+		/// </summary>
+		public List<AddComputeTaskRequest> Tasks { get; set; } = new List<AddComputeTaskRequest>();
+	}
+
 	/// <summary>
 	/// Controller for the /api/v2/compute endpoint
 	/// </summary>
@@ -72,22 +89,23 @@ namespace Horde.Build.Compute.V2
 		/// <summary>
 		/// Add tasks to be executed remotely
 		/// </summary>
+		/// <param name="clusterId">Id of the compute cluster</param>
 		/// <param name="request">The request parameters</param>
 		/// <returns></returns>
 		[HttpPost]
 		[Authorize]
-		[Route("/api/v2/compute")]
-		public ActionResult AddTasksAsync([FromBody] AddComputeTaskRequest request)
+		[Route("/api/v2/compute/{clusterId}")]
+		public ActionResult AddTasksAsync(ClusterId clusterId, [FromBody] AddComputeTasksRequest request)
 		{
-			if(!_globalConfig.Value.Authorize(AclAction.AddComputeTasks, User))
+			ComputeClusterConfig? clusterConfig;
+			if (!_globalConfig.Value.TryGetComputeCluster(clusterId, out clusterConfig))
 			{
-				return Forbid(AclAction.AddComputeTasks);
+				return NotFound(clusterId);
+			}
+			if(!clusterConfig.Authorize(AclAction.AddComputeTasks, User))
+			{
+				return Forbid(AclAction.AddComputeTasks, clusterId);
 			}
-
-			Requirements requirements = request.Requirements ?? new Requirements();
-			byte[] nonce = StringUtils.ParseHexString(request.Nonce);
-			byte[] aesKey = StringUtils.ParseHexString(request.AesKey);
-			byte[] aesIv = StringUtils.ParseHexString(request.AesIv);
 
 			IPAddress? remoteIp = HttpContext.Connection.RemoteIpAddress;
 			if (remoteIp == null)
@@ -95,7 +113,17 @@ namespace Horde.Build.Compute.V2
 				return BadRequest("Missing remote IP address");
 			}
 
-			_computeService.AddRequest(requirements, remoteIp.ToString(), request.RemotePort, nonce, aesKey, aesIv);
+			foreach (AddComputeTaskRequest taskRequest in request.Tasks)
+			{
+				Requirements requirements = taskRequest.Requirements ?? request.Requirements ?? new Requirements();
+
+				int port = taskRequest.RemotePort ?? request.RemotePort;
+				byte[] nonce = StringUtils.ParseHexString(taskRequest.Nonce);
+				byte[] aesKey = StringUtils.ParseHexString(taskRequest.AesKey);
+				byte[] aesIv = StringUtils.ParseHexString(taskRequest.AesIv);
+
+				_computeService.AddRequest(clusterId, requirements, remoteIp.ToString(), port, nonce, aesKey, aesIv);
+			}
 			return Ok();
 		}
 	}
diff --git a/Engine/Source/Programs/Horde/Horde.Build/Compute/V2/ComputeServiceV2.cs b/Engine/Source/Programs/Horde/Horde.Build/Compute/V2/ComputeServiceV2.cs
--- a/Engine/Source/Programs/Horde/Horde.Build/Compute/V2/ComputeServiceV2.cs
+++ b/Engine/Source/Programs/Horde/Horde.Build/Compute/V2/ComputeServiceV2.cs
@@ -12,12 +12,15 @@ using Google.Protobuf.WellKnownTypes;
 using Horde.Build.Agents;
 using Horde.Build.Agents.Leases;
 using Horde.Build.Agents.Pools;
+using Horde.Build.Server;
 using Horde.Build.Tasks;
 using Horde.Build.Utilities;
 using HordeCommon;
 using HordeCommon.Rpc.Tasks;
 using Microsoft.AspNetCore.Http;
+using Microsoft.AspNetCore.Mvc.TagHelpers;
 using Microsoft.Extensions.Logging;
+using Microsoft.Extensions.Options;
 
 namespace Horde.Build.Compute.V2
 {
@@ -59,8 +62,14 @@ namespace Horde.Build.Compute.V2
 			}
 		}
 
+		class ClusterInfo
+		{
+			public Dictionary<IoHash, RequestQueue> Queues { get; } = new Dictionary<IoHash, RequestQueue>();
+		}
+
 		readonly object _lockObject = new object();
-		readonly Dictionary<IoHash, RequestQueue> _queues = new Dictionary<IoHash, RequestQueue>();
+		readonly IOptionsMonitor<GlobalConfig> _globalConfig;
+		readonly Dictionary<ClusterId, ClusterInfo> _clusters = new Dictionary<ClusterId, ClusterInfo>();
 		readonly ILogger _logger;
 
 		/// <inheritdoc/>
@@ -72,33 +81,42 @@ namespace Horde.Build.Compute.V2
 		/// <summary>
 		/// Constructor
 		/// </summary>
-		/// <param name="logger"></param>
-		public ComputeServiceV2(ILogger<ComputeServiceV2> logger)
+		public ComputeServiceV2(IOptionsMonitor<GlobalConfig> globalConfig, ILogger<ComputeServiceV2> logger)
 		{
+			_globalConfig = globalConfig;
 			_logger = logger;
 		}
 
 		/// <summary>
 		/// Adds a new compute request
 		/// </summary>
+		/// <param name="clusterId">The compute cluster id</param>
 		/// <param name="requirements">Requirements for the agent to serve the request</param>
 		/// <param name="ip">IP address to connect to</param>
 		/// <param name="port">Remote port to connect to</param>
 		/// <param name="nonce">Cryptographic nonce used to identify the request</param>
 		/// <param name="aesKey">Key for AES encryption on the compute channel</param>
 		/// <param name="aesIv">Initialization vector for AES encryption</param>
-		public void AddRequest(Requirements requirements, string ip, int port, ReadOnlyMemory<byte> nonce, ReadOnlyMemory<byte> aesKey, ReadOnlyMemory<byte> aesIv)
+		public void AddRequest(ClusterId clusterId, Requirements requirements, string ip, int port, ReadOnlyMemory<byte> nonce, ReadOnlyMemory<byte> aesKey, ReadOnlyMemory<byte> aesIv)
 		{
 			CbObject obj = CbSerializer.Serialize(requirements);
 			IoHash hash = IoHash.Compute(obj.GetView().Span);
 			lock (_lockObject)
 			{
+				ClusterInfo? cluster;
+				if (!_clusters.TryGetValue(clusterId, out cluster))
+				{
+					cluster = new ClusterInfo();
+					_clusters.Add(clusterId, cluster);
+				}
+
 				RequestQueue? queue;
-				if (!_queues.TryGetValue(hash, out queue))
+				if (!cluster.Queues.TryGetValue(hash, out queue))
 				{
 					queue = new RequestQueue(hash, requirements);
-					_queues.Add(hash, queue);
+					cluster.Queues.Add(hash, queue);
 				}
+
 				queue.Requests.Enqueue(new Request(ip, port, nonce, aesKey, aesIv));
 			}
 		}
@@ -113,20 +131,44 @@ namespace Horde.Build.Compute.V2
 
 				lock (_lockObject)
 				{
-					foreach ((IoHash hash, RequestQueue queue) in _queues)
+					ClusterId? removeClusterId = null;
+
+					GlobalConfig globalConfig = _globalConfig.CurrentValue;
+					foreach ((ClusterId clusterId, ClusterInfo clusterInfo) in _clusters)
 					{
-						if (agent.MeetsRequirements(queue.Requirements))
+						ComputeClusterConfig? clusterConfig;
+						if (!globalConfig.TryGetComputeCluster(clusterId, out clusterConfig))
 						{
-							request = queue.Requests.Dequeue();
-							requestQueue = queue;
+							removeClusterId ??= clusterId;
+						}
+						else if (clusterConfig.Condition == null || agent.SatisfiesCondition(clusterConfig.Condition))
+						{
+							foreach ((IoHash hash, RequestQueue queue) in clusterInfo.Queues)
+							{
+								if (agent.MeetsRequirements(queue.Requirements))
+								{
+									request = queue.Requests.Dequeue();
+									requestQueue = queue;
+
+									if (queue.Requests.Count == 0)
+									{
+										clusterInfo.Queues.Remove(hash);
+									}
+									break;
+								}
+							}
 
-							if (queue.Requests.Count == 0)
+							if (clusterInfo.Queues.Count == 0)
 							{
-								_queues.Remove(hash);
+								removeClusterId ??= clusterId;
 							}
-							break;
 						}
 					}
+
+					if (removeClusterId != null)
+					{
+						_clusters.Remove(removeClusterId.Value);
+					}
 				}
 
 				if (request != null)
diff --git a/Engine/Source/Programs/Horde/Horde.Build/Server/GlobalConfig.cs b/Engine/Source/Programs/Horde/Horde.Build/Server/GlobalConfig.cs
--- a/Engine/Source/Programs/Horde/Horde.Build/Server/GlobalConfig.cs
+++ b/Engine/Source/Programs/Horde/Horde.Build/Server/GlobalConfig.cs
@@ -10,6 +10,7 @@ using System.Security.Claims;
 using System.Text.Json.Serialization;
 using EpicGames.Core;
 using EpicGames.Horde.Common;
+using EpicGames.Horde.Compute;
 using EpicGames.Perforce;
 using EpicGames.Serialization;
 using Horde.Build.Acls;
@@ -138,6 +139,7 @@ namespace Horde.Build.Server
 		private readonly Dictionary<ProjectId, ProjectConfig> _projectLookup = new Dictionary<ProjectId, ProjectConfig>();
 		private readonly Dictionary<StreamId, StreamConfig> _streamLookup = new Dictionary<StreamId, StreamConfig>();
 		private readonly Dictionary<ToolId, ToolConfig> _toolLookup = new Dictionary<ToolId, ToolConfig>();
+		private readonly Dictionary<ClusterId, ComputeClusterConfig> _computeClusterLookup = new Dictionary<ClusterId, ComputeClusterConfig>();
 
 		/// <summary>
 		/// Called after the config file has been read
@@ -168,6 +170,13 @@ namespace Horde.Build.Server
 				tool.PostLoad(this);
 			}
 
+			_computeClusterLookup.Clear();
+			foreach (ComputeClusterConfig computeCluster in Compute)
+			{
+				_computeClusterLookup.Add(computeCluster.Id, computeCluster);
+				computeCluster.PostLoad(this);
+			}
+
 			Storage.PostLoad(this);
 		}
 
@@ -195,6 +204,14 @@ namespace Horde.Build.Server
 		/// <returns>True if the stream configuration was found</returns>
 		public bool TryGetTool(ToolId toolId, [NotNullWhen(true)] out ToolConfig? config) => _toolLookup.TryGetValue(toolId, out config);
 
+		/// <summary>
+		/// Attempts to get compute cluster configuration from this object
+		/// </summary>
+		/// <param name="clusterId">Compute cluster id</param>
+		/// <param name="config">Receives the cluster configuration on success</param>
+		/// <returns>True on success</returns>
+		public bool TryGetComputeCluster(ClusterId clusterId, [NotNullWhen(true)] out ComputeClusterConfig? config) => _computeClusterLookup.TryGetValue(clusterId, out config);
+
 		/// <summary>
 		/// Authorizes a user to perform a given action
 		/// </summary>
@@ -292,10 +309,16 @@ namespace Horde.Build.Server
 	[DebuggerDisplay("{Id}")]
 	public class ComputeClusterConfig
 	{
+		/// <summary>
+		/// The owning global config instance
+		/// </summary>
+		[JsonIgnore]
+		public GlobalConfig GlobalConfig { get; private set; } = null!;
+
 		/// <summary>
 		/// Name of the partition
 		/// </summary>
-		public string Id { get; set; } = "default";
+		public ClusterId Id { get; set; } = new ClusterId("default");
 
 		/// <summary>
 		/// Name of the namespace to use
@@ -312,10 +335,34 @@ namespace Horde.Build.Server
 		/// </summary>
 		public string ResponseBucketId { get; set; } = "responses";
 
+		/// <summary>
+		/// Filter for agents to include
+		/// </summary>
+		public Condition? Condition { get; set; }
+
 		/// <summary>
 		/// Access control list
 		/// </summary>
 		public AclConfig? Acl { get; set; }
+
+		/// <summary>
+		/// Callback post loading this config file
+		/// </summary>
+		/// <param name="globalConfig">The global config instance</param>
+		public void PostLoad(GlobalConfig globalConfig)
+		{
+			GlobalConfig = globalConfig;
+		}
+
+		/// <summary>
+		/// Authorizes a user to perform a given action
+		/// </summary>
+		/// <param name="action">The action being performed</param>
+		/// <param name="user">The principal to validate</param>
+		public bool Authorize(AclAction action, ClaimsPrincipal user)
+		{
+			return Acl?.Authorize(action, user) ?? GlobalConfig.Authorize(action, user);
+		}
 	}
 
 	/// <summary>
diff --git a/Engine/Source/Programs/Horde/Horde.Build/Utilities/HordeControllerBase.cs b/Engine/Source/Programs/Horde/Horde.Build/Utilities/HordeControllerBase.cs
--- a/Engine/Source/Programs/Horde/Horde.Build/Utilities/HordeControllerBase.cs
+++ b/Engine/Source/Programs/Horde/Horde.Build/Utilities/HordeControllerBase.cs
@@ -1,6 +1,7 @@
 // Copyright Epic Games, Inc. All Rights Reserved.
 
 using EpicGames.Core;
+using EpicGames.Horde.Compute;
 using EpicGames.Horde.Storage;
 using Horde.Build.Acls;
 using Horde.Build.Agents;
@@ -69,6 +70,15 @@ namespace Horde.Build.Utilities
 			return Forbid(action, "agent {AgentId}", agentId);
 		}
 
+		/// <summary>
+		/// Returns a 403 (forbidden) response with the given action and object
+		/// </summary>
+		[NonAction]
+		protected ActionResult Forbid(AclAction action, ClusterId clusterId)
+		{
+			return Forbid(action, "cluster {ClusterId}", clusterId);
+		}
+
 		/// <summary>
 		/// Returns a 403 (forbidden) response with the given action and object
 		/// </summary>
@@ -159,6 +169,15 @@ namespace Horde.Build.Utilities
 			return NotFound("Lease {LeaseId} not found for agent {AgentId}", leaseId, agentId);
 		}
 
+		/// <summary>
+		/// Returns a 404 response for the given object
+		/// </summary>
+		[NonAction]
+		protected ActionResult NotFound(ClusterId clusterId)
+		{
+			return NotFound("Cluster {ClusterId} not found", clusterId);
+		}
+
 		/// <summary>
 		/// Returns a 404 response for the given object
 		/// </summary>
diff --git a/Engine/Source/Programs/Shared/EpicGames.Horde/Compute/ComputeChannel.cs b/Engine/Source/Programs/Shared/EpicGames.Horde/Compute/ComputeChannel.cs
--- a/Engine/Source/Programs/Shared/EpicGames.Horde/Compute/ComputeChannel.cs
+++ b/Engine/Source/Programs/Shared/EpicGames.Horde/Compute/ComputeChannel.cs
@@ -2,18 +2,17 @@
 
 using EpicGames.Core;
 using EpicGames.Serialization;
-using Microsoft.Extensions.Caching.Memory;
 using System;
 using System.Collections.Generic;
-using System.Diagnostics;
 using System.IO;
-using System.Net.Http.Headers;
 using System.Net.Sockets;
 using System.Reflection;
 using System.Security.Cryptography;
 using System.Threading;
 using System.Threading.Tasks;
 
+#pragma warning disable CA5401
+
 namespace EpicGames.Horde.Compute
 {
 	/// <summary>
diff --git a/Engine/Source/Programs/Shared/EpicGames.Horde/Compute/ComputeClient.cs b/Engine/Source/Programs/Shared/EpicGames.Horde/Compute/ComputeClient.cs
--- a/Engine/Source/Programs/Shared/EpicGames.Horde/Compute/ComputeClient.cs
+++ b/Engine/Source/Programs/Shared/EpicGames.Horde/Compute/ComputeClient.cs
@@ -39,12 +39,19 @@ namespace EpicGames.Horde.Compute
 		class AddComputeTaskRequest
 		{
 			public Requirements? Requirements { get; set; }
-			public int RemotePort { get; set; }
+			public int? RemotePort { get; set; }
 			public string Nonce { get; set; } = String.Empty;
 			public string AesKey { get; set; } = String.Empty;
 			public string AesIv { get; set; } = String.Empty;
 		}
 
+		class AddComputeTasksRequest
+		{
+			public Requirements? Requirements { get; set; }
+			public int RemotePort { get; set; }
+			public List<AddComputeTaskRequest> Tasks { get; set; } = new List<AddComputeTaskRequest>();
+		}
+
 		abstract class RequestInfo : IDisposable
 		{
 			public ReadOnlyMemory<byte> Nonce { get; }
@@ -160,7 +167,7 @@ namespace EpicGames.Horde.Compute
 		/// </summary>
 		public void Start()
 		{
-			_listener = new TcpListener(IPAddress.Any, 0);
+			_listener = new TcpListener(IPAddress.IPv6Any, 0);
 			_listener.Start();
 			_listenTask = Task.Run(() => ListenAsync(_listener, _cancellationSource.Token));
 		}
@@ -183,29 +190,33 @@ namespace EpicGames.Horde.Compute
 		/// <summary>
 		/// Adds a new remote request
 		/// </summary>
+		/// <param name="clusterId">Cluster to execute the request</param>
 		/// <param name="requirements">Requirements for the agent</param>
 		/// <param name="handler">Handler for the connection</param>
-		public async Task<IComputeRequest<TResult>> AddRequestAsync<TResult>(Requirements? requirements, Func<ComputeChannel, CancellationToken, Task<TResult>> handler)
+		public async Task<IComputeRequest<TResult>> AddRequestAsync<TResult>(ClusterId clusterId, Requirements? requirements, Func<ComputeChannel, CancellationToken, Task<TResult>> handler)
 		{
 			RequestInfo<TResult> requestInfo = new RequestInfo<TResult>(handler, _cancellationSource.Token);
-			await AddRequestAsync(requirements, requestInfo);
+			await AddRequestAsync(clusterId, requirements, requestInfo);
 			return requestInfo;
 		}
 
-		async Task AddRequestAsync(Requirements? requirements, RequestInfo requestInfo)
+		async Task AddRequestAsync(ClusterId clusterId, Requirements? requirements, RequestInfo requestInfo)
 		{
 			try
 			{
 				HttpClient client = _createHttpClient();
 
-				AddComputeTaskRequest request = new AddComputeTaskRequest();
+				AddComputeTaskRequest task = new AddComputeTaskRequest();
+				task.Nonce = StringUtils.FormatHexString(requestInfo.Nonce.Span);
+				task.AesKey = StringUtils.FormatHexString(requestInfo.AesKey.Span);
+				task.AesIv = StringUtils.FormatHexString(requestInfo.AesIv.Span);
+
+				AddComputeTasksRequest request = new AddComputeTasksRequest();
 				request.Requirements = requirements;
 				request.RemotePort = ((IPEndPoint)_listener!.LocalEndpoint).Port;
-				request.Nonce = StringUtils.FormatHexString(requestInfo.Nonce.Span);
-				request.AesKey = StringUtils.FormatHexString(requestInfo.AesKey.Span);
-				request.AesIv = StringUtils.FormatHexString(requestInfo.AesIv.Span);
+				request.Tasks.Add(task);
 
-				using (HttpResponseMessage response = await client.PostAsync("api/v2/compute", request, _cancellationSource.Token))
+				using (HttpResponseMessage response = await client.PostAsync($"api/v2/compute/{clusterId}", request, _cancellationSource.Token))
 				{
 					response.EnsureSuccessStatusCode();
 					_requests.Add(requestInfo.Nonce, requestInfo);
