Horde: AutomationView fixes and error boundary
#rnx
#skipci

diff --git a/Engine/Source/Programs/Horde/HordeDashboard/src/components/AutomationView.tsx b/Engine/Source/Programs/Horde/HordeDashboard/src/components/AutomationView.tsx
--- a/Engine/Source/Programs/Horde/HordeDashboard/src/components/AutomationView.tsx
+++ b/Engine/Source/Programs/Horde/HordeDashboard/src/components/AutomationView.tsx
@@ -7,7 +7,7 @@ import { action, makeObservable, observable } from 'mobx';
 import { observer } from 'mobx-react-lite';
 import moment from 'moment';
 import React, { useState } from 'react';
-import { useLocation, useNavigate } from 'react-router-dom';
+import { useSearchParams } from 'react-router-dom';
 import backend from '../backend';
 import { GetTestDataRefResponse, GetTestMetaResponse, GetTestResponse, GetTestStreamResponse, GetTestSuiteResponse, TestOutcome } from '../backend/Api';
 import dashboard, { StatusColor } from '../backend/Dashboard';
@@ -17,6 +17,7 @@ import { getHumanTime, getShortNiceTime, msecToElapsed } from '../base/utilities
 import { hordeClasses, modeColors } from '../styles/Styles';
 import { AutomationSuiteDetails } from './AutomationSuiteDetails';
 import { Breadcrumbs } from './Breadcrumbs';
+import ErrorBoundary from './ErrorBoundary';
 import { TopNav } from './TopNav';
 
 // some aliases
@@ -56,6 +57,9 @@ type TestDataState = {
 
 type FilterState = "Success" | "Failed" | "Consecutive Failures" | "Skipped";
 
+// used to regenerate combo boxes upon automation project switches
+let multiComboBoxId = 0;
+
 class TestDataHandler {
 
    constructor(search: URLSearchParams) {
@@ -66,7 +70,33 @@ class TestDataHandler {
 
    setAutomation(automation: string) {
       this.state.automation = automation;
+
+      if (automation) {
+         const streams = this.getAutomationStreams(automation);
+         this.state.streams = this.state.streams?.filter(s => streams.indexOf(s) !== -1);
+         if (!this.state.streams?.length) {
+            this.state.streams = undefined;
+         }
+
+         //this.state.tests = undefined;
+         //this.state.suites = undefined;
+
+         const tests = this.streamTests;
+         this.state.tests = this.state.tests?.filter(test => !!tests.find(t => t.name === test));
+         if (!this.state.tests?.length) {
+            this.state.tests = undefined;
+         }
+
+         const suites = this.streamSuites;
+         this.state.suites = this.state.suites?.filter(suite => !!suites.find(s => s.name === suite));
+         if (!this.state.suites?.length) {
+            this.state.suites = undefined;
+         }
+
+      }
+
       if (this.updateSearch()) {
+         multiComboBoxId++;
          this.setUpdated();
       }
    }
@@ -1405,7 +1435,7 @@ const AutomationChooser: React.FC<{ handler: TestDataHandler }> = observer(({ ha
 });
 
 
-const MultiOptionChooser: React.FC<{ options: IComboBoxOption[], initialKeysIn: string[], updateKeys: (selectedKeys: string[]) => void }> = observer(({ options, initialKeysIn, updateKeys }) => {
+const MultiOptionChooser: React.FC<{ options: IComboBoxOption[], initialKeysIn: string[], updateKeys: (selectedKeys: string[]) => void }> = ({ options, initialKeysIn, updateKeys }) => {
 
    let initialKeys = [...initialKeysIn];
 
@@ -1456,8 +1486,8 @@ const MultiOptionChooser: React.FC<{ options: IComboBoxOption[], initialKeysIn:
 
    const comboBoxStyles: Partial<IComboBoxStyles> = { root: { width: 270 } };
 
-   return <ComboBox placeholder="None" defaultSelectedKey={initialKeys} multiSelect options={options} onChange={onChange} styles={comboBoxStyles} />
-});
+   return <ComboBox key={`multi_option_id_${multiComboBoxId}` } placeholder="None" defaultSelectedKey={initialKeys} multiSelect options={options} onChange={onChange} styles={comboBoxStyles} />
+};
 
 const TestChooser: React.FC<{ handler: TestDataHandler }> = observer(({ handler }) => {
 
@@ -1834,14 +1864,14 @@ const AutomationOperationsBar: React.FC<{ handler: TestDataHandler }> = observer
 export const AutomationView: React.FC = observer(() => {
 
    const [state, setState] = useState<{ handler?: TestDataHandler, search?: string }>({});
-   const location = useLocation();
-   const navigate = useNavigate();
+   const [searchParams, setSearchParams] = useSearchParams();
 
    let handler = state.handler;
 
    // subscribe
    if (!handler) {
-      handler = new TestDataHandler(new URLSearchParams(location.search));
+      handler = new TestDataHandler(new URLSearchParams(searchParams));
+      if (handler.updated) { }
       setState({ handler: handler, search: handler.search.toString() });
       return null;
    }
@@ -1852,8 +1882,8 @@ export const AutomationView: React.FC = observer(() => {
    const csearch = handler.search.toString();
 
    if (state.search !== csearch) {
-      location.search = csearch;
-      navigate(location, { replace: true });
+      // this causes an error in dev mode, might be something to fix when move to new router (createBrowserRouter)
+      setSearchParams(csearch);
       setState({ ...state, search: csearch });
       return null;
    }
@@ -1876,25 +1906,26 @@ export const AutomationView: React.FC = observer(() => {
          {(!!suiteRefs?.length && !!metaData && !!suite) && <AutomationSuiteDetails suite={suite} suiteRefs={suiteRefs} metaData={metaData} onClose={() => handler?.setSuiteRef(undefined)} />}
          <TopNav />
          <Breadcrumbs items={[{ text: 'Automation' }]} />
-         <Stack horizontalAlign="center" grow styles={{ root: { width: "100%", padding: 12, backgroundColor: modeColors.background } }}>
-            <Stack styles={{ root: { width: 1774 } }}>
-               {!handler.loaded && <Stack>
-               </Stack>}
-               {handler.loaded && <Stack style={{ paddingTop: 8 }}>
-                  <Stack horizontal >
-                     <AutomationSidebarLeft handler={handler} />
-                     <Stack grow>
-                        <Stack style={{ paddingRight: 62 }}>
-                           <AutomationOperationsBar handler={handler} />
+         <ErrorBoundary>
+            <Stack horizontalAlign="center" grow styles={{ root: { width: "100%", padding: 12, backgroundColor: modeColors.background } }}>
+               <Stack styles={{ root: { width: 1774 } }}>
+                  {!handler.loaded && <Stack>
+                  </Stack>}
+                  {handler.loaded && <Stack style={{ paddingTop: 8 }}>
+                     <Stack horizontal >
+                        <AutomationSidebarLeft handler={handler} />
+                        <Stack grow>
+                           <Stack style={{ paddingRight: 62 }}>
+                              <AutomationOperationsBar handler={handler} />
+                           </Stack>
+                           <AutomationCenter handler={handler} />
                         </Stack>
-                        <AutomationCenter handler={handler} />
                      </Stack>
-                  </Stack>
-               </Stack>}
+                  </Stack>}
+               </Stack>
             </Stack>
-         </Stack>
+         </ErrorBoundary>
       </Stack>
-
    );
 });
 
diff --git a/Engine/Source/Programs/Horde/HordeDashboard/src/components/ErrorBoundary.tsx b/Engine/Source/Programs/Horde/HordeDashboard/src/components/ErrorBoundary.tsx
new file mode 100644
--- /dev/null
+++ b/Engine/Source/Programs/Horde/HordeDashboard/src/components/ErrorBoundary.tsx
@@ -0,0 +1,46 @@
+import { Stack, Text } from "@fluentui/react";
+import { Component, ErrorInfo, ReactNode } from "react";
+import ErrorHandler from "./ErrorHandler";
+
+interface Props {
+   children?: ReactNode;
+}
+
+interface State {
+   hasError: boolean;
+}
+
+class ErrorBoundary extends Component<Props, State> {
+   public state: State = {
+      hasError: false
+   };
+
+   public static getDerivedStateFromError(_: Error): State {
+      return { hasError: true };
+   }
+
+   public componentDidCatch(error: Error, errorInfo: ErrorInfo) {      
+      console.error("Uncaught Exception:", error, errorInfo);
+      ErrorHandler.set({
+         reason: `${error}`,
+         title: `Uncaught Exception`,
+         message: `Error: ${error}\n\n${errorInfo?.componentStack}`
+
+      }, true);
+
+   }
+
+   public render() {
+
+      if (this.state.hasError) {
+         return <Stack style={{padding: 128}}><Text variant="mediumPlus">Uncaught Exception, please see console for more information</Text></Stack>;
+      }
+                  
+      return this.props.children;
+   }
+
+   static error?: Error;
+   static errorInfo?: ErrorInfo;
+}
+
+export default ErrorBoundary;
\ No newline at end of file
