Horde: Update Redis with a config snapshot when 'update config on startup' is enabled. This prevents the locally read config data being immediately overriden by the Redis version.

diff --git a/Engine/Source/Programs/Horde/Horde.Build/Configuration/ConfigService.cs b/Engine/Source/Programs/Horde/Horde.Build/Configuration/ConfigService.cs
--- a/Engine/Source/Programs/Horde/Horde.Build/Configuration/ConfigService.cs
+++ b/Engine/Source/Programs/Horde/Horde.Build/Configuration/ConfigService.cs
@@ -232,7 +232,9 @@ namespace Horde.Build.Configuration
 			}
 			if (data.Length == 0)
 			{
-				data = SerializeSnapshot(await CreateSnapshotAsync(CancellationToken.None));
+				ConfigSnapshot snapshot = await CreateSnapshotAsync(CancellationToken.None);
+				await WriteSnapshotAsync(snapshot);
+				data = SerializeSnapshot(snapshot);
 			}
 			return new ConfigState(IoHash.Compute(data.Span), CreateGlobalConfig(data));
 		}
