Horde: Prevent steps that have never succeeded being excluded from issue reports.

diff --git a/Engine/Source/Programs/Horde/Horde.Build/Issues/IssueReportService.cs b/Engine/Source/Programs/Horde/Horde.Build/Issues/IssueReportService.cs
--- a/Engine/Source/Programs/Horde/Horde.Build/Issues/IssueReportService.cs
+++ b/Engine/Source/Programs/Horde/Horde.Build/Issues/IssueReportService.cs
@@ -176,7 +176,7 @@ namespace Horde.Build.Issues
 						IssueReport report = new IssueReport(streamConfig.Id, workflowConfig.Id, workflowStats, workflowConfig.TriageChannel, workflowConfig.GroupIssuesByTemplate);
 						foreach (IIssueSpan span in spans)
 						{
-							if (span.LastSuccess != null && span.LastFailure.Annotations.WorkflowId == workflowConfig.Id)
+							if (span.LastFailure.Annotations.WorkflowId == workflowConfig.Id)
 							{
 								report.IssueSpans.Add(span);
 							}
