Horde: Add a ConfigIncludeContext attribute, which pushes an object onto a stack when including parsing include files, allowing custom propogation of permissions to child objects.

diff --git a/Engine/Source/Programs/Horde/Horde.Build/Configuration/ConfigContext.cs b/Engine/Source/Programs/Horde/Horde.Build/Configuration/ConfigContext.cs
--- a/Engine/Source/Programs/Horde/Horde.Build/Configuration/ConfigContext.cs
+++ b/Engine/Source/Programs/Horde/Horde.Build/Configuration/ConfigContext.cs
@@ -30,6 +30,11 @@ namespace Horde.Build.Configuration
 		/// </summary>
 		public Stack<string> ScopeStack { get; } = new Stack<string>();
 
+		/// <summary>
+		/// Stack of objects
+		/// </summary>
+		public Stack<object> IncludeContextStack { get; } = new Stack<object>();
+
 		/// <summary>
 		/// Map of property path to the file declaring a value for it
 		/// </summary>
diff --git a/Engine/Source/Programs/Horde/Horde.Build/Configuration/ConfigType.cs b/Engine/Source/Programs/Horde/Horde.Build/Configuration/ConfigType.cs
--- a/Engine/Source/Programs/Horde/Horde.Build/Configuration/ConfigType.cs
+++ b/Engine/Source/Programs/Horde/Horde.Build/Configuration/ConfigType.cs
@@ -34,6 +34,14 @@ namespace Horde.Build.Configuration
 	{
 	}
 
+	/// <summary>
+	/// Captures the current scope of outer objects in the current variable
+	/// </summary>
+	[AttributeUsage(AttributeTargets.Class)]
+	public sealed class ConfigIncludeContextAttribute : Attribute
+	{
+	}
+
 	/// <summary>
 	/// Attribute used to mark <see cref="Uri"/> properties that are relative to their containing file
 	/// </summary>
@@ -403,6 +411,7 @@ namespace Horde.Build.Configuration
 
 		readonly Type _type;
 		readonly bool _isIncludeRoot;
+		readonly bool _isIncludeContext;
 		readonly Dictionary<string, Property> _nameToProperty = new Dictionary<string, Property>(StringComparer.OrdinalIgnoreCase);
 		readonly Dictionary<string, Property> _nameToIncludeProperty = new Dictionary<string, Property>(StringComparer.OrdinalIgnoreCase);
 		readonly Dictionary<string, ClassConfigType>? _knownTypes;
@@ -415,6 +424,7 @@ namespace Horde.Build.Configuration
 
 			_type = type;
 			_isIncludeRoot = type.GetCustomAttribute<ConfigIncludeRootAttribute>() != null;
+			_isIncludeContext = type.GetCustomAttribute<ConfigIncludeContextAttribute>() != null;
 
 			// Find all the direct include properties
 			PropertyInfo[] propertyInfos = type.GetProperties(BindingFlags.Instance | BindingFlags.Public | BindingFlags.GetProperty);
@@ -598,6 +608,11 @@ namespace Horde.Build.Configuration
 
 		async Task ParseIncludesAsync(JsonObject jsonObject, object targetObject, ClassConfigType targetType, ConfigContext context, CancellationToken cancellationToken)
 		{
+			if (_isIncludeContext)
+			{
+				context.IncludeContextStack.Push(JsonSerializer.Deserialize(jsonObject, _type, context.JsonOptions)!);
+			}
+
 			foreach ((string name, JsonNode? node) in jsonObject)
 			{
 				if (_nameToIncludeProperty.TryGetValue(name, out Property? property) && node != null)
@@ -605,6 +620,11 @@ namespace Horde.Build.Configuration
 					await property.ParseIncludesAsync(node, targetObject, targetType, context, cancellationToken);
 				}
 			}
+
+			if (_isIncludeContext)
+			{
+				context.IncludeContextStack.Pop();
+			}
 		}
 	}
 }
diff --git a/Engine/Source/Programs/Horde/Horde.Build/Server/GlobalConfig.cs b/Engine/Source/Programs/Horde/Horde.Build/Server/GlobalConfig.cs
--- a/Engine/Source/Programs/Horde/Horde.Build/Server/GlobalConfig.cs
+++ b/Engine/Source/Programs/Horde/Horde.Build/Server/GlobalConfig.cs
@@ -36,6 +36,7 @@ namespace Horde.Build.Server
 	/// <summary>
 	/// Directive to merge config data from another source
 	/// </summary>
+	[ConfigIncludeContext]
 	public class ConfigInclude
 	{
 		/// <summary>
