Horde: Support filtering on artifacts panel
#skipci
#rnx

diff --git a/Engine/Source/Programs/Horde/HordeDashboard/src/components/jobDetailsV2/JobDetailArtifactsV2.tsx b/Engine/Source/Programs/Horde/HordeDashboard/src/components/jobDetailsV2/JobDetailArtifactsV2.tsx
--- a/Engine/Source/Programs/Horde/HordeDashboard/src/components/jobDetailsV2/JobDetailArtifactsV2.tsx
+++ b/Engine/Source/Programs/Horde/HordeDashboard/src/components/jobDetailsV2/JobDetailArtifactsV2.tsx
@@ -1,7 +1,7 @@
 // Copyright Epic Games, Inc. All Rights Reserved.
 
-import { Stack, Text, IColumn, mergeStyleSets, Icon, DetailsList, Selection, SelectionMode, DetailsListLayoutMode, ScrollablePane, ScrollbarVisibility, StickyPositionType, IDetailsListProps, IDetailsHeaderStyles, Sticky, DetailsHeader, PrimaryButton, SpinnerSize, Spinner, Link } from '@fluentui/react';
-import React, { useEffect } from 'react';
+import { Stack, Text, IColumn, mergeStyleSets, Icon, DetailsList, Selection, SelectionMode, DetailsListLayoutMode, ScrollablePane, ScrollbarVisibility, StickyPositionType, IDetailsListProps, IDetailsHeaderStyles, Sticky, DetailsHeader, PrimaryButton, SpinnerSize, Spinner, Link, TextField } from '@fluentui/react';
+import React, { useEffect, useState } from 'react';
 import { ArtifactData, GetArtifactResponse, GetArtifactZipRequest } from '../../backend/Api';
 import { hordeClasses } from '../../styles/Styles';
 import { observer } from 'mobx-react-lite';
@@ -214,6 +214,8 @@ export const JobDetailArtifactsV2: React.FC<{ jobDetails: JobDetailsV2; stepId:
 
    const dataView = jobDetails.getDataView<ArtifactsDataView>("ArtifactsDataView");
 
+   const [state, setState] = useState<{ filter?: string }>({});
+
    useEffect(() => {
       return () => {
          dataView?.clear();
@@ -264,14 +266,48 @@ export const JobDetailArtifactsV2: React.FC<{ jobDetails: JobDetailsV2; stepId:
 
    let height = Math.min(36 * artifacts.length + 60, 500);
 
+   if (state.filter) {
+      const f = state.filter.toLowerCase();
+      artifacts = artifacts.filter(a => a.name?.toLowerCase().indexOf(f) !== -1);
+   }
+
    return (<Stack id={sideRail.url} styles={{ root: { paddingTop: 18, paddingRight: 12 } }}>
       <Stack className={hordeClasses.raised}>
          <Stack tokens={{ childrenGap: 12 }}>
             <Stack horizontal horizontalAlign="space-between" styles={{ root: { minHeight: 32 } }}>
                <Text variant="mediumPlus" styles={{ root: { fontFamily: "Horde Open Sans SemiBold" } }}>Artifacts</Text>
-               <Stack horizontal>
-                  {artifactState.isDownloading && <Spinner styles={{ root: { marginRight: 10 } }} size={SpinnerSize.medium}></Spinner>}
-                  <PrimaryButton styles={{ root: { fontFamily: 'Horde Open Sans Semibold !important' } }} disabled={artifactState.isDownloading} onClick={artifactState.downloadItems.bind(artifactState, true, false, undefined)}>{buttonText}</PrimaryButton>
+               <Stack horizontal tokens={{ childrenGap: 24 }}>
+                  <Stack>
+                     <TextField
+                        spellCheck={false}
+                        deferredValidationTime={500}
+                        validateOnLoad={false}
+
+                        styles={{
+                           root: { width: 320, fontSize: 12 }, fieldGroup: {
+                              borderWidth: 1
+                           }
+                        }}
+
+                        placeholder="Filter"
+
+                        onGetErrorMessage={(newValue) => {
+
+                           const filter = newValue ? newValue : undefined;
+                           if (filter !== state.filter) {
+                              setState({...state, filter: filter});
+                           }
+
+                           return undefined;
+                        }}
+
+                     />
+
+                  </Stack>
+                  <Stack>
+                     {artifactState.isDownloading && <Spinner styles={{ root: { marginRight: 10 } }} size={SpinnerSize.medium}></Spinner>}
+                     <PrimaryButton styles={{ root: { fontFamily: 'Horde Open Sans Semibold !important' } }} disabled={artifactState.isDownloading} onClick={artifactState.downloadItems.bind(artifactState, true, false, undefined)}>{buttonText}</PrimaryButton>
+                  </Stack>
                </Stack>
             </Stack>
             <Stack>
