Horde: Extract XGE task metadata as structured log events

As a start, they are logged to the agent process logger (as opposed to job logger).

diff --git a/Engine/Source/Programs/Horde/Horde.Agent.Tests/Execution/WorkspaceExecutorTest.cs b/Engine/Source/Programs/Horde/Horde.Agent.Tests/Execution/WorkspaceExecutorTest.cs
--- a/Engine/Source/Programs/Horde/Horde.Agent.Tests/Execution/WorkspaceExecutorTest.cs
+++ b/Engine/Source/Programs/Horde/Horde.Agent.Tests/Execution/WorkspaceExecutorTest.cs
@@ -9,6 +9,7 @@ using EpicGames.Core;
 using Horde.Agent.Execution;
 using Horde.Agent.Services;
 using Microsoft.Extensions.Logging;
+using Microsoft.Extensions.Logging.Abstractions;
 using Microsoft.VisualStudio.TestTools.UnitTesting;
 
 namespace Horde.Agent.Tests.Execution;
@@ -46,7 +47,7 @@ public sealed class WorkspaceExecutorTest : IDisposable
 		_workspace.SetFile(1, "main.cpp", "main");
 		_workspace.SetFile(1, "foo/bar/baz.h", "baz");
 		
-		_executor = new (_session, JobId, "batch1", AgentType, null, _workspace, null!);
+		_executor = new (_session, JobId, "batch1", AgentType, null, _workspace, null!, NullLogger.Instance);
 	}
 
 	public void Dispose()
@@ -67,7 +68,7 @@ public sealed class WorkspaceExecutorTest : IDisposable
 	[TestMethod]
 	public async Task RegularAndAutoSdkWorkspace()
 	{
-		WorkspaceExecutor executor = new (_session, JobId, "batch1", AgentType, _autoSdkWorkspace, _workspace, null!);
+		WorkspaceExecutor executor = new (_session, JobId, "batch1", AgentType, _autoSdkWorkspace, _workspace, null!, NullLogger.Instance);
 
 		await executor.InitializeAsync(_logger, CancellationToken.None);
 		AssertWorkspaceFile(_autoSdkWorkspace, "HostWin64/Android/base.h", "base");
@@ -79,7 +80,7 @@ public sealed class WorkspaceExecutorTest : IDisposable
 	[TestMethod]
 	public async Task EnvVars()
 	{
-		WorkspaceExecutor executor = new (_session, JobId, "batch1", AgentType, _autoSdkWorkspace, _workspace, null!);
+		WorkspaceExecutor executor = new (_session, JobId, "batch1", AgentType, _autoSdkWorkspace, _workspace, null!, NullLogger.Instance);
 		await executor.InitializeAsync(_logger, CancellationToken.None);
 
 		WorkspaceMaterializerSettings settings = await _workspace.GetSettingsAsync(CancellationToken.None);
@@ -102,7 +103,7 @@ public sealed class WorkspaceExecutorTest : IDisposable
 	{
 		_server.AddJob("jobPreflight", StreamId, 1, 1000);		
 		_workspace.SetFile(1000, "New/Feature/Foo.cs", "foo");
-		WorkspaceExecutor executor = new (_session, "jobPreflight", "batch1", AgentType, null, _workspace, null!);
+		WorkspaceExecutor executor = new (_session, "jobPreflight", "batch1", AgentType, null, _workspace, null!, NullLogger.Instance);
 		
 		await executor.InitializeAsync(_logger, CancellationToken.None);
 		AssertWorkspaceFile(_workspace, "main.cpp", "main");
@@ -115,7 +116,7 @@ public sealed class WorkspaceExecutorTest : IDisposable
 	public async Task JobWithNoChange()
 	{
 		_server.AddJob("jobNoChange", StreamId, 0, 0);		
-		WorkspaceExecutor executor = new (_session, "jobNoChange", "batch1", AgentType, null, _workspace, null!);
+		WorkspaceExecutor executor = new (_session, "jobNoChange", "batch1", AgentType, null, _workspace, null!, NullLogger.Instance);
 		await Assert.ThrowsExceptionAsync<WorkspaceMaterializationException>(() => executor.InitializeAsync(_logger, CancellationToken.None));
 	}
 
diff --git a/Engine/Source/Programs/Horde/Horde.Agent.Tests/Helper.cs b/Engine/Source/Programs/Horde/Horde.Agent.Tests/Helper.cs
--- a/Engine/Source/Programs/Horde/Horde.Agent.Tests/Helper.cs
+++ b/Engine/Source/Programs/Horde/Horde.Agent.Tests/Helper.cs
@@ -188,7 +188,7 @@ namespace Horde.Agent.Tests
 		private readonly Func<BeginStepResponse, ILogger, CancellationToken, Task<JobStepOutcome>> _func;
 
 		public SimpleTestExecutor(Func<BeginStepResponse, ILogger, CancellationToken, Task<JobStepOutcome>> func)
-			: base(null!, null!, null!, null!, null!)
+			: base(null!, null!, null!, null!, null!, NullLogger.Instance)
 		{
 			_func = func;
 		}
diff --git a/Engine/Source/Programs/Horde/Horde.Agent.Tests/LogParserTests.cs b/Engine/Source/Programs/Horde/Horde.Agent.Tests/LogParserTests.cs
--- a/Engine/Source/Programs/Horde/Horde.Agent.Tests/LogParserTests.cs
+++ b/Engine/Source/Programs/Horde/Horde.Agent.Tests/LogParserTests.cs
@@ -107,7 +107,6 @@ namespace Horde.Agent.Tests
 				@"1 error generated.",
 				@"",
 				@"Error executing d:\build\AutoSDK\Sync\HostWin64\Android\-24\ndk\21.4.7075529\toolchains\llvm\prebuilt\windows-x86_64\bin\clang++.exe (tool returned code: 1)",
-				@"AnimNode_BlendListByInt.gen.cpp [arm64] (0:22.43 at +2:47)",
 			};
 
 			List<LogEvent> logEvents = Parse(lines);
@@ -597,6 +596,63 @@ namespace Horde.Agent.Tests
 			CheckEventGroup(logEvents.Slice(6, 1), 7, 1, LogLevel.Information, KnownLogEvents.Systemic_Xge_ServiceNotRunning);
 			CheckEventGroup(logEvents.Slice(7, 1), 8, 1, LogLevel.Error, KnownLogEvents.Systemic_Xge_BuildFailed);
 		}
+		
+		[TestMethod]
+		public void XoreaxTaskMetadataMatcher()
+		{
+			string[] lines =
+			{
+				@"MiMalloc.c (0:01.17 at +0:11)",
+				@"SomeFile.cpp (4:31.12 at +0:16)",
+				@"Default.rc2 (Agent 'SomeAgentName (Core #7)', 2:34.56 at +55:18)",
+				@" (1:01.12 at +58:03)",
+			};
+
+			List<LogEvent> logEvents = Parse(lines);
+			Assert.AreEqual(4, logEvents.Count);
+			CheckEventGroup(logEvents.Slice(0, 1), 0, 1, LogLevel.Information, KnownLogEvents.Systemic_Xge_TaskMetadata);
+			CheckEventGroup(logEvents.Slice(1, 1), 1, 1, LogLevel.Information, KnownLogEvents.Systemic_Xge_TaskMetadata);
+			CheckEventGroup(logEvents.Slice(2, 1), 2, 1, LogLevel.Information, KnownLogEvents.Systemic_Xge_TaskMetadata);
+			CheckEventGroup(logEvents.Slice(3, 1), 3, 1, LogLevel.Information, KnownLogEvents.Systemic_Xge_TaskMetadata);
+
+			Assert.AreEqual("MiMalloc.c", logEvents[0].GetProperty<string>("name"));
+			Assert.AreEqual(77, logEvents[0].GetProperty<int>("duration"));
+			Assert.AreEqual(11, logEvents[0].GetProperty<int>("startTime"));
+			
+			Assert.AreEqual("SomeFile.cpp", logEvents[1].GetProperty<string>("name"));
+			Assert.AreEqual(16272, logEvents[1].GetProperty<int>("duration"));
+			Assert.AreEqual(16, logEvents[1].GetProperty<int>("startTime"));
+			
+			Assert.AreEqual("Default.rc2", logEvents[2].GetProperty<string>("name"));
+			Assert.AreEqual("SomeAgentName (Core #7)", logEvents[2].GetProperty<string>("agent"));
+			Assert.AreEqual(9296, logEvents[2].GetProperty<int>("duration"));
+			Assert.AreEqual(3318, logEvents[2].GetProperty<int>("startTime"));
+			
+			Assert.AreEqual("", logEvents[3].GetProperty<string>("name"));
+			Assert.AreEqual(3672, logEvents[3].GetProperty<int>("duration"));
+			Assert.AreEqual(3483, logEvents[3].GetProperty<int>("startTime"));
+		}
+
+		private class FakeLogSink : ILogEventSink
+		{
+			public List<LogEvent> ReceivedLogEvents { get; } = new();
+			public void ProcessEvent(LogEvent logEvent)
+			{
+				ReceivedLogEvents.Add(logEvent);
+			}
+		}
+		
+		[TestMethod]
+		public void AdditionalLogSinks()
+		{
+			FakeLogSink fakeLogSink = new();
+			using LogParser parser = new (NullLogger.Instance, new List<string>(), new List<ILogEventSink>() { fakeLogSink });
+
+			parser.WriteLine(@"MiMalloc.c (0:01.17 at +0:11)");
+			
+			Assert.AreEqual(1, fakeLogSink.ReceivedLogEvents.Count);
+			Assert.AreEqual("MiMalloc.c", fakeLogSink.ReceivedLogEvents[0].GetProperty<string>("name"));
+		}
 
 		[TestMethod]
 		public void LogChannelMatcher()
@@ -781,10 +837,10 @@ namespace Horde.Agent.Tests
 				Assert.IsTrue(enumerator.MoveNext());
 
 				LogEvent logEvent = enumerator.Current;
-				Assert.AreEqual(level, logEvent.Level);
-				Assert.AreEqual(eventId, logEvent.Id);
-				Assert.AreEqual(idx, logEvent.LineIndex);
-				Assert.AreEqual(count, logEvent.LineCount);
+				Assert.AreEqual(level, logEvent.Level, "Log level mismatch");
+				Assert.AreEqual(eventId, logEvent.Id, "Event ID mismatch");
+				Assert.AreEqual(idx, logEvent.LineIndex, "Line index mismatch");
+				Assert.AreEqual(count, logEvent.LineCount, "Line count mismatch");
 				Assert.AreEqual(index + idx, logEvent.GetProperty(LogLine));
 			}
 			Assert.IsFalse(enumerator.MoveNext());
diff --git a/Engine/Source/Programs/Horde/Horde.Agent/Execution/JobExecutor.cs b/Engine/Source/Programs/Horde/Horde.Agent/Execution/JobExecutor.cs
--- a/Engine/Source/Programs/Horde/Horde.Agent/Execution/JobExecutor.cs
+++ b/Engine/Source/Programs/Horde/Horde.Agent/Execution/JobExecutor.cs
@@ -150,6 +150,11 @@ namespace Horde.Agent.Execution
 		protected string _agentTypeName;
 
 		protected IHttpClientFactory _httpClientFactory;
+		
+		/// <summary>
+		/// Logger for the local agent process (as opposed to job logger)
+		/// </summary>
+		readonly protected ILogger _logger;
 
 		protected GetJobResponse _job;
 		protected GetStreamResponse _stream;
@@ -167,7 +172,7 @@ namespace Horde.Agent.Execution
 
 		public bool UseNewLogger { get; private set; }
 
-		public JobExecutor(ISession session, string jobId, string batchId, string agentTypeName, IHttpClientFactory httpClientFactory)
+		public JobExecutor(ISession session, string jobId, string batchId, string agentTypeName, IHttpClientFactory httpClientFactory, ILogger logger)
 		{
 			_session = session;
 
@@ -176,6 +181,7 @@ namespace Horde.Agent.Execution
 			_agentTypeName = agentTypeName;
 
 			_httpClientFactory = httpClientFactory;
+			_logger = logger;
 
 			_job = null!;
 			_stream = null!;
@@ -1160,8 +1166,23 @@ namespace Horde.Agent.Execution
 
 			List<string> ignorePatterns = await ReadIgnorePatternsAsync(workspaceDir, logger);
 
+			FilteringEventSink xgeEventSink = new(new [] { KnownLogEvents.Systemic_Xge_TaskMetadata.Id }, (logEvent) =>
+			{
+				try
+				{
+					string taskName = logEvent.GetProperty<string>("name");
+					int duration = logEvent.GetProperty<int>("duration");
+					string agent = logEvent.GetProperty<string>("agent");
+					_logger.LogInformation("Ran XGE task '{XgeTaskName}' on agent '{XgeAgent}' for {Duration} secs", taskName, agent, duration);	
+				}
+				catch (Exception e)
+				{
+					_logger.LogWarning(e, "Failed to log XGE task execution");
+				}
+			});
+			
 			int exitCode;
-			using (LogParser filter = new LogParser(logger, ignorePatterns))
+			using (LogParser filter = new LogParser(logger, ignorePatterns, new List<ILogEventSink>() { xgeEventSink }))
 			{
 				await ExecuteCleanupScriptAsync(cleanupScript, filter, logger);
 				try
diff --git a/Engine/Source/Programs/Horde/Horde.Agent/Execution/LocalExecutor.cs b/Engine/Source/Programs/Horde/Horde.Agent/Execution/LocalExecutor.cs
--- a/Engine/Source/Programs/Horde/Horde.Agent/Execution/LocalExecutor.cs
+++ b/Engine/Source/Programs/Horde/Horde.Agent/Execution/LocalExecutor.cs
@@ -20,8 +20,8 @@ namespace Horde.Agent.Execution
 		private readonly LocalExecutorSettings _settings;
 		private readonly DirectoryReference _localWorkspaceDir;
 
-		public LocalExecutor(ISession session, string jobId, string batchId, string agentTypeName, LocalExecutorSettings settings, IHttpClientFactory httpClientFactory) 
-			: base(session, jobId, batchId, agentTypeName, httpClientFactory)
+		public LocalExecutor(ISession session, string jobId, string batchId, string agentTypeName, LocalExecutorSettings settings, IHttpClientFactory httpClientFactory, ILogger logger) 
+			: base(session, jobId, batchId, agentTypeName, httpClientFactory, logger)
 		{
 			_settings = settings;
 			if (settings.WorkspaceDir == null)
@@ -74,18 +74,20 @@ namespace Horde.Agent.Execution
 	{
 		readonly LocalExecutorSettings _settings;
 		readonly IHttpClientFactory _httpClientFactory;
+		readonly ILogger<LocalExecutor> _logger;
 
 		public override string Name => "Local";
 
-		public LocalExecutorFactory(IOptions<LocalExecutorSettings> settings, IHttpClientFactory httpClientFactory)
+		public LocalExecutorFactory(IOptions<LocalExecutorSettings> settings, IHttpClientFactory httpClientFactory, ILogger<LocalExecutor> logger)
 		{
 			_settings = settings.Value;
 			_httpClientFactory = httpClientFactory;
+			_logger = logger;
 		}
 
 		public override JobExecutor CreateExecutor(ISession session, ExecuteJobTask executeJobTask, BeginBatchResponse beginBatchResponse)
 		{
-			return new LocalExecutor(session, executeJobTask.JobId, executeJobTask.BatchId, beginBatchResponse.AgentType, _settings, _httpClientFactory);
+			return new LocalExecutor(session, executeJobTask.JobId, executeJobTask.BatchId, beginBatchResponse.AgentType, _settings, _httpClientFactory, _logger);
 		}
 	}
 }
diff --git a/Engine/Source/Programs/Horde/Horde.Agent/Execution/PerforceExecutor.cs b/Engine/Source/Programs/Horde/Horde.Agent/Execution/PerforceExecutor.cs
--- a/Engine/Source/Programs/Horde/Horde.Agent/Execution/PerforceExecutor.cs
+++ b/Engine/Source/Programs/Horde/Horde.Agent/Execution/PerforceExecutor.cs
@@ -37,8 +37,8 @@ namespace Horde.Agent.Execution
 		protected WorkspaceInfo? _autoSdkWorkspace;
 		protected WorkspaceInfo _workspace;
 
-		public PerforceExecutor(ISession session, string jobId, string batchId, string agentTypeName, AgentWorkspace? autoSdkWorkspaceInfo, AgentWorkspace workspaceInfo, DirectoryReference rootDir, IHttpClientFactory httpClientFactory)
-			: base(session, jobId, batchId, agentTypeName, httpClientFactory)
+		public PerforceExecutor(ISession session, string jobId, string batchId, string agentTypeName, AgentWorkspace? autoSdkWorkspaceInfo, AgentWorkspace workspaceInfo, DirectoryReference rootDir, IHttpClientFactory httpClientFactory, ILogger logger)
+			: base(session, jobId, batchId, agentTypeName, httpClientFactory, logger)
 		{
 			_autoSdkWorkspaceInfo = autoSdkWorkspaceInfo;
 			_workspaceInfo = workspaceInfo;
@@ -436,17 +436,19 @@ namespace Horde.Agent.Execution
 	class PerforceExecutorFactory : JobExecutorFactory
 	{
 		readonly IHttpClientFactory _httpClientFactory;
+		readonly ILogger<PerforceExecutor> _logger;
 
 		public override string Name => PerforceExecutor.Name;
 
-		public PerforceExecutorFactory(IHttpClientFactory httpClientFactory)
+		public PerforceExecutorFactory(IHttpClientFactory httpClientFactory, ILogger<PerforceExecutor> logger)
 		{
 			_httpClientFactory = httpClientFactory;
+			_logger = logger;
 		}
 
 		public override JobExecutor CreateExecutor(ISession session, ExecuteJobTask executeJobTask, BeginBatchResponse beginBatchResponse)
 		{
-			return new PerforceExecutor(session, executeJobTask.JobId, executeJobTask.BatchId, beginBatchResponse.AgentType, executeJobTask.AutoSdkWorkspace, executeJobTask.Workspace, session.WorkingDir, _httpClientFactory);
+			return new PerforceExecutor(session, executeJobTask.JobId, executeJobTask.BatchId, beginBatchResponse.AgentType, executeJobTask.AutoSdkWorkspace, executeJobTask.Workspace, session.WorkingDir, _httpClientFactory, _logger);
 		}
 	}
 }
diff --git a/Engine/Source/Programs/Horde/Horde.Agent/Execution/TestExecutor.cs b/Engine/Source/Programs/Horde/Horde.Agent/Execution/TestExecutor.cs
--- a/Engine/Source/Programs/Horde/Horde.Agent/Execution/TestExecutor.cs
+++ b/Engine/Source/Programs/Horde/Horde.Agent/Execution/TestExecutor.cs
@@ -22,8 +22,8 @@ namespace Horde.Agent.Execution
 	{
 		public const string Name = "Test";
 
-		public TestExecutor(ISession session, string jobId, string batchId, string agentTypeName, IHttpClientFactory httpClientFactory)
-			: base(session, jobId, batchId, agentTypeName, httpClientFactory)
+		public TestExecutor(ISession session, string jobId, string batchId, string agentTypeName, IHttpClientFactory httpClientFactory, ILogger logger)
+			: base(session, jobId, batchId, agentTypeName, httpClientFactory, logger)
 		{
 		}
 
@@ -193,17 +193,19 @@ namespace Horde.Agent.Execution
 	class TestExecutorFactory : JobExecutorFactory
 	{
 		readonly IHttpClientFactory _httpClientFactory;
+		readonly ILogger<TestExecutor> _logger;
 
 		public override string Name => TestExecutor.Name;
 
-		public TestExecutorFactory(IHttpClientFactory httpClientFactory)
+		public TestExecutorFactory(IHttpClientFactory httpClientFactory, ILogger<TestExecutor> logger)
 		{
 			_httpClientFactory = httpClientFactory;
+			_logger = logger;
 		}
 
 		public override JobExecutor CreateExecutor(ISession session, ExecuteJobTask executeJobTask, BeginBatchResponse beginBatchResponse)
 		{
-			return new TestExecutor(session, executeJobTask.JobId, executeJobTask.BatchId, beginBatchResponse.AgentType, _httpClientFactory);
+			return new TestExecutor(session, executeJobTask.JobId, executeJobTask.BatchId, beginBatchResponse.AgentType, _httpClientFactory, _logger);
 		}
 	}
 }
diff --git a/Engine/Source/Programs/Horde/Horde.Agent/Execution/WorkspaceExecutor.cs b/Engine/Source/Programs/Horde/Horde.Agent/Execution/WorkspaceExecutor.cs
--- a/Engine/Source/Programs/Horde/Horde.Agent/Execution/WorkspaceExecutor.cs
+++ b/Engine/Source/Programs/Horde/Horde.Agent/Execution/WorkspaceExecutor.cs
@@ -23,8 +23,8 @@ namespace Horde.Agent.Execution
 		private readonly IWorkspaceMaterializer? _autoSdkWorkspace;
 		private readonly IWorkspaceMaterializer _workspace;
 
-		public WorkspaceExecutor(ISession session, string jobId, string batchId, string agentTypeName, IWorkspaceMaterializer? autoSdkWorkspace, IWorkspaceMaterializer workspace, IHttpClientFactory httpClientFactory)
-			: base(session, jobId, batchId, agentTypeName, httpClientFactory)
+		public WorkspaceExecutor(ISession session, string jobId, string batchId, string agentTypeName, IWorkspaceMaterializer? autoSdkWorkspace, IWorkspaceMaterializer workspace, IHttpClientFactory httpClientFactory, ILogger logger)
+			: base(session, jobId, batchId, agentTypeName, httpClientFactory, logger)
 		{
 			_autoSdkWorkspace = autoSdkWorkspace;
 			_workspace = workspace;
diff --git a/Engine/Source/Programs/Horde/Horde.Agent/Parser/LogParser.cs b/Engine/Source/Programs/Horde/Horde.Agent/Parser/LogParser.cs
--- a/Engine/Source/Programs/Horde/Horde.Agent/Parser/LogParser.cs
+++ b/Engine/Source/Programs/Horde/Horde.Agent/Parser/LogParser.cs
@@ -19,8 +19,9 @@ namespace Horde.Agent.Parser
 		/// </summary>
 		/// <param name="logger">The logger to receive parsed output messages</param>
 		/// <param name="ignorePatterns">List of patterns to ignore</param>
-		public LogParser(ILogger logger, List<string> ignorePatterns)
-			: base(logger)
+		/// <param name="logEventSinks">Additional sinks to receive log events</param>
+		public LogParser(ILogger logger, List<string> ignorePatterns, List<ILogEventSink>? logEventSinks = null)
+			: base(logger, logEventSinks)
 		{
 			foreach (string ignorePattern in ignorePatterns)
 			{
diff --git a/Engine/Source/Programs/Shared/EpicGames.Core/KnownLogEvents.cs b/Engine/Source/Programs/Shared/EpicGames.Core/KnownLogEvents.cs
--- a/Engine/Source/Programs/Shared/EpicGames.Core/KnownLogEvents.cs
+++ b/Engine/Source/Programs/Shared/EpicGames.Core/KnownLogEvents.cs
@@ -204,6 +204,16 @@ namespace EpicGames.Core
 		/// Cache size reached
 		/// </summary>
 		public static EventId Systemic_Xge_CacheLimit { get; } = new EventId(714);
+		
+		/// <summary>
+		/// Current logging level may impact performance
+		/// </summary>
+		public static EventId Systemic_Xge_DetailedLogging { get; } = new EventId(715);
+		
+		/// <summary>
+		/// Metadata about an XGE task (local or remote execution, start/end time etc)
+		/// </summary>
+		public static EventId Systemic_Xge_TaskMetadata { get; } = new EventId(716);
 
 		/// <summary>
 		/// DDC is slow
diff --git a/Engine/Source/Programs/Shared/EpicGames.Core/LogEvent.cs b/Engine/Source/Programs/Shared/EpicGames.Core/LogEvent.cs
--- a/Engine/Source/Programs/Shared/EpicGames.Core/LogEvent.cs
+++ b/Engine/Source/Programs/Shared/EpicGames.Core/LogEvent.cs
@@ -588,6 +588,44 @@ namespace EpicGames.Core
 		}
 	}
 
+	/// <summary>
+	/// Interface for a log event sink
+	/// </summary>
+	public interface ILogEventSink
+	{
+		/// <summary>
+		/// Process log event 
+		/// </summary>
+		/// <param name="logEvent">Log event</param>
+		void ProcessEvent(LogEvent logEvent);
+	}
+
+	/// <summary>
+	/// Simple filtering log event sink with a callback for convenience
+	/// </summary>
+	public class FilteringEventSink : ILogEventSink
+	{
+		private IReadOnlyList<LogEvent> LogEvents => _logEvents;
+		private readonly List<LogEvent> _logEvents = new();
+		private readonly int[] _includeEventIds;
+		private readonly Action<LogEvent>? _eventCallback;
+		
+		public FilteringEventSink(int[] includeEventIds, Action<LogEvent>? eventCallback = null)
+		{
+			_includeEventIds = includeEventIds;
+			_eventCallback = eventCallback;
+		}
+
+		public void ProcessEvent(LogEvent logEvent)
+		{
+			if (_includeEventIds.Any(x => x == logEvent.Id.Id))
+			{
+				_logEvents.Add(logEvent);
+				_eventCallback?.Invoke(logEvent);
+			}
+		}
+	}
+	
 	/// <summary>
 	/// Converter for serialization of <see cref="LogEvent"/> instances to Json streams
 	/// </summary>
diff --git a/Engine/Source/Programs/Shared/EpicGames.Core/LogEventParser.cs b/Engine/Source/Programs/Shared/EpicGames.Core/LogEventParser.cs
--- a/Engine/Source/Programs/Shared/EpicGames.Core/LogEventParser.cs
+++ b/Engine/Source/Programs/Shared/EpicGames.Core/LogEventParser.cs
@@ -141,6 +141,11 @@ namespace EpicGames.Core
 		/// </summary>
 		ILogger _logger;
 
+		/// <summary>
+		/// Log events sinks in addition to <see cref="_logger" />
+		/// </summary>
+		List<ILogEventSink> _logEventSinks = new();
+
 		/// <summary>
 		/// Timer for the parser being active
 		/// </summary>
@@ -169,10 +174,16 @@ namespace EpicGames.Core
 		/// Constructor
 		/// </summary>
 		/// <param name="logger">The logger to receive parsed output messages</param>
-		public LogEventParser(ILogger logger)
+		/// <param name="logEventSinks">Additional sinks to receive log events</param>
+		public LogEventParser(ILogger logger, List<ILogEventSink>? logEventSinks = null)
 		{
 			_logger = logger;
 			_buffer = new LogBuffer(50);
+			
+			if (logEventSinks != null)
+			{
+				_logEventSinks.AddRange(logEventSinks);
+			}
 		}
 
 		/// <inheritdoc/>
@@ -524,6 +535,10 @@ namespace EpicGames.Core
 			foreach (LogEvent logEvent in logEvents)
 			{
 				_logger.Log(logEvent.Level, logEvent.Id, logEvent, null, (state, exception) => state.ToString());
+				foreach (ILogEventSink logEventSink in _logEventSinks)
+				{
+					logEventSink.ProcessEvent(logEvent);
+				}
 			}
 		}
 	}
diff --git a/Engine/Source/Programs/UnrealBuildTool/Matchers/XoreaxEventMatcher.cs b/Engine/Source/Programs/UnrealBuildTool/Matchers/XoreaxEventMatcher.cs
--- a/Engine/Source/Programs/UnrealBuildTool/Matchers/XoreaxEventMatcher.cs
+++ b/Engine/Source/Programs/UnrealBuildTool/Matchers/XoreaxEventMatcher.cs
@@ -1,5 +1,7 @@
 // Copyright Epic Games, Inc. All Rights Reserved.
 
+using System;
+using System.Diagnostics.CodeAnalysis;
 using System.Text.RegularExpressions;
 using EpicGames.Core;
 using Microsoft.Extensions.Logging;
@@ -8,6 +10,11 @@ namespace UnrealBuildTool.Matchers
 {
 	class XoreaxEventMatcher : ILogEventMatcher
 	{
+		public const string PropertyName = "name";
+		public const string PropertyAgent = "agent";
+		public const string PropertyDuration = "duration";
+		public const string PropertyStartTime = "startTime";
+		
 		static readonly Regex s_buildService = new Regex(@"\(BuildService.exe\) is not running");
 
 		static readonly Regex s_xgConsole = new Regex(@"BUILD FAILED: (.*)xgConsole\.exe(.*)");
@@ -21,6 +28,11 @@ namespace UnrealBuildTool.Matchers
 		static readonly Regex s_cacheWarning = new Regex(@"^\s*WARNING: \d+ items \([^\)]*\) removed from the cache due to reaching the cache size limit");
 		static readonly Regex s_cacheWarning2 = new Regex(@"^\s*WARNING: Several items removed from the cache due to reaching the cache size limit");
 		static readonly Regex s_cacheWarning3 = new Regex(@"^\s*WARNING: The Build Cache is close to full");
+		
+		static readonly Regex s_taskFinishedOnHelper = new Regex(@"(.*?) \(Agent '(.+)', (.*) at \+(.*)\)");
+		static readonly Regex s_taskFinishedLocally = new Regex(@"(.*?) \((.*) at \+(.*)\)");
+		static readonly Regex s_duration = new Regex(@"(\d+):(\d+)\.(\d+)");
+		static readonly Regex s_startTime = new Regex(@"\+?(\d+):(\d+)");
 
 		public LogEventMatch? Match(ILogCursor cursor)
 		{
@@ -83,7 +95,52 @@ namespace UnrealBuildTool.Matchers
 					return builder.ToMatch(LogEventPriority.High, LogLevel.Information, eventId);
 				}
 			}
+			
+			if (cursor.TryMatch(s_taskFinishedOnHelper, out Match? taskHelperMatch))
+			{
+				LogEventBuilder builder = new LogEventBuilder(cursor);
+				builder.AddProperty(PropertyName, taskHelperMatch.Groups[1].Value);
+				builder.AddProperty(PropertyAgent, taskHelperMatch.Groups[2].Value);
+				SetDuration(builder, taskHelperMatch.Groups[3].Value);
+				SetStartTime(builder, taskHelperMatch.Groups[4].Value);
+				return builder.ToMatch(LogEventPriority.High, LogLevel.Information, KnownLogEvents.Systemic_Xge_TaskMetadata);
+			}
+			
+			if (cursor.TryMatch(s_taskFinishedLocally, out Match? taskLocalMatch))
+			{
+				LogEventBuilder builder = new LogEventBuilder(cursor);
+				builder.AddProperty(PropertyName, taskLocalMatch.Groups[1].Value);
+				SetDuration(builder, taskLocalMatch.Groups[2].Value);
+				SetStartTime(builder, taskLocalMatch.Groups[3].Value);
+				return builder.ToMatch(LogEventPriority.High, LogLevel.Information, KnownLogEvents.Systemic_Xge_TaskMetadata);
+			}
+			
 			return null;
 		}
+
+		private static void SetDuration(LogEventBuilder builder, string durationText)
+		{
+			Match m = s_duration.Match(durationText);
+			if (m.Success)
+			{
+				int hours = Convert.ToInt32(m.Groups[1].Value);
+				int minutes = Convert.ToInt32(m.Groups[2].Value);
+				int seconds = Convert.ToInt32(m.Groups[3].Value);
+				TimeSpan duration = new TimeSpan(hours, minutes, seconds);
+				builder.AddProperty(PropertyDuration, duration.TotalSeconds);
+			}
+		}
+		
+		private static void SetStartTime(LogEventBuilder builder, string startTimeText)
+		{
+			Match m = s_startTime.Match(startTimeText);
+			if (m.Success)
+			{
+				int minutes = Convert.ToInt32(m.Groups[1].Value);
+				int seconds = Convert.ToInt32(m.Groups[2].Value);
+				TimeSpan startTime = new TimeSpan(0, minutes, seconds);
+				builder.AddProperty(PropertyStartTime, startTime.TotalSeconds);
+			}
+		}
 	}
 }
