Horde: use createBrowserRouter with RR v6, also add error and root routing components
#skipci
#rnx

diff --git a/Engine/Source/Programs/Horde/HordeDashboard/src/App.tsx b/Engine/Source/Programs/Horde/HordeDashboard/src/App.tsx
--- a/Engine/Source/Programs/Horde/HordeDashboard/src/App.tsx
+++ b/Engine/Source/Programs/Horde/HordeDashboard/src/App.tsx
@@ -3,7 +3,7 @@
 import { Image, Spinner, SpinnerSize, Stack, Text } from '@fluentui/react';
 import React, { useState } from 'react';
 import { useDarkreader } from 'react-darkreader';
-import { BrowserRouter, Navigate, Route, Routes } from 'react-router-dom';
+import { createBrowserRouter, Navigate, Outlet, RouteObject, RouterProvider } from 'react-router-dom';
 import backend from './backend';
 import { DashboardPreference } from './backend/Api';
 import { getSiteConfig } from './backend/Config';
@@ -23,7 +23,6 @@ import { PerforceServerView } from './components/PerforceView';
 import { PoolView } from './components/PoolView';
 import { PreflightRedirector } from './components/Preflight';
 import { ProjectHome } from './components/ProjectHome';
-import { ServerSettingsView } from './components/ServerSettings';
 import { StreamView } from './components/StreamView';
 import { TestReportView } from './components/TestReportView';
 import { UserHomeView } from './components/UserHome';
@@ -31,6 +30,13 @@ import { UtilizationReportView } from './components/UtilizationReportView';
 import hordePlugins from './Plugins';
 import { modeColors, preloadFonts } from './styles/Styles';
 
+
+let router: any;
+
+const RouteError: React.FC = () => {
+   return <Navigate to="/index" replace={true} />
+}
+
 const Main: React.FC = () => {
 
    const [init, setInit] = useState(false);
@@ -115,33 +121,46 @@ const Main: React.FC = () => {
       return null;
    }
 
+   if (!router) {
+
+      const routes: RouteObject[] = [
+         {
+            path: "/", element: <Root />, errorElement: <RouteError />, children: [
+               { path: "index", element: <UserHomeView /> },
+               { path: "project/:projectId", element: <ProjectHome /> },
+               { path: "pool/:poolId", element: <PoolView /> },
+               { path: "job/:jobId", element: <JobDetailViewV2 /> },
+               { path: "log/:logId", element: <LogView /> },
+               { path: "testreport/:testdataId", element: <TestReportView /> },
+               { path: "stream/:streamId", element: <StreamView /> },
+               { path: "agents/:agentId?", element: <AgentView /> },
+               { path: "admin/token", element: <AdminToken /> },
+               { path: "reports/utilization", element: <UtilizationReportView /> },
+               { path: "preflight", element: <PreflightRedirector /> },
+               { path: "dashboard", element: <DashboardView /> },
+               { path: "perforce/servers", element: <PerforceServerView /> },
+               { path: "notices", element: <NoticeView /> },
+               { path: "devices", element: <DeviceView /> },
+               { path: "audit/agent/:agentId", element: <AuditLogView /> },
+               { path: "audit/issue/:issueId", element: <AuditLogView /> },
+               { path: "automation", element: <AutomationView /> },
+               { path: "debug/lease/:leaseId", element: <DebugView /> }
+            ]
+         }
+      ];
+
+      // mount plugins
+      const pluginRoutes = hordePlugins.routes.map((route) => {
+         return { path: route.path, element: <route.component /> };
+      })
+
+      routes[0].children!.push(...pluginRoutes);
+
+      router = createBrowserRouter(routes);
+   }
+
    return (
-      <Routes>
-         <Route path="/index" element={<UserHomeView />} />
-         <Route path="/project/:projectId" element={<ProjectHome />} />
-         <Route path="/pool/:poolId" element={<PoolView />} />
-         <Route path="/job/:jobId" element={<JobDetailViewV2 />} />
-         <Route path="/log/:logId" element={<LogView />} />
-         <Route path="/testreport/:testdataId" element={<TestReportView />} />
-         <Route path="/stream/:streamId" element={<StreamView />} />
-         <Route path="/agents/:agentId?" element={<AgentView />} />
-         <Route path="/admin/token" element={<AdminToken />} />
-         <Route path="/reports/utilization" element={<UtilizationReportView />} />
-         <Route path="/preflight" element={<PreflightRedirector />} />
-         <Route path="/dashboard" element={<DashboardView />} />
-         <Route path="/perforce/servers" element={<PerforceServerView />} />
-         <Route path="/notices" element={<NoticeView />} />
-         <Route path="/devices" element={<DeviceView />} />
-         <Route path="/server/settings" element={<ServerSettingsView />} />
-         <Route path="/audit/agent/:agentId" element={<AuditLogView />} />
-         <Route path="/audit/issue/:issueId" element={<AuditLogView />} />
-         <Route path="/automation" element={<AutomationView />} />
-         <Route path="/debug/lease/:leaseId" element={<DebugView />} />
-         {hordePlugins.routes.map((route, index) => {
-            return <Route key={`key_plugin_route_${index}`} path={route.path} element={<route.component />} />;
-         })}
-         <Route path='*' element={<Navigate to='/index' replace={true} />} />
-      </Routes>
+      <RouterProvider router={router} />
    );
 };
 
@@ -149,8 +168,8 @@ const Darkmode: React.FC = () => {
 
    const additionalCSS = `
       .ms-Toggle-thumb {background-color: ${modeColors.text};}}
-      .ms-Toggle-thumb:hover {background-color: ${modeColors.text};}}      
-      .ms-Toggle-thumb:hover {background-color: ${modeColors.text};}}           
+      .ms-Toggle-thumb:hover {background-color: ${modeColors.text};}}
+      .ms-Toggle-thumb:hover {background-color: ${modeColors.text};}}
    `;
 
    // NOTE: if an Stack child isn't respecting className="horde-no-darktheme", check that you are using style: {} instead of styles:{root:{}}!
@@ -165,9 +184,7 @@ const App: React.FC = () => {
       <React.Fragment>
          <Darkmode />
          <ErrorDialog />
-         <BrowserRouter>
-            <Main />
-         </BrowserRouter >
+         <Main />
       </React.Fragment>
    );
 };
@@ -175,4 +192,17 @@ const App: React.FC = () => {
 export default App;
 
 
+const HomeRedirect: React.FC = () => {
+   if (window.location.pathname === "/" || !window.location.pathname) {
+      return <Navigate to="/index" replace={true} />
+   }
+   return null;
+}
+
+const Root: React.FC = () => {
+   return <div>
+      <HomeRedirect />
+      <Outlet />
+   </div>
+}
 
