Horde: Sanitize structured log events sent to the server, ensuring that individual lines or packets never exceed a specified length.

diff --git a/Engine/Source/Programs/Horde/Horde.Agent.Tests/JsonRpcLoggerTests.cs b/Engine/Source/Programs/Horde/Horde.Agent.Tests/JsonRpcLoggerTests.cs
--- a/Engine/Source/Programs/Horde/Horde.Agent.Tests/JsonRpcLoggerTests.cs
+++ b/Engine/Source/Programs/Horde/Horde.Agent.Tests/JsonRpcLoggerTests.cs
@@ -19,6 +19,7 @@ using Microsoft.Extensions.Caching.Memory;
 using Microsoft.Extensions.Logging;
 using Microsoft.Extensions.Logging.Abstractions;
 using Microsoft.VisualStudio.TestTools.UnitTesting;
+using System.Linq;
 
 namespace Horde.Agent.Tests
 {
@@ -58,9 +59,11 @@ namespace Horde.Agent.Tests
 		{
 			JsonLogEvent jsonLogEvent;
 			Assert.IsTrue(JsonLogEvent.TryParse(Encoding.UTF8.GetBytes(message), out jsonLogEvent));
-			ArrayBufferWriter<byte> writer = new ArrayBufferWriter<byte>();
-			int count = JsonRpcLogger.WriteEvent(jsonLogEvent, writer);
-			string[] result = Encoding.UTF8.GetString(writer.WrittenSpan).Split('\n', StringSplitOptions.RemoveEmptyEntries);
+
+			JsonRpcLogWriter writer = new JsonRpcLogWriter();
+			int count = writer.SanitizeAndWriteEvent(jsonLogEvent);
+
+			string[] result = Encoding.UTF8.GetString(writer.CreatePacket().Span).Split('\n', StringSplitOptions.RemoveEmptyEntries);
 			Assert.AreEqual(count, result.Length);
 			return result;
 		}
@@ -72,13 +75,13 @@ namespace Horde.Agent.Tests
 
 			string[] result = SplitMultiLineMessage(msg);
 			Assert.AreEqual(7, result.Length);
-			Assert.AreEqual(@"{""level"":""Information"",""message"":""This"",""format"":""This"",""properties"":{""Line"":""line\nsplit\nin\nfour"",""Var1"":123,""Var2"":{""$type"":""SourceFile"",""$text"":""D:\\build\\\u002B\u002BUE5\\Sync\\GenerateProjectFiles.bat"",""relativePath"":""GenerateProjectFiles.bat"",""depotPath"":""//UE5/Main/GenerateProjectFiles.bat@20392842""},""Line$0"":""line"",""Line$1"":""split"",""Line$2"":""in"",""Line$3"":""four""},""line"":0,""lineCount"":7}", result[0]);
-			Assert.AreEqual(@"{""level"":""Information"",""message"":""is a"",""format"":""is a"",""properties"":{""Line"":""line\nsplit\nin\nfour"",""Var1"":123,""Var2"":{""$type"":""SourceFile"",""$text"":""D:\\build\\\u002B\u002BUE5\\Sync\\GenerateProjectFiles.bat"",""relativePath"":""GenerateProjectFiles.bat"",""depotPath"":""//UE5/Main/GenerateProjectFiles.bat@20392842""},""Line$0"":""line"",""Line$1"":""split"",""Line$2"":""in"",""Line$3"":""four""},""line"":1,""lineCount"":7}", result[1]);
-			Assert.AreEqual(@"{""level"":""Information"",""message"":""multi-line"",""format"":""multi-{Line$0}"",""properties"":{""Line"":""line\nsplit\nin\nfour"",""Var1"":123,""Var2"":{""$type"":""SourceFile"",""$text"":""D:\\build\\\u002B\u002BUE5\\Sync\\GenerateProjectFiles.bat"",""relativePath"":""GenerateProjectFiles.bat"",""depotPath"":""//UE5/Main/GenerateProjectFiles.bat@20392842""},""Line$0"":""line"",""Line$1"":""split"",""Line$2"":""in"",""Line$3"":""four""},""line"":2,""lineCount"":7}", result[2]);
-			Assert.AreEqual(@"{""level"":""Information"",""message"":""split"",""format"":""{Line$1}"",""properties"":{""Line"":""line\nsplit\nin\nfour"",""Var1"":123,""Var2"":{""$type"":""SourceFile"",""$text"":""D:\\build\\\u002B\u002BUE5\\Sync\\GenerateProjectFiles.bat"",""relativePath"":""GenerateProjectFiles.bat"",""depotPath"":""//UE5/Main/GenerateProjectFiles.bat@20392842""},""Line$0"":""line"",""Line$1"":""split"",""Line$2"":""in"",""Line$3"":""four""},""line"":3,""lineCount"":7}", result[3]);
-			Assert.AreEqual(@"{""level"":""Information"",""message"":""in"",""format"":""{Line$2}"",""properties"":{""Line"":""line\nsplit\nin\nfour"",""Var1"":123,""Var2"":{""$type"":""SourceFile"",""$text"":""D:\\build\\\u002B\u002BUE5\\Sync\\GenerateProjectFiles.bat"",""relativePath"":""GenerateProjectFiles.bat"",""depotPath"":""//UE5/Main/GenerateProjectFiles.bat@20392842""},""Line$0"":""line"",""Line$1"":""split"",""Line$2"":""in"",""Line$3"":""four""},""line"":4,""lineCount"":7}", result[4]);
-			Assert.AreEqual(@"{""level"":""Information"",""message"":""four-end"",""format"":""{Line$3}-end"",""properties"":{""Line"":""line\nsplit\nin\nfour"",""Var1"":123,""Var2"":{""$type"":""SourceFile"",""$text"":""D:\\build\\\u002B\u002BUE5\\Sync\\GenerateProjectFiles.bat"",""relativePath"":""GenerateProjectFiles.bat"",""depotPath"":""//UE5/Main/GenerateProjectFiles.bat@20392842""},""Line$0"":""line"",""Line$1"":""split"",""Line$2"":""in"",""Line$3"":""four""},""line"":5,""lineCount"":7}", result[5]);
-			Assert.AreEqual(@"{""level"":""Information"",""message"":""log message 123 D:\\build\\\u002B\u002BUE5\\Sync\\GenerateProjectFiles.bat"",""format"":""log message {Var1} {Var2}"",""properties"":{""Line"":""line\nsplit\nin\nfour"",""Var1"":123,""Var2"":{""$type"":""SourceFile"",""$text"":""D:\\build\\\u002B\u002BUE5\\Sync\\GenerateProjectFiles.bat"",""relativePath"":""GenerateProjectFiles.bat"",""depotPath"":""//UE5/Main/GenerateProjectFiles.bat@20392842""},""Line$0"":""line"",""Line$1"":""split"",""Line$2"":""in"",""Line$3"":""four""},""line"":6,""lineCount"":7}", result[6]);
+			Assert.AreEqual(@"{""level"":""Information"",""message"":""This"",""format"":""This"",""properties"":{""Var1"":123,""Var2"":{""$type"":""SourceFile"",""$text"":""D:\\build\\\u002B\u002BUE5\\Sync\\GenerateProjectFiles.bat"",""relativePath"":""GenerateProjectFiles.bat"",""depotPath"":""//UE5/Main/GenerateProjectFiles.bat@20392842""}},""line"":0,""lineCount"":7}", result[0]);
+			Assert.AreEqual(@"{""level"":""Information"",""message"":""is a"",""format"":""is a"",""properties"":{""Var1"":123,""Var2"":{""$type"":""SourceFile"",""$text"":""D:\\build\\\u002B\u002BUE5\\Sync\\GenerateProjectFiles.bat"",""relativePath"":""GenerateProjectFiles.bat"",""depotPath"":""//UE5/Main/GenerateProjectFiles.bat@20392842""}},""line"":1,""lineCount"":7}", result[1]);
+			Assert.AreEqual(@"{""level"":""Information"",""message"":""multi-line"",""format"":""multi-{Line$0}"",""properties"":{""Var1"":123,""Var2"":{""$type"":""SourceFile"",""$text"":""D:\\build\\\u002B\u002BUE5\\Sync\\GenerateProjectFiles.bat"",""relativePath"":""GenerateProjectFiles.bat"",""depotPath"":""//UE5/Main/GenerateProjectFiles.bat@20392842""},""Line$0"":""line""},""line"":2,""lineCount"":7}", result[2]);
+			Assert.AreEqual(@"{""level"":""Information"",""message"":""split"",""format"":""{Line$1}"",""properties"":{""Var1"":123,""Var2"":{""$type"":""SourceFile"",""$text"":""D:\\build\\\u002B\u002BUE5\\Sync\\GenerateProjectFiles.bat"",""relativePath"":""GenerateProjectFiles.bat"",""depotPath"":""//UE5/Main/GenerateProjectFiles.bat@20392842""},""Line$1"":""split""},""line"":3,""lineCount"":7}", result[3]);
+			Assert.AreEqual(@"{""level"":""Information"",""message"":""in"",""format"":""{Line$2}"",""properties"":{""Var1"":123,""Var2"":{""$type"":""SourceFile"",""$text"":""D:\\build\\\u002B\u002BUE5\\Sync\\GenerateProjectFiles.bat"",""relativePath"":""GenerateProjectFiles.bat"",""depotPath"":""//UE5/Main/GenerateProjectFiles.bat@20392842""},""Line$2"":""in""},""line"":4,""lineCount"":7}", result[4]);
+			Assert.AreEqual(@"{""level"":""Information"",""message"":""four-end"",""format"":""{Line$3}-end"",""properties"":{""Var1"":123,""Var2"":{""$type"":""SourceFile"",""$text"":""D:\\build\\\u002B\u002BUE5\\Sync\\GenerateProjectFiles.bat"",""relativePath"":""GenerateProjectFiles.bat"",""depotPath"":""//UE5/Main/GenerateProjectFiles.bat@20392842""},""Line$3"":""four""},""line"":5,""lineCount"":7}", result[5]);
+			Assert.AreEqual(@"{""level"":""Information"",""message"":""log message 123 D:\\build\\\u002B\u002BUE5\\Sync\\GenerateProjectFiles.bat"",""format"":""log message {Var1} {Var2}"",""properties"":{""Var1"":123,""Var2"":{""$type"":""SourceFile"",""$text"":""D:\\build\\\u002B\u002BUE5\\Sync\\GenerateProjectFiles.bat"",""relativePath"":""GenerateProjectFiles.bat"",""depotPath"":""//UE5/Main/GenerateProjectFiles.bat@20392842""}},""line"":6,""lineCount"":7}", result[6]);
 		}
 
 		[TestMethod]
@@ -195,5 +198,94 @@ namespace Horde.Agent.Tests
 				Assert.AreEqual(extractedBodyText[idx], $"Testing {idx}");
 			}
 		}
+
+		[TestMethod]
+		public void LongMessageTest()
+		{
+			DateTime time = new DateTime(2023, 1, 1, 0, 0, 0, 0, DateTimeKind.Utc);
+
+			string LongLineA = String.Join("\n", Enumerable.Range(0, 5).Select(x => $"A{x}"));
+			string LongLineB = String.Join("\n", Enumerable.Range(0, 5).Select(x => $"B{x}"));
+			string LongLineC = String.Join("\n", Enumerable.Range(0, 5).Select(x => $"C{x}"));
+
+			Dictionary<string, object> properties = new Dictionary<string, object>
+			{
+				["LongLineA"] = LongLineA,
+				["LongLineB"] = LongLineB,
+				["LongLineC"] = LongLineC,
+			};
+
+			LogEvent baseEvent = new LogEvent(time, LogLevel.Information, default, $"start {LongLineA} x {LongLineB} y\nz {LongLineC} end", "start {LongLineA} x {LongLineB} y\nz {LongLineC} end", properties, null);
+
+			JsonRpcLogWriter writer = new JsonRpcLogWriter();
+			writer.SanitizeAndWriteEvent(new JsonLogEvent(baseEvent));
+
+			string[] output = Encoding.UTF8.GetString(writer.CreatePacket().Span).Split('\n');
+			string[] expected =
+			{
+				@"{""time"":""2023-01-01T00:00:00"",""level"":""Information"",""message"":""start A0"",""format"":""start {LongLineA$0}"",""properties"":{""LongLineA$0"":""A0""},""line"":0,""lineCount"":14}",
+				@"{""time"":""2023-01-01T00:00:00"",""level"":""Information"",""message"":""A1"",""format"":""{LongLineA$1}"",""properties"":{""LongLineA$1"":""A1""},""line"":1,""lineCount"":14}",
+				@"{""time"":""2023-01-01T00:00:00"",""level"":""Information"",""message"":""A2"",""format"":""{LongLineA$2}"",""properties"":{""LongLineA$2"":""A2""},""line"":2,""lineCount"":14}",
+				@"{""time"":""2023-01-01T00:00:00"",""level"":""Information"",""message"":""A3"",""format"":""{LongLineA$3}"",""properties"":{""LongLineA$3"":""A3""},""line"":3,""lineCount"":14}",
+				@"{""time"":""2023-01-01T00:00:00"",""level"":""Information"",""message"":""A4 x B0"",""format"":""{LongLineA$4} x {LongLineB$0}"",""properties"":{""LongLineA$4"":""A4"",""LongLineB$0"":""B0""},""line"":4,""lineCount"":14}",
+				@"{""time"":""2023-01-01T00:00:00"",""level"":""Information"",""message"":""B1"",""format"":""{LongLineB$1}"",""properties"":{""LongLineB$1"":""B1""},""line"":5,""lineCount"":14}",
+				@"{""time"":""2023-01-01T00:00:00"",""level"":""Information"",""message"":""B2"",""format"":""{LongLineB$2}"",""properties"":{""LongLineB$2"":""B2""},""line"":6,""lineCount"":14}",
+				@"{""time"":""2023-01-01T00:00:00"",""level"":""Information"",""message"":""B3"",""format"":""{LongLineB$3}"",""properties"":{""LongLineB$3"":""B3""},""line"":7,""lineCount"":14}",
+				@"{""time"":""2023-01-01T00:00:00"",""level"":""Information"",""message"":""B4 y"",""format"":""{LongLineB$4} y"",""properties"":{""LongLineB$4"":""B4""},""line"":8,""lineCount"":14}",
+				@"{""time"":""2023-01-01T00:00:00"",""level"":""Information"",""message"":""z C0"",""format"":""z {LongLineC$0}"",""properties"":{""LongLineC$0"":""C0""},""line"":9,""lineCount"":14}",
+				@"{""time"":""2023-01-01T00:00:00"",""level"":""Information"",""message"":""C1"",""format"":""{LongLineC$1}"",""properties"":{""LongLineC$1"":""C1""},""line"":10,""lineCount"":14}",
+				@"{""time"":""2023-01-01T00:00:00"",""level"":""Information"",""message"":""C2"",""format"":""{LongLineC$2}"",""properties"":{""LongLineC$2"":""C2""},""line"":11,""lineCount"":14}",
+				@"{""time"":""2023-01-01T00:00:00"",""level"":""Information"",""message"":""C3"",""format"":""{LongLineC$3}"",""properties"":{""LongLineC$3"":""C3""},""line"":12,""lineCount"":14}",
+				@"{""time"":""2023-01-01T00:00:00"",""level"":""Information"",""message"":""C4 end"",""format"":""{LongLineC$4} end"",""properties"":{""LongLineC$4"":""C4""},""line"":13,""lineCount"":14}",
+				""
+			};
+
+			for (int idx = 0; idx < output.Length; idx++)
+			{
+				Assert.AreEqual(expected[idx], output[idx]);
+			}
+		}
+
+		[TestMethod]
+		public void MaxPacketLengthTest()
+		{
+			DateTime time = new DateTime(2023, 1, 1, 0, 0, 0, 0, DateTimeKind.Utc);
+
+			LogEvent baseEvent = new LogEvent(time, LogLevel.Information, default, "abcdefghijklmnopqrstuvwxyzABCDEFGHIJKLMNOPQRSTUVWXYZ", null, null, null);
+
+			JsonRpcLogWriter writer = new JsonRpcLogWriter(maxPacketLength: 85);
+			writer.SanitizeAndWriteEvent(new JsonLogEvent(baseEvent));
+			writer.SanitizeAndWriteEvent(new JsonLogEvent(baseEvent));
+
+			for(int idx = 0; idx < 2; idx++)
+			{
+				string output = Encoding.UTF8.GetString(writer.CreatePacket().Span);
+				string expected = @"{""time"":""2023-01-01T00:00:00"",""level"":""Information"",""message"":""abcdefghijklmnopqrstuvwxyzABCDEFGHIJKLMNOPQRSTUVWXYZ""}" + "\n";
+				Assert.AreEqual(expected, output);
+			}
+		}
+
+		[TestMethod]
+		public void MaxLineLengthTest()
+		{
+			DateTime time = new DateTime(2023, 1, 1, 0, 0, 0, 0, DateTimeKind.Utc);
+
+			LogEvent baseEvent = new LogEvent(time, LogLevel.Information, default, "abcdefghijklmnopqrstuvwxyzABCDEFGHIJKLMNOPQRSTUVWXYZ", null, null, null);
+
+			JsonRpcLogWriter writer = new JsonRpcLogWriter(85);
+			writer.SanitizeAndWriteEvent(new JsonLogEvent(baseEvent));
+
+			string[] output = Encoding.UTF8.GetString(writer.CreatePacket().Span).Split('\n');
+			string[] expected =
+			{
+				@"{""time"":""2023-01-01T00:00:00"",""level"":""Information"",""message"":""abcdefghijklmnopqrstuvwxyzABCDEFGHI [...]""}",
+				""
+			};
+
+			for (int idx = 0; idx < output.Length; idx++)
+			{
+				Assert.AreEqual(expected[idx], output[idx]);
+			}
+		}
 	}
 }
diff --git a/Engine/Source/Programs/Horde/Horde.Agent/Commands/Utilities/BuildGraphCommand.cs b/Engine/Source/Programs/Horde/Horde.Agent/Commands/Utilities/BuildGraphCommand.cs
--- a/Engine/Source/Programs/Horde/Horde.Agent/Commands/Utilities/BuildGraphCommand.cs
+++ b/Engine/Source/Programs/Horde/Horde.Agent/Commands/Utilities/BuildGraphCommand.cs
@@ -11,6 +11,7 @@ using System.Threading;
 using System.Threading.Tasks;
 using EpicGames.Core;
 using Horde.Agent.Parser;
+using Horde.Agent.Utility;
 using HordeCommon;
 using HordeCommon.Rpc;
 using Microsoft.Extensions.Logging;
diff --git a/Engine/Source/Programs/Horde/Horde.Agent/Services/ServerLoggerFactory.cs b/Engine/Source/Programs/Horde/Horde.Agent/Services/ServerLoggerFactory.cs
--- a/Engine/Source/Programs/Horde/Horde.Agent/Services/ServerLoggerFactory.cs
+++ b/Engine/Source/Programs/Horde/Horde.Agent/Services/ServerLoggerFactory.cs
@@ -6,6 +6,7 @@ using Microsoft.Extensions.Logging;
 using EpicGames.Horde.Storage.Backends;
 using System.Net.Http;
 using Microsoft.Extensions.Options;
+using Horde.Agent.Utility;
 
 namespace Horde.Agent.Services
 {
diff --git a/Engine/Source/Programs/Horde/Horde.Agent/Utilities/HordeRpc.cs b/Engine/Source/Programs/Horde/Horde.Agent/Utilities/HordeRpc.cs
--- a/Engine/Source/Programs/Horde/Horde.Agent/Utilities/HordeRpc.cs
+++ b/Engine/Source/Programs/Horde/Horde.Agent/Utilities/HordeRpc.cs
@@ -167,12 +167,12 @@ namespace HordeCommon.Rpc
 
 	partial class WriteOutputRequest
 	{
-		public WriteOutputRequest(string logId, long offset, int lineIndex, byte[] data, bool flush)
+		public WriteOutputRequest(string logId, long offset, int lineIndex, ByteString data, bool flush)
 		{
 			LogId = logId;
 			Offset = offset;
 			LineIndex = lineIndex;
-			Data = ByteString.CopyFrom(data);
+			Data = data;
 			Flush = flush;
 		}
 	}
diff --git a/Engine/Source/Programs/Horde/Horde.Agent/Utility/JsonRpcLogSink.cs b/Engine/Source/Programs/Horde/Horde.Agent/Utility/JsonRpcLogSink.cs
new file mode 100644
--- /dev/null
+++ b/Engine/Source/Programs/Horde/Horde.Agent/Utility/JsonRpcLogSink.cs
@@ -0,0 +1,272 @@
+// Copyright Epic Games, Inc. All Rights Reserved.
+
+using System;
+using System.Collections.Generic;
+using System.Threading;
+using System.Threading.Tasks;
+using EpicGames.Core;
+using EpicGames.Horde.Logs;
+using EpicGames.Horde.Storage;
+using Google.Protobuf;
+using Grpc.Core;
+using Horde.Common.Rpc;
+using HordeCommon;
+using HordeCommon.Rpc;
+using Microsoft.Extensions.Logging;
+using static Horde.Common.Rpc.LogRpc;
+
+namespace Horde.Agent.Utility
+{
+	interface IJsonRpcLogSink : IAsyncDisposable
+	{
+		Task WriteEventsAsync(List<CreateEventRequest> events, CancellationToken cancellationToken);
+		Task WriteOutputAsync(WriteOutputRequest request, CancellationToken cancellationToken);
+		Task SetOutcomeAsync(JobStepOutcome outcome, CancellationToken cancellationToken);
+	}
+
+	sealed class JsonRpcLogSink : IJsonRpcLogSink
+	{
+		readonly IRpcConnection _rpcClient;
+		readonly string? _jobId;
+		readonly string? _jobBatchId;
+		readonly string? _jobStepId;
+		readonly ILogger _logger;
+
+		public JsonRpcLogSink(IRpcConnection rpcClient, string? jobId, string? jobBatchId, string? jobStepId, ILogger logger)
+		{
+			_rpcClient = rpcClient;
+			_jobId = jobId;
+			_jobBatchId = jobBatchId;
+			_jobStepId = jobStepId;
+			_logger = logger;
+		}
+
+		public ValueTask DisposeAsync() => new ValueTask();
+
+		/// <inheritdoc/>
+		public async Task WriteEventsAsync(List<CreateEventRequest> events, CancellationToken cancellationToken)
+		{
+			await _rpcClient.InvokeAsync((HordeRpc.HordeRpcClient x) => x.CreateEventsAsync(new CreateEventsRequest(events)), cancellationToken);
+		}
+
+		/// <inheritdoc/>
+		public async Task WriteOutputAsync(WriteOutputRequest request, CancellationToken cancellationToken)
+		{
+			await _rpcClient.InvokeAsync((HordeRpc.HordeRpcClient x) => x.WriteOutputAsync(request), cancellationToken);
+		}
+
+		/// <inheritdoc/>
+		public async Task SetOutcomeAsync(JobStepOutcome outcome, CancellationToken cancellationToken)
+		{
+			// Update the outcome of this jobstep
+			if (_jobId != null && _jobBatchId != null && _jobStepId != null)
+			{
+				try
+				{
+					await _rpcClient.InvokeAsync((HordeRpc.HordeRpcClient x) => x.UpdateStepAsync(new UpdateStepRequest(_jobId, _jobBatchId, _jobStepId, JobStepState.Unspecified, outcome)), cancellationToken);
+				}
+				catch (Exception ex)
+				{
+					_logger.LogWarning(ex, "Unable to update step outcome to {NewOutcome}", outcome);
+				}
+			}
+		}
+	}
+
+	class JsonRpcAndStorageLogSink : IJsonRpcLogSink, IAsyncDisposable
+	{
+		const int FlushLength = 1024 * 1024;
+
+		readonly IRpcConnection _connection;
+		readonly string _logId;
+		readonly LogBuilder _builder;
+		readonly IJsonRpcLogSink? _inner;
+		readonly TreeWriter _writer;
+		readonly ILogger _logger;
+
+		int _bufferLength;
+
+		// Background task
+		readonly object _lockObject = new object();
+
+		// Tailing task
+		readonly Task _tailTask;
+		AsyncEvent _tailTaskStop;
+		readonly AsyncEvent _newTailDataEvent = new AsyncEvent();
+
+		public JsonRpcAndStorageLogSink(IRpcConnection connection, string logId, IJsonRpcLogSink? inner, IStorageClient store, ILogger logger)
+		{
+			_connection = connection;
+			_logId = logId;
+			_builder = new LogBuilder(LogFormat.Json, logger);
+			_inner = inner;
+			_writer = new TreeWriter(store);
+			_logger = logger;
+
+			_tailTaskStop = new AsyncEvent();
+			_tailTask = Task.Run(() => TickTailAsync());
+		}
+
+		public async ValueTask DisposeAsync()
+		{
+			if (_tailTaskStop != null)
+			{
+				_tailTaskStop.Latch();
+				try
+				{
+					await _tailTask;
+				}
+				catch (OperationCanceledException)
+				{
+				}
+				catch (Exception ex)
+				{
+					_logger.LogError(ex, "Exception on log tailing task ({LogId}): {Message}", _logId, ex.Message);
+				}
+				_tailTaskStop = null!;
+			}
+
+			if (_inner != null)
+			{
+				await _inner.DisposeAsync();
+			}
+
+			if (_writer != null)
+			{
+				_writer.Dispose();
+			}
+		}
+
+		async Task TickTailAsync()
+		{
+			int tailNext = -1;
+			while (!_tailTaskStop.IsSet())
+			{
+				Task newTailDataTask;
+
+				// Get the data to send to the server
+				ReadOnlyMemory<byte> tailData = ReadOnlyMemory<byte>.Empty;
+				lock (_lockObject)
+				{
+					if (tailNext != -1)
+					{
+						tailNext = Math.Max(tailNext, _builder.FlushedLineCount);
+						tailData = _builder.ReadTailData(tailNext, 16 * 1024);
+					}
+					newTailDataTask = _newTailDataEvent.Task;
+				}
+
+				// If we don't have any updates for the server, wait until we do.
+				if (tailNext != -1 && tailData.IsEmpty)
+				{
+					await newTailDataTask;
+					continue;
+				}
+
+				// Update the next tailing position
+				int newTailNext = await UpdateLogTailAsync(tailNext, tailData);
+				if (newTailNext != tailNext)
+				{
+					tailNext = newTailNext;
+					_logger.LogInformation("Modified tail position for log {LogId} to {TailNext}", _logId, tailNext);
+				}
+			}
+		}
+
+		/// <inheritdoc/>
+		public async Task SetOutcomeAsync(JobStepOutcome outcome, CancellationToken cancellationToken)
+		{
+			if (_inner != null)
+			{
+				await _inner.SetOutcomeAsync(outcome, cancellationToken);
+			}
+		}
+
+		/// <inheritdoc/>
+		public async Task WriteEventsAsync(List<CreateEventRequest> events, CancellationToken cancellationToken)
+		{
+			if (_inner != null)
+			{
+				await _inner.WriteEventsAsync(events, cancellationToken);
+			}
+		}
+
+		/// <inheritdoc/>
+		public async Task WriteOutputAsync(WriteOutputRequest request, CancellationToken cancellationToken)
+		{
+			if (_inner != null)
+			{
+				await _inner.WriteOutputAsync(request, cancellationToken);
+			}
+
+			_builder.WriteData(request.Data.Memory);
+			_bufferLength += request.Data.Length;
+
+			if (request.Flush || _bufferLength > FlushLength)
+			{
+				NodeHandle target = await _builder.FlushAsync(_writer, request.Flush, cancellationToken);
+				await UpdateLogAsync(target, _builder.LineCount, cancellationToken);
+				_bufferLength = 0;
+			}
+
+			_newTailDataEvent.Set();
+		}
+
+		#region RPC calls
+
+		protected virtual async Task UpdateLogAsync(NodeHandle target, int lineCount, CancellationToken cancellationToken)
+		{
+			_logger.LogDebug("Updating log {LogId} to line {LineCount}, target {Locator}", _logId, lineCount, target);
+
+			UpdateLogRequest request = new UpdateLogRequest();
+			request.LogId = _logId;
+			request.LineCount = lineCount;
+			request.Target = target.ToString();
+			await _connection.InvokeAsync((LogRpcClient client) => client.UpdateLogAsync(request, cancellationToken: cancellationToken), cancellationToken);
+		}
+
+		protected virtual async Task<int> UpdateLogTailAsync(int tailNext, ReadOnlyMemory<byte> tailData)
+		{
+			DateTime deadline = DateTime.UtcNow.AddMinutes(2.0);
+			using (IRpcClientRef<LogRpcClient> clientRef = await _connection.GetClientRefAsync<LogRpcClient>(CancellationToken.None))
+			{
+				using (AsyncDuplexStreamingCall<UpdateLogTailRequest, UpdateLogTailResponse> call = clientRef.Client.UpdateLogTail(deadline: deadline))
+				{
+					// Write the request to the server
+					UpdateLogTailRequest request = new UpdateLogTailRequest();
+					request.LogId = _logId;
+					request.TailNext = tailNext;
+					request.TailData = UnsafeByteOperations.UnsafeWrap(tailData);
+					await call.RequestStream.WriteAsync(request);
+
+					// Wait until the server responds or we need to trigger a new update
+					Task<bool> moveNextAsync = call.ResponseStream.MoveNext();
+
+					Task task = await Task.WhenAny(moveNextAsync, clientRef.DisposingTask, _tailTaskStop.Task, Task.Delay(TimeSpan.FromMinutes(1.0), CancellationToken.None));
+					if (task == clientRef.DisposingTask)
+					{
+						_logger.LogDebug("Cancelling long poll from client side (server migration)");
+					}
+					else if (task == _tailTaskStop.Task)
+					{
+						_logger.LogDebug("Cancelling long poll from client side (complete)");
+					}
+
+					// Close the request stream to indicate that we're finished
+					await call.RequestStream.CompleteAsync();
+
+					// Wait for a response or a new update to come in, then close the request stream
+					UpdateLogTailResponse? response = null;
+					while (await moveNextAsync)
+					{
+						response = call.ResponseStream.Current;
+						moveNextAsync = call.ResponseStream.MoveNext();
+					}
+					return response?.TailNext ?? -1;
+				}
+			}
+		}
+
+		#endregion
+	}
+}
diff --git a/Engine/Source/Programs/Horde/Horde.Agent/Utility/JsonRpcLogWriter.cs b/Engine/Source/Programs/Horde/Horde.Agent/Utility/JsonRpcLogWriter.cs
new file mode 100644
--- /dev/null
+++ b/Engine/Source/Programs/Horde/Horde.Agent/Utility/JsonRpcLogWriter.cs
@@ -0,0 +1,320 @@
+// Copyright Epic Games, Inc. All Rights Reserved.
+
+using System;
+using System.Buffers;
+using System.Collections.Generic;
+using System.Diagnostics;
+using System.Diagnostics.CodeAnalysis;
+using System.Linq;
+using System.Text;
+using System.Text.Json;
+using System.Text.Json.Nodes;
+using System.Threading.Tasks;
+using EpicGames.Core;
+
+namespace Horde.Agent.Utility
+{
+	using JsonObject = System.Text.Json.Nodes.JsonObject;
+
+	class JsonRpcLogWriter
+	{
+		[DebuggerDisplay("{Format}")]
+		class FormattedLine
+		{
+			public string Format { get; set; }
+			public Dictionary<string, string> Properties { get; } = new Dictionary<string, string>(StringComparer.OrdinalIgnoreCase);
+
+			public FormattedLine(string format)
+			{
+				Format = format;
+			}
+		}
+
+		static readonly string s_messagePropertyName = LogEventPropertyName.Message.ToString();
+		static readonly string s_formatPropertyName = LogEventPropertyName.Format.ToString();
+		static readonly string s_linePropertyName = LogEventPropertyName.Line.ToString();
+		static readonly string s_lineCountPropertyName = LogEventPropertyName.LineCount.ToString();
+
+		static readonly Utf8String s_escapedNewline = "\\n";
+
+		readonly List<JsonLogEvent> _logEvents = new List<JsonLogEvent>();
+		readonly ArrayBufferWriter<byte> _lineWriter;
+		readonly ArrayBufferWriter<byte> _packetWriter;
+		int _packetLength;
+
+		/// <summary>
+		/// Current packet length
+		/// </summary>
+		public int PacketLength => _packetLength;
+
+		/// <summary>
+		/// Maximum length of an individual line
+		/// </summary>
+		public int MaxLineLength { get; }
+
+		/// <summary>
+		/// Maximum size of a packet
+		/// </summary>
+		public int MaxPacketLength { get; }
+
+		/// <summary>
+		/// Constructor
+		/// </summary>
+		/// <param name="maxLineLength">Maximum length for an individual line</param>
+		/// <param name="maxPacketLength">Maximum length for a packet</param>
+		public JsonRpcLogWriter(int maxLineLength = 4 * 1024, int maxPacketLength = 256 * 1024)
+		{
+			MaxLineLength = maxLineLength;
+			_lineWriter = new ArrayBufferWriter<byte>(maxLineLength);
+
+			MaxPacketLength = maxPacketLength;
+			_packetWriter = new ArrayBufferWriter<byte>(maxPacketLength);
+		}
+
+		/// <summary>
+		/// Creates a packet from the current data
+		/// </summary>
+		/// <returns>Packet data</returns>
+		public ReadOnlyMemory<byte> CreatePacket()
+		{
+			_packetWriter.Clear();
+
+			int eventCount = 0;
+			for (; eventCount < _logEvents.Count && (eventCount == 0 || _packetWriter.WrittenCount + _logEvents[eventCount].Data.Length + 1 < MaxPacketLength); eventCount++)
+			{
+				JsonLogEvent jsonLogEvent = _logEvents[eventCount];
+
+				Span<byte> span = _packetWriter.GetSpan(jsonLogEvent.Data.Length + 1);
+				jsonLogEvent.Data.Span.CopyTo(span);
+				span[jsonLogEvent.Data.Length] = (byte)'\n';
+				_packetWriter.Advance(jsonLogEvent.Data.Length + 1);
+
+				_packetLength -= jsonLogEvent.Data.Length + 1;
+			}
+			_logEvents.RemoveRange(0, eventCount);
+
+			return _packetWriter.WrittenMemory;
+		}
+
+		/// <summary>
+		/// Writes an event
+		/// </summary>
+		/// <param name="jsonLogEvent">Event to write</param>
+		public int SanitizeAndWriteEvent(JsonLogEvent jsonLogEvent)
+		{
+			ReadOnlySpan<byte> span = jsonLogEvent.Data.Span;
+			if (jsonLogEvent.LineCount == 1 && span.IndexOf(s_escapedNewline) != -1)
+			{
+				JsonObject obj = (JsonObject)JsonNode.Parse(span)!;
+
+				JsonValue? formatValue = obj["format"] as JsonValue;
+				if (formatValue != null && formatValue.TryGetValue(out string? format))
+				{
+					return WriteEventWithFormat(jsonLogEvent, obj, format);
+				}
+
+				JsonValue? messageValue = obj["message"] as JsonValue;
+				if (messageValue != null && messageValue.TryGetValue(out string? message))
+				{
+					return WriteEventWithMessage(jsonLogEvent, obj, message);
+				}
+			}
+
+			WriteEventInternal(jsonLogEvent);
+			return 1;
+		}
+
+		/// <summary>
+		/// Writes an event with a format string, splitting it into multiple lines if necessary
+		/// </summary>
+		int WriteEventWithFormat(JsonLogEvent jsonLogEvent, JsonObject obj, string format)
+		{
+			// There is an object containing "common" properties
+			IEnumerable<KeyValuePair<string, object?>> propertyValueList = Enumerable.Empty<KeyValuePair<string, object?>>();
+
+			// Split the format string into lines
+			string[] formatLines = format.Split('\n');
+			List<FormattedLine> lines = new List<FormattedLine>();
+
+			// Split all the multi-line properties into separate properties
+			JsonObject? properties = obj["properties"] as JsonObject;
+			if (properties == null)
+			{
+				lines.AddRange(formatLines.Select(x => new FormattedLine(x)));
+			}
+			else
+			{
+				// Get all the current property values
+				Dictionary<string, string> propertyValues = new Dictionary<string, string>(StringComparer.OrdinalIgnoreCase);
+				foreach ((string name, JsonNode? node) in properties)
+				{
+					string value = String.Empty;
+					if (node != null)
+					{
+						if (node is JsonObject valueObject)
+						{
+							value = valueObject["$text"]?.ToString() ?? String.Empty;
+						}
+						else
+						{
+							value = node.ToString();
+						}
+					}
+					propertyValues[name] = value;
+				}
+
+				// Keep splitting properties in the last line
+				foreach (string formatLine in formatLines)
+				{
+					lines.Add(new FormattedLine(formatLine));
+					for (; ; )
+					{
+						FormattedLine line = lines[^1];
+						if (!TryGetNextMultiLineProperty(line.Format, propertyValues, out string? prefix, out string? suffix, out string? propertyName, out string[]? propertyLines))
+						{
+							break;
+						}
+
+						properties!.Remove(propertyName);
+
+						StringBuilder builder = new StringBuilder();
+						builder.Append(prefix);
+
+						for (int lineNum = 0; ; lineNum++)
+						{
+							string newPropertyName = $"{propertyName}${lineNum}";
+							builder.Append($"{{{newPropertyName}}}");
+
+							line.Properties.Add(newPropertyName, propertyLines[lineNum]);
+
+							if (lineNum + 1 >= propertyLines.Length)
+							{
+								break;
+							}
+
+							line.Format = builder.ToString();
+							builder.Clear();
+
+							line = new FormattedLine(String.Empty);
+							lines.Add(line);
+						}
+
+						builder.Append(suffix);
+						line.Format = builder.ToString();
+					}
+				}
+
+				// Get the enumerable property list for formatting
+				propertyValueList = propertyValues.Select(x => new KeyValuePair<string, object?>(x.Key, x.Value));
+			}
+
+			// Finally split the format string into multiple lines
+			for (int idx = 0; idx < lines.Count; idx++)
+			{
+				FormattedLine line = lines[idx];
+				foreach ((string name, string value) in line.Properties)
+				{
+					properties![name] = value;
+				}
+
+				string message = MessageTemplate.Render(lines[idx].Format, propertyValueList.Concat(line.Properties.Select(x => new KeyValuePair<string, object?>(x.Key, x.Value))));
+				WriteSingleEvent(jsonLogEvent, obj, message, lines[idx].Format, idx, lines.Count);
+
+				foreach ((string name, _) in line.Properties)
+				{
+					properties!.Remove(name);
+				}
+			}
+			return lines.Count;
+		}
+
+		static bool TryGetNextMultiLineProperty(string format, Dictionary<string, string> properties, [NotNullWhen(true)] out string? prefix, [NotNullWhen(true)] out string? suffix, [NotNullWhen(true)] out string? propertyName, [NotNullWhen(true)] out string[]? propertyLines)
+		{
+			int nameStart = -1;
+			for (int idx = 0; idx < format.Length; idx++)
+			{
+				if (format[idx] == '{')
+				{
+					nameStart = idx + 1;
+				}
+				else if (format[idx] == '}' && nameStart != -1)
+				{
+					string name = format.Substring(nameStart, idx - nameStart);
+					if (properties.TryGetValue(name, out string? text))
+					{
+						int textLineEnd = text.IndexOf('\n', StringComparison.Ordinal);
+						if (textLineEnd != -1)
+						{
+							prefix = format.Substring(0, nameStart - 1);
+							suffix = format.Substring(idx + 1);
+							propertyName = name;
+							propertyLines = text.Split('\n');
+							return true;
+						}
+					}
+				}
+			}
+
+			prefix = null;
+			suffix = null;
+			propertyName = null;
+			propertyLines = null;
+			return false;
+		}
+
+		int WriteEventWithMessage(JsonLogEvent jsonLogEvent, JsonObject obj, string message)
+		{
+			string[] lines = message.Split('\n');
+			for (int idx = 0; idx < lines.Length; idx++)
+			{
+				WriteSingleEvent(jsonLogEvent, obj, lines[idx], null, idx, lines.Length);
+			}
+			return lines.Length;
+		}
+
+		void WriteSingleEvent(JsonLogEvent jsonLogEvent, JsonObject obj, string message, string? format, int line, int lineCount)
+		{
+			obj[s_messagePropertyName] = message;
+			if (format != null)
+			{
+				obj[s_formatPropertyName] = format;
+			}
+			if (lineCount > 1)
+			{
+				obj[s_linePropertyName] = line;
+				obj[s_lineCountPropertyName] = lineCount;
+			}
+
+			_lineWriter.Clear();
+			using (Utf8JsonWriter jsonWriter = new Utf8JsonWriter(_lineWriter))
+			{
+				obj.WriteTo(jsonWriter);
+			}
+
+			JsonLogEvent nextEvent = new JsonLogEvent(jsonLogEvent.Level, jsonLogEvent.EventId, line, lineCount, _lineWriter.WrittenMemory.ToArray());
+			WriteEventInternal(nextEvent);
+		}
+
+		void WriteEventInternal(JsonLogEvent jsonLogEvent)
+		{
+			if (jsonLogEvent.Data.Length > MaxLineLength)
+			{
+				LogEvent logEvent = LogEvent.Read(jsonLogEvent.Data.Span);
+
+				int maxMessageLength = Math.Max(10, MaxLineLength - 50);
+
+				string message = logEvent.Message;
+				if (message.Length > maxMessageLength)
+				{
+					message = message.Substring(0, maxMessageLength);
+				}
+
+				logEvent = new LogEvent(logEvent.Time, logEvent.Level, logEvent.Id, logEvent.LineIndex, logEvent.LineCount, $"{message} [...]", null, null, logEvent.Exception);
+				jsonLogEvent = new JsonLogEvent(logEvent);
+			}
+
+			_logEvents.Add(jsonLogEvent);
+			_packetLength += jsonLogEvent.Data.Length + 1;
+		}
+	}
+}
diff --git a/Engine/Source/Programs/Horde/Horde.Agent/Utility/JsonRpcLogger.cs b/Engine/Source/Programs/Horde/Horde.Agent/Utility/JsonRpcLogger.cs
--- a/Engine/Source/Programs/Horde/Horde.Agent/Utility/JsonRpcLogger.cs
+++ b/Engine/Source/Programs/Horde/Horde.Agent/Utility/JsonRpcLogger.cs
@@ -4,8 +4,8 @@ using System;
 using System.Buffers;
 using System.Collections.Generic;
 using System.Diagnostics;
+using System.Diagnostics.CodeAnalysis;
 using System.Linq;
-using System.Net.Http;
 using System.Text;
 using System.Text.Json;
 using System.Text.Json.Nodes;
@@ -13,277 +13,15 @@ using System.Threading;
 using System.Threading.Channels;
 using System.Threading.Tasks;
 using EpicGames.Core;
-using EpicGames.Horde.Logs;
-using EpicGames.Horde.Storage;
-using EpicGames.Horde.Storage.Backends;
 using Google.Protobuf;
-using Grpc.Core;
-using Horde.Agent.Services;
-using Horde.Agent.Utility;
-using Horde.Common.Rpc;
 using HordeCommon;
 using HordeCommon.Rpc;
 using Microsoft.Extensions.Logging;
 
-using static Horde.Common.Rpc.LogRpc;
-
-namespace Horde.Agent.Parser
+namespace Horde.Agent.Utility
 {
 	using JsonObject = System.Text.Json.Nodes.JsonObject;
 
-	interface IJsonRpcLogSink : IAsyncDisposable
-	{
-		Task WriteEventsAsync(List<CreateEventRequest> events, CancellationToken cancellationToken);
-		Task WriteOutputAsync(WriteOutputRequest request, CancellationToken cancellationToken);
-		Task SetOutcomeAsync(JobStepOutcome outcome, CancellationToken cancellationToken);
-	}
-
-	sealed class JsonRpcLogSink : IJsonRpcLogSink
-	{
-		readonly IRpcConnection _rpcClient;
-		readonly string? _jobId;
-		readonly string? _jobBatchId;
-		readonly string? _jobStepId;
-		readonly ILogger _logger;
-
-		public JsonRpcLogSink(IRpcConnection rpcClient, string? jobId, string? jobBatchId, string? jobStepId, ILogger logger)
-		{
-			_rpcClient = rpcClient;
-			_jobId = jobId;
-			_jobBatchId = jobBatchId;
-			_jobStepId = jobStepId;
-			_logger = logger;
-		}
-
-		public ValueTask DisposeAsync() => new ValueTask();
-
-		/// <inheritdoc/>
-		public async Task WriteEventsAsync(List<CreateEventRequest> events, CancellationToken cancellationToken)
-		{
-			await _rpcClient.InvokeAsync((HordeRpc.HordeRpcClient x) => x.CreateEventsAsync(new CreateEventsRequest(events)), cancellationToken);
-		}
-
-		/// <inheritdoc/>
-		public async Task WriteOutputAsync(WriteOutputRequest request, CancellationToken cancellationToken)
-		{
-			await _rpcClient.InvokeAsync((HordeRpc.HordeRpcClient x) => x.WriteOutputAsync(request), cancellationToken);
-		}
-
-		/// <inheritdoc/>
-		public async Task SetOutcomeAsync(JobStepOutcome outcome, CancellationToken cancellationToken)
-		{
-			// Update the outcome of this jobstep
-			if (_jobId != null && _jobBatchId != null && _jobStepId != null)
-			{
-				try
-				{
-					await _rpcClient.InvokeAsync((HordeRpc.HordeRpcClient x) => x.UpdateStepAsync(new UpdateStepRequest(_jobId, _jobBatchId, _jobStepId, JobStepState.Unspecified, outcome)), cancellationToken);
-				}
-				catch (Exception ex)
-				{
-					_logger.LogWarning(ex, "Unable to update step outcome to {NewOutcome}", outcome);
-				}
-			}
-		}
-	}
-
-	class JsonRpcAndStorageLogSink : IJsonRpcLogSink, IAsyncDisposable
-	{
-		const int FlushLength = 1024 * 1024;
-
-		readonly IRpcConnection _connection;
-		readonly string _logId;
-		readonly LogBuilder _builder;
-		readonly IJsonRpcLogSink? _inner;
-		readonly TreeWriter _writer;
-		readonly ILogger _logger;
-
-		int _bufferLength;
-
-		// Background task
-		readonly object _lockObject = new object();
-
-		// Tailing task
-		readonly Task _tailTask;
-		AsyncEvent _tailTaskStop;
-		readonly AsyncEvent _newTailDataEvent = new AsyncEvent();
-
-		public JsonRpcAndStorageLogSink(IRpcConnection connection, string logId, IJsonRpcLogSink? inner, IStorageClient store, ILogger logger)
-		{
-			_connection = connection;
-			_logId = logId;
-			_builder = new LogBuilder(LogFormat.Json, logger);
-			_inner = inner;
-			_writer = new TreeWriter(store);
-			_logger = logger;
-
-			_tailTaskStop = new AsyncEvent();
-			_tailTask = Task.Run(() => TickTailAsync());
-		}
-
-		public async ValueTask DisposeAsync()
-		{
-			if(_tailTaskStop != null)
-			{
-				_tailTaskStop.Latch();
-				try
-				{
-					await _tailTask;
-				}
-				catch (OperationCanceledException)
-				{
-				}
-				catch (Exception ex)
-				{
-					_logger.LogError(ex, "Exception on log tailing task ({LogId}): {Message}", _logId, ex.Message);
-				}
-				_tailTaskStop = null!;
-			}
-
-			if (_inner != null)
-			{
-				await _inner.DisposeAsync();
-			}
-
-			if (_writer != null)
-			{
-				_writer.Dispose();
-			}
-		}
-
-		async Task TickTailAsync()
-		{
-			int tailNext = -1;
-			while (!_tailTaskStop.IsSet())
-			{
-				Task newTailDataTask;
-
-				// Get the data to send to the server
-				ReadOnlyMemory<byte> tailData = ReadOnlyMemory<byte>.Empty;
-				lock (_lockObject)
-				{
-					if (tailNext != -1)
-					{
-						tailNext = Math.Max(tailNext, _builder.FlushedLineCount);
-						tailData = _builder.ReadTailData(tailNext, 16 * 1024);
-					}
-					newTailDataTask = _newTailDataEvent.Task;
-				}
-
-				// If we don't have any updates for the server, wait until we do.
-				if (tailNext != -1 && tailData.IsEmpty)
-				{
-					await newTailDataTask;
-					continue;
-				}
-
-				// Update the next tailing position
-				int newTailNext = await UpdateLogTailAsync(tailNext, tailData);
-				if (newTailNext != tailNext)
-				{
-					tailNext = newTailNext;
-					_logger.LogInformation("Modified tail position for log {LogId} to {TailNext}", _logId, tailNext);
-				}
-			}
-		}
-
-		/// <inheritdoc/>
-		public async Task SetOutcomeAsync(JobStepOutcome outcome, CancellationToken cancellationToken)
-		{
-			if (_inner != null)
-			{
-				await _inner.SetOutcomeAsync(outcome, cancellationToken);
-			}
-		}
-
-		/// <inheritdoc/>
-		public async Task WriteEventsAsync(List<CreateEventRequest> events, CancellationToken cancellationToken)
-		{
-			if (_inner != null)
-			{
-				await _inner.WriteEventsAsync(events, cancellationToken);
-			}
-		}
-
-		/// <inheritdoc/>
-		public async Task WriteOutputAsync(WriteOutputRequest request, CancellationToken cancellationToken)
-		{
-			if (_inner != null)
-			{
-				await _inner.WriteOutputAsync(request, cancellationToken);
-			}
-
-			_builder.WriteData(request.Data.Memory);
-			_bufferLength += request.Data.Length;
-
-			if (request.Flush || _bufferLength > FlushLength)
-			{
-				NodeHandle target = await _builder.FlushAsync(_writer, request.Flush, cancellationToken);
-				await UpdateLogAsync(target, _builder.LineCount, cancellationToken);
-				_bufferLength = 0;
-			}
-
-			_newTailDataEvent.Set();
-		}
-
-		#region RPC calls
-
-		protected virtual async Task UpdateLogAsync(NodeHandle target, int lineCount, CancellationToken cancellationToken)
-		{
-			_logger.LogDebug("Updating log {LogId} to line {LineCount}, target {Locator}", _logId, lineCount, target);
-
-			UpdateLogRequest request = new UpdateLogRequest();
-			request.LogId = _logId;
-			request.LineCount = lineCount;
-			request.Target = target.ToString();
-			await _connection.InvokeAsync((LogRpcClient client) => client.UpdateLogAsync(request, cancellationToken: cancellationToken), cancellationToken);
-		}
-
-		protected virtual async Task<int> UpdateLogTailAsync(int tailNext, ReadOnlyMemory<byte> tailData)
-		{
-			DateTime deadline = DateTime.UtcNow.AddMinutes(2.0);
-			using (IRpcClientRef<LogRpcClient> clientRef = await _connection.GetClientRefAsync<LogRpcClient>(CancellationToken.None))
-			{
-				using (AsyncDuplexStreamingCall<UpdateLogTailRequest, UpdateLogTailResponse> call = clientRef.Client.UpdateLogTail(deadline: deadline))
-				{
-					// Write the request to the server
-					UpdateLogTailRequest request = new UpdateLogTailRequest();
-					request.LogId = _logId;
-					request.TailNext = tailNext;
-					request.TailData = UnsafeByteOperations.UnsafeWrap(tailData);
-					await call.RequestStream.WriteAsync(request);
-
-					// Wait until the server responds or we need to trigger a new update
-					Task<bool> moveNextAsync = call.ResponseStream.MoveNext();
-
-					Task task = await Task.WhenAny(moveNextAsync, clientRef.DisposingTask, _tailTaskStop.Task, Task.Delay(TimeSpan.FromMinutes(1.0), CancellationToken.None));
-					if (task == clientRef.DisposingTask)
-					{
-						_logger.LogDebug("Cancelling long poll from client side (server migration)");
-					}
-					else if (task == _tailTaskStop.Task)
-					{
-						_logger.LogDebug("Cancelling long poll from client side (complete)");
-					}
-
-					// Close the request stream to indicate that we're finished
-					await call.RequestStream.CompleteAsync();
-
-					// Wait for a response or a new update to come in, then close the request stream
-					UpdateLogTailResponse? response = null;
-					while (await moveNextAsync)
-					{
-						response = call.ResponseStream.Current;
-						moveNextAsync = call.ResponseStream.MoveNext();
-					}
-					return response?.TailNext ?? -1;
-				}
-			}
-		}
-
-		#endregion
-	}
-
 	/// <summary>
 	/// Interface for a log device
 	/// </summary>
@@ -305,23 +43,6 @@ namespace Horde.Agent.Parser
 	/// </summary>
 	sealed class JsonRpcLogger : IServerLogger
 	{
-		class QueueItem
-		{
-			public byte[] Data { get; }
-			public CreateEventRequest? CreateEvent { get; }
-
-			public QueueItem(byte[] data, CreateEventRequest? createEvent)
-			{
-				Data = data;
-				CreateEvent = createEvent;
-			}
-
-			public override string ToString()
-			{
-				return Encoding.UTF8.GetString(Data);
-			}
-		}
-
 		internal readonly IJsonRpcLogSink _sink;
 		internal readonly string _logId;
 		internal readonly bool _warnings;
@@ -435,8 +156,8 @@ namespace Horde.Agent.Parser
 			const int MaxWarnings = 50;
 			int numWarnings = 0;
 
-			// Buffers for chunks and events read in a single iteration
-			ArrayBufferWriter<byte> writer = new ArrayBufferWriter<byte>();
+			// Buffer for events read in a single iteration
+			JsonRpcLogWriter writer = new JsonRpcLogWriter();
 			List<CreateEventRequest> events = new List<CreateEventRequest>();
 
 			// The current jobstep outcome
@@ -445,7 +166,6 @@ namespace Horde.Agent.Parser
 			// Whether we've written the flush command
 			for (; ; )
 			{
-				writer.Clear();
 				events.Clear();
 
 				// Save off the current line number for sending to the server
@@ -453,12 +173,12 @@ namespace Horde.Agent.Parser
 
 				// Get the next data
 				Task waitTask = Task.Delay(TimeSpan.FromSeconds(2.0));
-				while (writer.WrittenCount < 256 * 1024)
+				while (writer.PacketLength < writer.MaxPacketLength)
 				{
 					JsonLogEvent jsonLogEvent;
 					if (_dataChannel.Reader.TryRead(out jsonLogEvent))
 					{
-						int lineCount = WriteEvent(jsonLogEvent, writer);
+						int lineCount = writer.SanitizeAndWriteEvent(jsonLogEvent);
 						if (jsonLogEvent.LineIndex == 0)
 						{
 							if (jsonLogEvent.Level == LogLevel.Warning && ++numWarnings <= MaxWarnings)
@@ -472,7 +192,6 @@ namespace Horde.Agent.Parser
 						}
 						lineIndex += lineCount;
 					}
-					else
 					{
 						Task<bool> readTask = _dataChannel.Reader.WaitToReadAsync().AsTask();
 						if (await Task.WhenAny(readTask, waitTask) == waitTask)
@@ -487,18 +206,18 @@ namespace Horde.Agent.Parser
 				}
 
 				// Upload it to the server
-				if (writer.WrittenCount > 0)
+				if (writer.PacketLength > 0)
 				{
-					byte[] data = writer.WrittenSpan.ToArray();
+					ByteString packet = UnsafeByteOperations.UnsafeWrap(writer.CreatePacket());
 					try
 					{
-						await _sink.WriteOutputAsync(new WriteOutputRequest(_logId, offset, initialLineIndex, data, false), CancellationToken.None);
+						await _sink.WriteOutputAsync(new WriteOutputRequest(_logId, offset, initialLineIndex, packet, false), CancellationToken.None);
 					}
 					catch (Exception ex)
 					{
-						_inner.LogWarning(ex, "Unable to write data to server (log {LogId}, offset {Offset}, length {Length}, lines {StartLine}-{EndLine})", _logId, offset, data.Length, initialLineIndex, lineIndex);
+						_inner.LogWarning(ex, "Unable to write data to server (log {LogId}, offset {Offset}, length {Length}, lines {StartLine}-{EndLine})", _logId, offset, packet.Length, initialLineIndex, lineIndex);
 					}
-					offset += data.Length;
+					offset += packet.Length;
 				}
 
 				// Write all the events
@@ -533,7 +252,7 @@ namespace Horde.Agent.Parser
 				{
 					try
 					{
-						await _sink.WriteOutputAsync(new WriteOutputRequest(_logId, offset, lineIndex, Array.Empty<byte>(), true), CancellationToken.None);
+						await _sink.WriteOutputAsync(new WriteOutputRequest(_logId, offset, lineIndex, ByteString.Empty, true), CancellationToken.None);
 					}
 					catch (Exception ex)
 					{
@@ -544,159 +263,6 @@ namespace Horde.Agent.Parser
 			}
 		}
 
-		static readonly string s_messagePropertyName = LogEventPropertyName.Message.ToString();
-		static readonly string s_formatPropertyName = LogEventPropertyName.Format.ToString();
-		static readonly string s_linePropertyName = LogEventPropertyName.Line.ToString();
-		static readonly string s_lineCountPropertyName = LogEventPropertyName.LineCount.ToString();
-
-		static readonly Utf8String s_newline = "\n";
-		static readonly Utf8String s_escapedNewline = "\\n";
-
-		public static int WriteEvent(JsonLogEvent jsonLogEvent, IBufferWriter<byte> writer)
-		{
-			ReadOnlySpan<byte> span = jsonLogEvent.Data.Span;
-			if (jsonLogEvent.LineCount == 1 && span.IndexOf(s_escapedNewline) != -1)
-			{
-				JsonObject obj = (JsonObject)JsonNode.Parse(span)!;
-
-				JsonValue? formatValue = obj["format"] as JsonValue;
-				if (formatValue != null && formatValue.TryGetValue(out string? format))
-				{
-					return WriteEventWithFormat(obj, format, writer);
-				}
-
-				JsonValue? messageValue = obj["message"] as JsonValue;
-				if (messageValue != null && messageValue.TryGetValue(out string? message))
-				{
-					return WriteEventWithMessage(obj, message, writer);
-				}
-			}
-
-			writer.Write(span);
-			writer.Write(s_newline);
-			return 1;
-		}
-
-		static int WriteEventWithFormat(JsonObject obj, string format, IBufferWriter<byte> writer)
-		{
-			IEnumerable<KeyValuePair<string, object?>> propertyValueList = Enumerable.Empty<KeyValuePair<string, object?>>();
-
-			// Split all the multi-line properties into separate properties
-			JsonObject? properties = obj["properties"] as JsonObject;
-			if (properties != null)
-			{
-				// Get all the current property values
-				Dictionary<string, string> propertyValues = new Dictionary<string, string>(StringComparer.Ordinal);
-				foreach ((string name, JsonNode? node) in properties)
-				{
-					string value = String.Empty;
-					if (node != null)
-					{
-						if (node is JsonObject valueObject)
-						{
-							value = valueObject["$text"]?.ToString() ?? String.Empty;
-						}
-						else
-						{
-							value = node.ToString();
-						}
-					}
-					propertyValues[name] = value;
-				}
-
-				// Split all the multi-line properties into separate things
-				int nameStart = -1;
-				for (int idx = 0; idx < format.Length; idx++)
-				{
-					if (format[idx] == '{')
-					{
-						nameStart = idx + 1;
-					}
-					else if (format[idx] == '}' && nameStart != -1)
-					{
-						string name = format.Substring(nameStart, idx - nameStart);
-						if (propertyValues.TryGetValue(name, out string? text))
-						{
-							int textLineEnd = text.IndexOf('\n', StringComparison.Ordinal);
-							if (textLineEnd != -1)
-							{
-								int lineNum = 0;
-
-								StringBuilder builder = new StringBuilder();
-								builder.Append(format, 0, nameStart - 1);
-
-								string delimiter = String.Empty;
-								for (int textLineStart = 0; textLineStart < text.Length;)
-								{
-									string newName = $"{name}${lineNum++}";
-									string newLine = text.Substring(textLineStart, textLineEnd - textLineStart);
-
-									// Insert this line
-									builder.Append($"{delimiter}{{{newName}}}");
-									properties![newName] = newLine;
-									propertyValues[newName] = newLine;
-									delimiter = "\n";
-
-									// Move to the next line
-									textLineStart = ++textLineEnd;
-									while (textLineEnd < text.Length && text[textLineEnd] != '\n')
-									{
-										textLineEnd++;
-									}
-								}
-
-								builder.Append(format, idx + 1, format.Length - (idx + 1));
-								format = builder.ToString();
-							}
-						}
-					}
-				}
-
-				// Get the enumerable property list for formatting
-				propertyValueList = propertyValues.Select(x => new KeyValuePair<string, object?>(x.Key, x.Value));
-			}
-
-			// Finally split the format string into multiple lines
-			string[] lines = format.Split('\n');
-			for (int idx = 0; idx < lines.Length; idx++)
-			{
-				string message = MessageTemplate.Render(lines[idx], propertyValueList);
-				WriteSingleEvent(obj, message, lines[idx], idx, lines.Length, writer);
-			}
-			return lines.Length;
-		}
-
-		static int WriteEventWithMessage(JsonObject obj, string message, IBufferWriter<byte> writer)
-		{
-			string[] lines = message.Split('\n');
-			for (int idx = 0; idx < lines.Length; idx++)
-			{
-				WriteSingleEvent(obj, lines[idx], null, idx, lines.Length, writer);
-			}
-			return lines.Length;
-		}
-
-		static void WriteSingleEvent(JsonObject obj, string message, string? format, int line, int lineCount, IBufferWriter<byte> writer)
-		{
-			obj[s_messagePropertyName] = message;
-			if (format != null)
-			{
-				obj[s_formatPropertyName] = format;
-			}
-			if (lineCount > 1)
-			{
-				obj[s_linePropertyName] = line;
-				obj[s_lineCountPropertyName] = lineCount;
-			}
-
-			using (Utf8JsonWriter jsonWriter = new Utf8JsonWriter(writer))
-			{
-				obj.WriteTo(jsonWriter);
-			}
-
-			writer.Write(s_newline.Span);
-		}
-
 		void AddEvent(ReadOnlySpan<byte> span, int lineIndex, int lineCount, EventSeverity severity, List<CreateEventRequest> events)
 		{
 			try
diff --git a/Engine/Source/Programs/Shared/EpicGames.Core/JsonLogEvent.cs b/Engine/Source/Programs/Shared/EpicGames.Core/JsonLogEvent.cs
--- a/Engine/Source/Programs/Shared/EpicGames.Core/JsonLogEvent.cs
+++ b/Engine/Source/Programs/Shared/EpicGames.Core/JsonLogEvent.cs
@@ -1,7 +1,6 @@
 // Copyright Epic Games, Inc. All Rights Reserved.
 
 using System;
-using System.Diagnostics.CodeAnalysis;
 using System.Text;
 using System.Text.Json;
 using Microsoft.Extensions.Logging;
@@ -38,6 +37,14 @@ namespace EpicGames.Core
 		/// </summary>
 		public ReadOnlyMemory<byte> Data { get; }
 
+		/// <summary>
+		/// Constructor
+		/// </summary>
+		public JsonLogEvent(LogEvent logEvent)
+			: this(logEvent.Level, logEvent.Id, logEvent.LineIndex, logEvent.LineCount, logEvent.ToJsonBytes())
+		{
+		}
+
 		/// <summary>
 		/// Constructor
 		/// </summary>
