Horde: Fix incorrect value for StepId in telemetry.

diff --git a/Engine/Source/Programs/Horde/Horde.Build/Jobs/JobStepRefCollection.cs b/Engine/Source/Programs/Horde/Horde.Build/Jobs/JobStepRefCollection.cs
--- a/Engine/Source/Programs/Horde/Horde.Build/Jobs/JobStepRefCollection.cs
+++ b/Engine/Source/Programs/Horde/Horde.Build/Jobs/JobStepRefCollection.cs
@@ -125,7 +125,7 @@ namespace Horde.Build.Jobs
 					Id = id.ToString(),
 					JobId = id.JobId.ToString(),
 					BatchId = id.BatchId.ToString(),
-					StepId = id.BatchId.ToString(),
+					StepId = id.StepId.ToString(),
 					AgentId = agentId,
 					BatchInitTime = initTime,
 					BatchWaitTime = waitTime,
