UBT: Fix seconds vs milliseconds parsing of XGE metadata

diff --git a/Engine/Source/Programs/Horde/Horde.Agent.Tests/LogParserTests.cs b/Engine/Source/Programs/Horde/Horde.Agent.Tests/LogParserTests.cs
--- a/Engine/Source/Programs/Horde/Horde.Agent.Tests/LogParserTests.cs
+++ b/Engine/Source/Programs/Horde/Horde.Agent.Tests/LogParserTests.cs
@@ -616,21 +616,21 @@ namespace Horde.Agent.Tests
 			CheckEventGroup(logEvents.Slice(3, 1), 3, 1, LogLevel.Information, KnownLogEvents.Systemic_Xge_TaskMetadata);
 
 			Assert.AreEqual("MiMalloc.c", logEvents[0].GetProperty<string>("name"));
-			Assert.AreEqual(77, logEvents[0].GetProperty<int>("duration"));
-			Assert.AreEqual(11, logEvents[0].GetProperty<int>("startTime"));
+			Assert.AreEqual(1170, logEvents[0].GetProperty<int>("duration"));
+			Assert.AreEqual(11000, logEvents[0].GetProperty<int>("startTime"));
 			
 			Assert.AreEqual("SomeFile.cpp", logEvents[1].GetProperty<string>("name"));
-			Assert.AreEqual(16272, logEvents[1].GetProperty<int>("duration"));
-			Assert.AreEqual(16, logEvents[1].GetProperty<int>("startTime"));
+			Assert.AreEqual(271120, logEvents[1].GetProperty<int>("duration"));
+			Assert.AreEqual(16000, logEvents[1].GetProperty<int>("startTime"));
 			
 			Assert.AreEqual("Default.rc2", logEvents[2].GetProperty<string>("name"));
 			Assert.AreEqual("SomeAgentName (Core #7)", logEvents[2].GetProperty<string>("agent"));
-			Assert.AreEqual(9296, logEvents[2].GetProperty<int>("duration"));
-			Assert.AreEqual(3318, logEvents[2].GetProperty<int>("startTime"));
+			Assert.AreEqual(154560, logEvents[2].GetProperty<int>("duration"));
+			Assert.AreEqual(3318000, logEvents[2].GetProperty<int>("startTime"));
 			
 			Assert.AreEqual("", logEvents[3].GetProperty<string>("name"));
-			Assert.AreEqual(3672, logEvents[3].GetProperty<int>("duration"));
-			Assert.AreEqual(3483, logEvents[3].GetProperty<int>("startTime"));
+			Assert.AreEqual(61120, logEvents[3].GetProperty<int>("duration"));
+			Assert.AreEqual(3483000, logEvents[3].GetProperty<int>("startTime"));
 		}
 
 		private class FakeLogSink : ILogEventSink
diff --git a/Engine/Source/Programs/UnrealBuildTool/Matchers/XoreaxEventMatcher.cs b/Engine/Source/Programs/UnrealBuildTool/Matchers/XoreaxEventMatcher.cs
--- a/Engine/Source/Programs/UnrealBuildTool/Matchers/XoreaxEventMatcher.cs
+++ b/Engine/Source/Programs/UnrealBuildTool/Matchers/XoreaxEventMatcher.cs
@@ -123,11 +123,11 @@ namespace UnrealBuildTool.Matchers
 			Match m = s_duration.Match(durationText);
 			if (m.Success)
 			{
-				int hours = Convert.ToInt32(m.Groups[1].Value);
-				int minutes = Convert.ToInt32(m.Groups[2].Value);
-				int seconds = Convert.ToInt32(m.Groups[3].Value);
-				TimeSpan duration = new TimeSpan(hours, minutes, seconds);
-				builder.AddProperty(PropertyDuration, (int)duration.TotalSeconds);
+				int minutes = Convert.ToInt32(m.Groups[1].Value);
+				int seconds = Convert.ToInt32(m.Groups[2].Value);
+				int milliseconds = Convert.ToInt32(m.Groups[3].Value) * 10;
+				TimeSpan duration = new TimeSpan(0, 0, minutes, seconds, milliseconds);
+				builder.AddProperty(PropertyDuration, (int)duration.TotalMilliseconds);
 			}
 		}
 		
@@ -139,7 +139,7 @@ namespace UnrealBuildTool.Matchers
 				int minutes = Convert.ToInt32(m.Groups[1].Value);
 				int seconds = Convert.ToInt32(m.Groups[2].Value);
 				TimeSpan startTime = new TimeSpan(0, minutes, seconds);
-				builder.AddProperty(PropertyStartTime, (int)startTime.TotalSeconds);
+				builder.AddProperty(PropertyStartTime, (int)startTime.TotalMilliseconds);
 			}
 		}
 	}
