Horde: Change XGE metric forwarding to use already parsed JSON log events from UBT

Previous approach never triggered as UBT with structured logging turned on already had parsed and matched the raw stdout log from xgConsole.

diff --git a/Engine/Source/Programs/Horde/Horde.Agent/Execution/JobExecutor.cs b/Engine/Source/Programs/Horde/Horde.Agent/Execution/JobExecutor.cs
--- a/Engine/Source/Programs/Horde/Horde.Agent/Execution/JobExecutor.cs
+++ b/Engine/Source/Programs/Horde/Horde.Agent/Execution/JobExecutor.cs
@@ -1087,9 +1087,43 @@ namespace Horde.Agent.Execution
 			}
 		}
 
+		private class InterceptingLogger : ILogger
+		{
+			private readonly ILogger _logger;
+			private readonly int[] _includeEventIds;
+			private readonly Action<LogLevel, EventId, object?, Exception?> _callback = null!;
+
+			public InterceptingLogger(ILogger logger, int[] includeEventIds, Action<LogLevel, EventId, object?, Exception?> callback)
+			{
+				_logger = logger;
+				_includeEventIds = includeEventIds;
+				_callback = callback;
+			}
+
+			public void Log<TState>(LogLevel logLevel, EventId eventId, TState state, Exception? exception, Func<TState, Exception?, string> formatter)
+			{
+				if (_includeEventIds.Any(id => id == eventId.Id))
+				{
+					_callback(logLevel, eventId, state, exception);	
+				}
+
+				_logger.Log(logLevel, eventId, state, exception, formatter);
+			}
+
+			public bool IsEnabled(LogLevel logLevel)
+			{
+				return _logger.IsEnabled(logLevel);
+			}
+
+			public IDisposable BeginScope<TState>(TState state)
+			{
+				return _logger.BeginScope(state);
+			}
+		}
+		
 		string RefPrefix => $"{_job.StreamId}/{_job.Change}-{_jobId}";
 
-		async Task<int> ExecuteCommandAsync(BeginStepResponse step, DirectoryReference workspaceDir, string fileName, string arguments, IStorageClient? store, ILogger logger, CancellationToken cancellationToken)
+		async Task<int> ExecuteCommandAsync(BeginStepResponse step, DirectoryReference workspaceDir, string fileName, string arguments, IStorageClient? store, ILogger jobLogger, CancellationToken cancellationToken)
 		{
 			// Combine all the supplied environment variables together
 			Dictionary<string, string> newEnvVars = new Dictionary<string, string>(_envVars, StringComparer.Ordinal);
@@ -1144,7 +1178,7 @@ namespace Horde.Agent.Execution
 				{
 					value = envVar.Value;
 				}
-				logger.LogInformation("Setting env var: {Key}={Value}", envVar.Key, value);
+				jobLogger.LogInformation("Setting env var: {Key}={Value}", envVar.Key, value);
 			}
 
 			// Add all the old environment variables into the list
@@ -1164,36 +1198,40 @@ namespace Horde.Agent.Execution
 				FileUtils.ForceDeleteDirectoryContents(telemetryDir);
 			}
 
-			List<string> ignorePatterns = await ReadIgnorePatternsAsync(workspaceDir, logger);
+			List<string> ignorePatterns = await ReadIgnorePatternsAsync(workspaceDir, jobLogger);
 
-			FilteringEventSink xgeEventSink = new(new [] { KnownLogEvents.Systemic_Xge_TaskMetadata.Id }, (logEvent) =>
+			InterceptingLogger interceptedJobLogger = new (jobLogger, new [] { KnownLogEvents.Systemic_Xge_TaskMetadata.Id }, (level, id, state, exception) =>
 			{
-				try
+				if (state is JsonLogEvent jsonLogEvent)
 				{
-					if (logEvent.TryGetProperty("agent", out string? agentName))
+					try
 					{
-						string taskName = logEvent.GetProperty<string>("name");
-						int duration = logEvent.GetProperty<int>("duration");
-						_logger.LogInformation("Executed XGE task {XgeTaskName} on agent {XgeAgent} for {Duration} secs", taskName, agentName, duration);	
+						LogEvent logEvent = LogEvent.Read(jsonLogEvent.Data.Span);
+						if (logEvent.TryGetProperty("agent", out string? agentName))
+						{
+							string taskName = logEvent.GetProperty<string>("name");
+							int duration = logEvent.GetProperty<int>("duration");
+							_logger.LogInformation("Executed XGE task '{XgeTaskName}' on agent '{XgeAgent}' for {Duration} secs", taskName, agentName, duration);	
+						}
+					}
+					catch (Exception e)
+					{
+						_logger.LogWarning(e, "Failed to log XGE task execution");
 					}
-				}
-				catch (Exception e)
-				{
-					_logger.LogWarning(e, "Failed to log XGE task execution");
 				}
 			});
-			
+
 			int exitCode;
-			using (LogParser filter = new LogParser(logger, ignorePatterns, new List<ILogEventSink>() { xgeEventSink }))
+			using (LogParser filter = new LogParser(interceptedJobLogger, ignorePatterns))
 			{
-				await ExecuteCleanupScriptAsync(cleanupScript, filter, logger);
+				await ExecuteCleanupScriptAsync(cleanupScript, filter, jobLogger);
 				try
 				{
-					exitCode = await ExecuteProcessAsync(fileName, arguments, newEnvVars, filter, logger, cancellationToken);
+					exitCode = await ExecuteProcessAsync(fileName, arguments, newEnvVars, filter, jobLogger, cancellationToken);
 				}
 				finally
 				{
-					await ExecuteCleanupScriptAsync(cleanupScript, filter, logger);
+					await ExecuteCleanupScriptAsync(cleanupScript, filter, jobLogger);
 				}
 				filter.Flush();
 			}
@@ -1203,7 +1241,7 @@ namespace Horde.Agent.Execution
 				List<TraceEventList> telemetryList = new List<TraceEventList>();
 				foreach (FileReference telemetryFile in DirectoryReference.EnumerateFiles(telemetryDir, "*.json"))
 				{
-					logger.LogInformation("Reading telemetry from {File}", telemetryFile);
+					jobLogger.LogInformation("Reading telemetry from {File}", telemetryFile);
 					byte[] data = await FileReference.ReadAllBytesAsync(telemetryFile, cancellationToken);
 
 					TraceEventList telemetry = JsonSerializer.Deserialize<TraceEventList>(data.AsSpan())!;
@@ -1217,7 +1255,7 @@ namespace Horde.Agent.Execution
 						telemetryList.Add(telemetry);
 					}
 
-					await ArtifactUploader.UploadAsync(RpcConnection, _jobId, _batchId, step.StepId, $"Telemetry/{telemetryFile.GetFileName()}", telemetryFile, logger, CancellationToken.None);
+					await ArtifactUploader.UploadAsync(RpcConnection, _jobId, _batchId, step.StepId, $"Telemetry/{telemetryFile.GetFileName()}", telemetryFile, jobLogger, CancellationToken.None);
 					FileUtils.ForceDeleteFile(telemetryFile);
 				}
 
@@ -1264,7 +1302,7 @@ namespace Horde.Agent.Execution
 
 						if (stack.Count > 1 && newSpan.Finish > stackTop.Finish)
 						{
-							logger.LogInformation("Trace event name='{Name}', service'{Service}', resource='{Resource}' has invalid finish time ({SpanFinish} < {StackFinish})", newSpan.Name, newSpan.Service, newSpan.Resource, newSpan.Finish, stackTop.Finish);
+							jobLogger.LogInformation("Trace event name='{Name}', service'{Service}', resource='{Resource}' has invalid finish time ({SpanFinish} < {StackFinish})", newSpan.Name, newSpan.Service, newSpan.Resource, newSpan.Finish, stackTop.Finish);
 							newSpan.Finish = stackTop.Finish;
 						}
 
@@ -1286,7 +1324,7 @@ namespace Horde.Agent.Execution
 						JsonSerializerOptions options = new JsonSerializerOptions { DefaultIgnoreCondition = JsonIgnoreCondition.WhenWritingNull };
 						await JsonSerializer.SerializeAsync(stream, rootSpan, options, cancellationToken);
 					}
-					await ArtifactUploader.UploadAsync(RpcConnection, _jobId, _batchId, step.StepId, "Trace.json", traceFile, logger, CancellationToken.None);
+					await ArtifactUploader.UploadAsync(RpcConnection, _jobId, _batchId, step.StepId, "Trace.json", traceFile, jobLogger, CancellationToken.None);
 
 					CreateTracingData(GlobalTracer.Instance.ActiveSpan, rootSpan);
 				}
@@ -1297,8 +1335,8 @@ namespace Horde.Agent.Execution
 				Dictionary<string, object> combinedTestData = new Dictionary<string, object>();
 				foreach (FileReference testDataFile in DirectoryReference.EnumerateFiles(testDataDir, "*.json", SearchOption.AllDirectories))
 				{
-					logger.LogInformation("Reading test data {TestDataFile}", testDataFile);
-					await ArtifactUploader.UploadAsync(RpcConnection, _jobId, _batchId, step.StepId, $"TestData/{testDataFile.MakeRelativeTo(testDataDir)}", testDataFile, logger, CancellationToken.None);
+					jobLogger.LogInformation("Reading test data {TestDataFile}", testDataFile);
+					await ArtifactUploader.UploadAsync(RpcConnection, _jobId, _batchId, step.StepId, $"TestData/{testDataFile.MakeRelativeTo(testDataDir)}", testDataFile, jobLogger, CancellationToken.None);
 
 					TestData testData;
 					using (FileStream stream = FileReference.Open(testDataFile, FileMode.Open))
@@ -1311,17 +1349,17 @@ namespace Horde.Agent.Execution
 					{
 						if (combinedTestData.ContainsKey(item.Key))
 						{
-							logger.LogWarning("Key '{Key}' already exists - ignoring", item.Key);
+							jobLogger.LogWarning("Key '{Key}' already exists - ignoring", item.Key);
 						}
 						else
 						{
-							logger.LogDebug("Adding data with key '{Key}'", item.Key);
+							jobLogger.LogDebug("Adding data with key '{Key}'", item.Key);
 							combinedTestData.Add(item.Key, item.Data);
 						}
 					}
 				}
 
-				logger.LogInformation("Found {NumResults} test results", combinedTestData.Count);
+				jobLogger.LogInformation("Found {NumResults} test results", combinedTestData.Count);
 				await UploadTestDataAsync(step.StepId, combinedTestData);
 			}
 					
@@ -1348,7 +1386,7 @@ namespace Horde.Agent.Execution
 					{
 						string artifactName = artifactFile.MakeRelativeTo(logDir);
 
-						string? artifactId = await ArtifactUploader.UploadAsync(RpcConnection, _jobId, _batchId, step.StepId, artifactName, artifactFile, logger, cancellationToken);
+						string? artifactId = await ArtifactUploader.UploadAsync(RpcConnection, _jobId, _batchId, step.StepId, artifactName, artifactFile, jobLogger, cancellationToken);
 						if (artifactId != null)
 						{
 							artifactFileToId[artifactFile] = artifactId;
@@ -1359,11 +1397,11 @@ namespace Horde.Agent.Execution
 					{
 						try
 						{
-							await CreateReportAsync(step.StepId, reportFile, artifactFileToId, logger);
+							await CreateReportAsync(step.StepId, reportFile, artifactFileToId, jobLogger);
 						}
 						catch (Exception ex)
 						{
-							logger.LogWarning("Unable to upload report: {Message}", ex.Message);
+							jobLogger.LogWarning("Unable to upload report: {Message}", ex.Message);
 						}
 					}
 				}
