Horde: Fix typo in log message for XGE metrics

diff --git a/Engine/Source/Programs/Horde/Horde.Agent/Execution/JobExecutor.cs b/Engine/Source/Programs/Horde/Horde.Agent/Execution/JobExecutor.cs
--- a/Engine/Source/Programs/Horde/Horde.Agent/Execution/JobExecutor.cs
+++ b/Engine/Source/Programs/Horde/Horde.Agent/Execution/JobExecutor.cs
@@ -1211,7 +1211,7 @@ namespace Horde.Agent.Execution
 						{
 							string taskName = logEvent.GetProperty<string>("name");
 							int duration = logEvent.GetProperty<int>("duration");
-							_logger.LogInformation("Executed XGE task '{XgeTaskName}' on agent '{XgeAgent}' for {Duration} secs", taskName, agentName, duration);	
+							_logger.LogInformation("Executed XGE task {XgeTaskName} on agent {XgeAgent} for {Duration} ms", taskName, agentName, duration);	
 						}
 					}
 					catch (Exception e)
