Horde: Fix formatting in config update failure notifications.

diff --git a/Engine/Source/Programs/Horde/Horde.Build/Notifications/Sinks/SlackNotificationSink.cs b/Engine/Source/Programs/Horde/Horde.Build/Notifications/Sinks/SlackNotificationSink.cs
--- a/Engine/Source/Programs/Horde/Horde.Build/Notifications/Sinks/SlackNotificationSink.cs
+++ b/Engine/Source/Programs/Horde/Horde.Build/Notifications/Sinks/SlackNotificationSink.cs
@@ -1956,7 +1956,7 @@ namespace Horde.Build.Notifications.Sinks
 							if (blame.Author != null)
 							{
 								string userId = await FormatMentionAsync(blame.Author.Id, true);
-								blameMessage += $" ({FormatUserOrGroupMention(userId)})";
+								blameMessage += $" ({userId})";
 							}
 							details.Add(blameMessage.ToString());
 						}
@@ -2002,7 +2002,7 @@ namespace Horde.Build.Notifications.Sinks
 			}
 			else if (uri.Scheme == PerforceConfigSource.Scheme)
 			{
-				return $"/{uri.AbsolutePath}";
+				return uri.AbsolutePath;
 			}
 			else
 			{
