Horde: Add config path to stream response if it exists
#rnx

diff --git a/Engine/Source/Programs/Horde/Horde.Build/Streams/StreamMessages.cs b/Engine/Source/Programs/Horde/Horde.Build/Streams/StreamMessages.cs
--- a/Engine/Source/Programs/Horde/Horde.Build/Streams/StreamMessages.cs
+++ b/Engine/Source/Programs/Horde/Horde.Build/Streams/StreamMessages.cs
@@ -395,6 +395,7 @@ namespace Horde.Build.Streams
 			Id = stream.Id.ToString();
 			ProjectId = stream.Config.ProjectConfig.Id.ToString();
 			Name = stream.Config.Name;
+			ConfigPath = stream.Config.Path ?? String.Empty;
 			ConfigRevision = stream.Config.Revision;
 			Order = stream.Config.Order;
 			NotificationChannel = stream.Config.NotificationChannel;
