Horde: Use config path in new build dialog to display location in repo
#rnx
#skipci

diff --git a/Engine/Source/Programs/Horde/HordeDashboard/src/components/NewBuild.tsx b/Engine/Source/Programs/Horde/HordeDashboard/src/components/NewBuild.tsx
--- a/Engine/Source/Programs/Horde/HordeDashboard/src/components/NewBuild.tsx
+++ b/Engine/Source/Programs/Horde/HordeDashboard/src/components/NewBuild.tsx
@@ -1644,7 +1644,7 @@ export const NewBuild: React.FC<{ streamId: string; show: boolean; onClose: (new
 
                      {mode === "Advanced" && <Stack style={{ paddingBottom: 12 }} tokens={{ childrenGap: 12 }}>
                         {!!stream.configRevision && <Stack>
-                           <TextField label="Template Revision" readOnly={true} value={stream.configRevision} />
+                           <TextField label="Template Path" readOnly={true} value={stream.configPath ?? ""} />
                         </Stack>
                         }
                         <Stack>
