Horde: Add a description field to job templates.

diff --git a/Engine/Source/Programs/Horde/Horde.Build/Jobs/Templates/ITemplate.cs b/Engine/Source/Programs/Horde/Horde.Build/Jobs/Templates/ITemplate.cs
--- a/Engine/Source/Programs/Horde/Horde.Build/Jobs/Templates/ITemplate.cs
+++ b/Engine/Source/Programs/Horde/Horde.Build/Jobs/Templates/ITemplate.cs
@@ -432,6 +432,11 @@ namespace Horde.Build.Jobs.Templates
 		/// </summary>
 		public string Name { get; }
 
+		/// <summary>
+		/// Description for the template
+		/// </summary>
+		public string? Description { get; }
+
 		/// <summary>
 		/// Priority of this job
 		/// </summary>
diff --git a/Engine/Source/Programs/Horde/Horde.Build/Jobs/Templates/TemplateCollection.cs b/Engine/Source/Programs/Horde/Horde.Build/Jobs/Templates/TemplateCollection.cs
--- a/Engine/Source/Programs/Horde/Horde.Build/Jobs/Templates/TemplateCollection.cs
+++ b/Engine/Source/Programs/Horde/Horde.Build/Jobs/Templates/TemplateCollection.cs
@@ -33,6 +33,9 @@ namespace Horde.Build.Jobs.Templates
 			[BsonRequired]
 			public string Name { get; private set; }
 
+			[BsonIgnoreIfNull]
+			public string? Description { get; set; }
+
 			public Priority? Priority { get; private set; }
 			public bool AllowPreflights { get; set; } = true;
 			public bool UpdateIssues { get; set; } = false;
@@ -61,6 +64,7 @@ namespace Horde.Build.Jobs.Templates
 			public TemplateDocument(TemplateConfig config)
 			{
 				Name = config.Name;
+				Description = config.Description;
 				Priority = config.Priority;
 				AllowPreflights = config.AllowPreflights;
 				UpdateIssues = config.UpdateIssues;
diff --git a/Engine/Source/Programs/Horde/Horde.Build/Jobs/Templates/TemplateMessages.cs b/Engine/Source/Programs/Horde/Horde.Build/Jobs/Templates/TemplateMessages.cs
--- a/Engine/Source/Programs/Horde/Horde.Build/Jobs/Templates/TemplateMessages.cs
+++ b/Engine/Source/Programs/Horde/Horde.Build/Jobs/Templates/TemplateMessages.cs
@@ -397,6 +397,11 @@ namespace Horde.Build.Jobs.Templates
 		[Required]
 		public string Name { get; set; } = null!;
 
+		/// <summary>
+		/// Description for the template
+		/// </summary>
+		public string? Description { get; set; }
+
 		/// <summary>
 		/// Default priority for this job
 		/// </summary>
@@ -469,6 +474,11 @@ namespace Horde.Build.Jobs.Templates
 		/// </summary>
 		public string Name { get; set; }
 
+		/// <summary>
+		/// Description for the template
+		/// </summary>
+		public string? Description { get; set; }
+
 		/// <summary>
 		/// Default priority for this job
 		/// </summary>
@@ -522,6 +532,7 @@ namespace Horde.Build.Jobs.Templates
 		public GetTemplateResponseBase(ITemplate template)
 		{
 			Name = template.Name;
+			Description = template.Description;
 			Priority = template.Priority;
 			AllowPreflights = template.AllowPreflights;
 			UpdateIssues = template.UpdateIssues;
