Horde: Optimize for the case where a very large changelist (eg. a merge commit) matches all predefined commit tags, by querying for a list of files with an increasing maximum length.

diff --git a/Engine/Source/Programs/Horde/Horde.Build/Issues/IssueService.cs b/Engine/Source/Programs/Horde/Horde.Build/Issues/IssueService.cs
--- a/Engine/Source/Programs/Horde/Horde.Build/Issues/IssueService.cs
+++ b/Engine/Source/Programs/Horde/Horde.Build/Issues/IssueService.cs
@@ -1125,6 +1125,11 @@ namespace Horde.Build.Issues
 			return containsFixChange;
 		}
 
+		/// <summary>
+		/// Maximum number of files to look at when looking for suspects
+		/// </summary>
+		const int MaxFilesForSuspects = 1000;
+
 		/// <summary>
 		/// Find suspects that can match a given span
 		/// </summary>
@@ -1148,7 +1153,7 @@ namespace Horde.Build.Issues
 				List<SuspectChange> suspectChanges = new List<SuspectChange>(changes.Count);
 				foreach (ICommit commit in changes)
 				{
-					IReadOnlyList<string> files = await commit.GetFilesAsync(CancellationToken.None);
+					IReadOnlyList<string> files = await commit.GetFilesAsync(MaxFilesForSuspects, CancellationToken.None);
 					suspectChanges.Add(new SuspectChange(commit, files));
 				}
 				handler.RankSuspects(fingerprint, suspectChanges);
diff --git a/Engine/Source/Programs/Horde/Horde.Build/Jobs/JobsController.cs b/Engine/Source/Programs/Horde/Horde.Build/Jobs/JobsController.cs
--- a/Engine/Source/Programs/Horde/Horde.Build/Jobs/JobsController.cs
+++ b/Engine/Source/Programs/Horde/Horde.Build/Jobs/JobsController.cs
@@ -177,7 +177,7 @@ namespace Horde.Build.Jobs
 			int change = await GetChangeToBuildAsync(create, streamConfig.Id, template, shelfInfo, commits, HttpContext.RequestAborted);
 
 			// And get the matching code changelist
-			ICommit? lastCodeCommit = await commits.GetLastCodeChange(change, HttpContext.RequestAborted);
+			ICommit? lastCodeCommit = await commits.GetLastCodeChangeAsync(change, HttpContext.RequestAborted);
 			int codeChange = lastCodeCommit?.Number ?? change;
 
 			// New properties for the job
diff --git a/Engine/Source/Programs/Horde/Horde.Build/Jobs/Schedules/ScheduleService.cs b/Engine/Source/Programs/Horde/Horde.Build/Jobs/Schedules/ScheduleService.cs
--- a/Engine/Source/Programs/Horde/Horde.Build/Jobs/Schedules/ScheduleService.cs
+++ b/Engine/Source/Programs/Horde/Horde.Build/Jobs/Schedules/ScheduleService.cs
@@ -399,11 +399,11 @@ namespace Horde.Build.Jobs.Schedules
 				}
 
 				// Adjust the changelist for the desired filter
-				if (commit == null || await ShouldBuildChangeAsync(commit, schedule.Config.Commits, fileFilter))
+				if (commit == null || await ShouldBuildChangeAsync(commit, schedule.Config.Commits, fileFilter, cancellationToken))
 				{
 					int codeChange = change;
 
-					ICommit? lastCodeCommit = await commits.GetLastCodeChange(change, cancellationToken);
+					ICommit? lastCodeCommit = await commits.GetLastCodeChangeAsync(change, cancellationToken);
 					if (lastCodeCommit != null)
 					{
 						codeChange = lastCodeCommit.Number;
@@ -448,7 +448,7 @@ namespace Horde.Build.Jobs.Schedules
 			if (template.SubmitNewChange != null)
 			{
 				int newChange = await commits.CreateNewAsync(template, cancellationToken);
-				ICommit? newCodeChange = await commits.GetLastCodeChange(newChange, cancellationToken);
+				ICommit? newCodeChange = await commits.GetLastCodeChangeAsync(newChange, cancellationToken);
 				triggerChanges = new List<(int, int)> { (newChange, newCodeChange?.Number ?? newChange) };
 			}
 
@@ -472,31 +472,31 @@ namespace Horde.Build.Jobs.Schedules
 		/// <summary>
 		/// Tests whether a schedule should build a particular change, based on its requested change filters
 		/// </summary>
-		/// <param name="details">The change details</param>
-		/// <param name="filterTags"></param>
+		/// <param name="commit">The commit details</param>
+		/// <param name="filterTags">Filter for the tags to trigger a build</param>
 		/// <param name="fileFilter">Filter for the files to trigger a build</param>
+		/// <param name="cancellationToken">Cancellation token for the operation</param>
 		/// <returns></returns>
-		private async ValueTask<bool> ShouldBuildChangeAsync(ICommit details, List<CommitTag>? filterTags, FileFilter? fileFilter)
+		private async ValueTask<bool> ShouldBuildChangeAsync(ICommit commit, List<CommitTag>? filterTags, FileFilter? fileFilter, CancellationToken cancellationToken)
 		{
-			if (Regex.IsMatch(details.Description, @"^\s*#\s*skipci", RegexOptions.Multiline))
+			if (Regex.IsMatch(commit.Description, @"^\s*#\s*skipci", RegexOptions.Multiline))
 			{
 				return false;
 			}
 			if (filterTags != null && filterTags.Count > 0)
 			{
-				IReadOnlyList<CommitTag> commitTags = await details.GetTagsAsync(CancellationToken.None);
+				IReadOnlyList<CommitTag> commitTags = await commit.GetTagsAsync(cancellationToken);
 				if (!commitTags.Any(x => filterTags.Contains(x)))
 				{
-					_logger.LogDebug("Not building change {Change} ({ChangeTags}) due to filter tags ({FilterTags})", details.Number, String.Join(", ", commitTags.Select(x => x.ToString())), String.Join(", ", filterTags.Select(x => x.ToString())));
+					_logger.LogDebug("Not building change {Change} ({ChangeTags}) due to filter tags ({FilterTags})", commit.Number, String.Join(", ", commitTags.Select(x => x.ToString())), String.Join(", ", filterTags.Select(x => x.ToString())));
 					return false;
 				}
 			}
 			if (fileFilter != null)
 			{
-				IReadOnlyList<string> files = await details.GetFilesAsync(CancellationToken.None);
-				if (!files.Any(x => fileFilter.Matches(x)))
+				if (await commit.MatchesFilterAsync(fileFilter, cancellationToken))
 				{
-					_logger.LogDebug("Not building change {Change} due to file filter", details.Number);
+					_logger.LogDebug("Not building change {Change} due to file filter", commit.Number);
 					return false;
 				}
 			}
diff --git a/Engine/Source/Programs/Horde/Horde.Build/Perforce/ICommit.cs b/Engine/Source/Programs/Horde/Horde.Build/Perforce/ICommit.cs
--- a/Engine/Source/Programs/Horde/Horde.Build/Perforce/ICommit.cs
+++ b/Engine/Source/Programs/Horde/Horde.Build/Perforce/ICommit.cs
@@ -5,6 +5,7 @@ using System.Collections.Generic;
 using System.Linq;
 using System.Threading;
 using System.Threading.Tasks;
+using EpicGames.Core;
 using EpicGames.Perforce;
 using Horde.Build.Streams;
 using Horde.Build.Users;
@@ -67,11 +68,20 @@ namespace Horde.Build.Perforce
 		/// <returns>True if the commit has the given tag</returns>
 		ValueTask<IReadOnlyList<CommitTag>> GetTagsAsync(CancellationToken cancellationToken);
 
+		/// <summary>
+		/// Determine if this commit matches the given filter. Prefer using commit tags rather than this method; the results can be cached.
+		/// </summary>
+		/// <param name="filter">Filter to test</param>
+		/// <param name="cancellationToken">Cancellation token for the operation</param>
+		/// <returns></returns>
+		ValueTask<bool> MatchesFilterAsync(FileFilter filter, CancellationToken cancellationToken);
+
 		/// <summary>
 		/// Gets the files for this change, relative to the root of the stream
 		/// </summary>
+		/// <param name="maxFiles">Maximum number of files to return. Querying large number of files may cause performance issues with merge commits.</param>
 		/// <param name="cancellationToken">Cancellation token for the operation</param>
 		/// <returns>List of files modified by this commit</returns>
-		ValueTask<IReadOnlyList<string>> GetFilesAsync(CancellationToken cancellationToken);
+		ValueTask<IReadOnlyList<string>> GetFilesAsync(int maxFiles, CancellationToken cancellationToken);
 	}
 }
diff --git a/Engine/Source/Programs/Horde/Horde.Build/Perforce/ICommitCollection.cs b/Engine/Source/Programs/Horde/Horde.Build/Perforce/ICommitCollection.cs
--- a/Engine/Source/Programs/Horde/Horde.Build/Perforce/ICommitCollection.cs
+++ b/Engine/Source/Programs/Horde/Horde.Build/Perforce/ICommitCollection.cs
@@ -77,7 +77,7 @@ namespace Horde.Build.Perforce
 		/// <param name="maxChange">Maximum code change to query</param>
 		/// <param name="cancellationToken">Cancellation token for the operation</param>
 		/// <returns>The last code change</returns>
-		public static async ValueTask<ICommit?> GetLastCodeChange(this ICommitCollection source, int? maxChange, CancellationToken cancellationToken = default)
+		public static async ValueTask<ICommit?> GetLastCodeChangeAsync(this ICommitCollection source, int? maxChange, CancellationToken cancellationToken = default)
 		{
 			return await source.FindAsync(null, maxChange, 1, new[] { CommitTag.Code }, cancellationToken).FirstOrDefaultAsync(cancellationToken);
 		}
diff --git a/Engine/Source/Programs/Horde/Horde.Build/Perforce/PerforceService.cs b/Engine/Source/Programs/Horde/Horde.Build/Perforce/PerforceService.cs
--- a/Engine/Source/Programs/Horde/Horde.Build/Perforce/PerforceService.cs
+++ b/Engine/Source/Programs/Horde/Horde.Build/Perforce/PerforceService.cs
@@ -158,7 +158,8 @@ namespace Horde.Build.Perforce
 			public string BasePath { get; }
 			public DateTime DateUtc { get; }
 
-			IReadOnlyList<string>? _files;
+			IReadOnlyList<string> _files = Array.Empty<string>();
+			bool _hasAllFiles;
 
 			public Commit(PerforceService owner, StreamConfig streamConfig, int number, int originalChange, UserId authorId, UserId ownerId, string description, string basePath, DateTime dateUtc)
 			{
@@ -174,35 +175,57 @@ namespace Horde.Build.Perforce
 				DateUtc = dateUtc;
 			}
 
-			public void SetFiles(IReadOnlyList<string> files)
+			public void SetFiles(IReadOnlyList<string> files, int? maxFiles)
 			{
 				_files = files;
+				_hasAllFiles = (maxFiles == null || files.Count < maxFiles.Value);
 			}
 
 			public async ValueTask<IReadOnlyList<CommitTag>> GetTagsAsync(CancellationToken cancellationToken)
 			{
 				List<CommitTag> commitTags = new List<CommitTag>();
-
-				IReadOnlyList<string> files = await GetFilesAsync(cancellationToken);
-				foreach (CommitTagConfig tagConfig in _streamConfig.GetAllCommitTags())
+				foreach (CommitTagConfig commitTagConfig in _streamConfig.GetAllCommitTags())
 				{
-					if (_streamConfig.TryGetCommitTagFilter(tagConfig.Name, out FileFilter? filter) && filter.ApplyTo(files).Any())
+					if (_streamConfig.TryGetCommitTagFilter(commitTagConfig.Name, out FileFilter? filter) && await MatchesFilterAsync(filter, cancellationToken))
 					{
-						commitTags.Add(tagConfig.Name);
+						commitTags.Add(commitTagConfig.Name);
 					}
 				}
-
 				return commitTags;
 			}
 
-			public async ValueTask<IReadOnlyList<string>> GetFilesAsync(CancellationToken cancellationToken)
+			public async ValueTask<bool> MatchesFilterAsync(FileFilter filter, CancellationToken cancellationToken)
+			{
+				int maxFiles = 1000;
+				for(; ;)
+				{
+					// Query the files up to the current maximum
+					IReadOnlyList<string> files = await GetFilesAsync(maxFiles, cancellationToken);
+					if (filter.ApplyTo(files).Any())
+					{
+						return true;
+					}
+
+					// If we've already got the full file list for this change, also don't need to query any more
+					if (_hasAllFiles)
+					{
+						return false;
+					}
+
+					// Query more files next time
+					maxFiles = Math.Max(files.Count, maxFiles) * 10;
+				}
+			}
+
+			public async ValueTask<IReadOnlyList<string>> GetFilesAsync(int maxFiles, CancellationToken cancellationToken)
 			{
-				if (_files == null)
+				if (maxFiles > _files.Count && !_hasAllFiles)
 				{
 					using (IPooledPerforceConnection perforce = await _owner.ConnectAsync(_streamConfig.ClusterName, null, cancellationToken))
 					{
-						DescribeRecord describeRecord = await perforce.DescribeAsync(Number, cancellationToken);
+						DescribeRecord describeRecord = await perforce.DescribeAsync(DescribeOptions.None, maxFiles, Number, cancellationToken);
 						_files = await perforce.GetStreamFilesAsync(_streamConfig, describeRecord, cancellationToken);
+						_hasAllFiles = _files.Count < maxFiles;
 					}
 				}
 				return _files;
@@ -595,16 +618,6 @@ namespace Horde.Build.Perforce
 			return user!;
 		}
 
-		PerforceCluster GetCluster(string? clusterName, string? serverAndPort = null)
-		{
-			PerforceCluster? cluster = _globalConfig.CurrentValue.FindPerforceCluster(clusterName, serverAndPort);
-			if (cluster == null)
-			{
-				throw new Exception($"Unknown Perforce cluster '{clusterName}'");
-			}
-			return cluster;
-		}
-
 		async Task<IPerforceServer> GetServerAsync(PerforceCluster cluster)
 		{
 			using IScope scope = GlobalTracer.Instance.BuildSpan("PerforceService.SelectServer").StartActive();
@@ -667,7 +680,7 @@ namespace Horde.Build.Perforce
 			return new Commit(this, streamConfig, number, originalChange, authorUser.Id, ownerUser.Id, description, basePath ?? String.Empty, dateUtc);
 		}
 
-		protected async ValueTask<ICommit> CreateCommitAsync(IPooledPerforceConnection perforce, StreamConfig streamConfig, DescribeRecord describeRecord, InfoRecord serverInfo, CancellationToken cancellationToken)
+		protected async ValueTask<ICommit> CreateCommitAsync(IPooledPerforceConnection perforce, StreamConfig streamConfig, DescribeRecord describeRecord, int maxFiles, InfoRecord serverInfo, CancellationToken cancellationToken)
 		{
 			DateTime timeUtc = new DateTime(describeRecord.Time.Ticks - serverInfo.TimeZoneOffsetSecs * TimeSpan.TicksPerSecond, DateTimeKind.Utc);
 
@@ -678,7 +691,7 @@ namespace Horde.Build.Perforce
 			{
 				throw new PerforceException($"Changelist {commit.Number} does not contain any files in {streamConfig.Id}");
 			}
-			commit.SetFiles(files);
+			commit.SetFiles(files, maxFiles);
 
 			return commit;
 		}
@@ -742,13 +755,14 @@ namespace Horde.Build.Perforce
 
 			using (PooledConnectionHandle perforce = await ConnectAsync(cluster, null, cancellationToken))
 			{
-				DescribeRecord record = await perforce.DescribeAsync(changeNumber, cancellationToken);
+				const int MaxFiles = 1000;
+				DescribeRecord describeRecord = await perforce.DescribeAsync(DescribeOptions.None, MaxFiles, changeNumber, cancellationToken);
 
 				InfoRecord info = await perforce.GetInfoAsync(cancellationToken);
 
-				DateTime timeUtc = new DateTime(record.Time.Ticks - info.TimeZoneOffsetSecs * TimeSpan.TicksPerSecond, DateTimeKind.Utc);
+				DateTime timeUtc = new DateTime(describeRecord.Time.Ticks - info.TimeZoneOffsetSecs * TimeSpan.TicksPerSecond, DateTimeKind.Utc);
 
-				ICommit commit = await CreateCommitAsync(perforce, streamConfig, record, info, cancellationToken);
+				ICommit commit = await CreateCommitAsync(perforce, streamConfig, describeRecord, MaxFiles, info, cancellationToken);
 				return commit;
 			}
 		}
@@ -820,29 +834,6 @@ namespace Horde.Build.Perforce
 			}
 		}
 
-		/// <inheritdoc/>
-		public async Task<List<ICommit>> GetChangeDetailsAsync(StreamConfig streamConfig, IReadOnlyList<int> changeNumbers, CancellationToken cancellationToken)
-		{
-			using IScope scope = GlobalTracer.Instance.BuildSpan("PerforceService.GetChangeDetailsAsync").StartActive();
-			scope.Span.SetTag("ClusterName", streamConfig.ClusterName);
-			scope.Span.SetTag("StreamName", streamConfig.Name);
-			scope.Span.SetTag("ChangeNumbers.Count", changeNumbers.Count);
-
-			List<ICommit> results = new List<ICommit>();
-			using (PooledConnectionHandle perforce = await ConnectWithStreamClientAsync(streamConfig, null, cancellationToken))
-			{
-				InfoRecord info = await perforce.GetInfoAsync(cancellationToken);
-
-				List<DescribeRecord> records = await perforce.DescribeAsync(DescribeOptions.Shelved, -1, changeNumbers.ToArray(), cancellationToken);
-				foreach (DescribeRecord record in records)
-				{
-					ICommit commit = await CreateCommitAsync(perforce, streamConfig, record, info, cancellationToken);
-					results.Add(commit);
-				}
-			}
-			return results;
-		}
-
 		/// <inheritdoc/>
 		public Task<int> DuplicateShelvedChangeAsync(string clusterName, int shelvedChange, CancellationToken cancellationToken)
 		{
@@ -1158,11 +1149,13 @@ namespace Horde.Build.Perforce
 		{
 			protected readonly PerforceService _perforceService;
 			protected readonly StreamConfig _streamConfig;
+			protected readonly ILogger _logger;
 
-			public CommitSource(PerforceService perforceService, StreamConfig streamConfig)
+			public CommitSource(PerforceService perforceService, StreamConfig streamConfig, ILogger logger)
 			{
 				_perforceService = perforceService;
 				_streamConfig = streamConfig;
+				_logger = logger;
 			}
 
 			public async Task<int> CreateNewAsync(string path, string description, CancellationToken cancellationToken = default)
@@ -1203,43 +1196,31 @@ namespace Horde.Build.Perforce
 					}
 					else
 					{
-						const int BatchSize = 20;
-
-						// Construct filters for all the tags we're searching for
-						List<FileFilter> filters = new List<FileFilter>();
-						foreach (CommitTag tag in tags)
-						{
-							FileFilter? filter;
-							if (_streamConfig.TryGetCommitTagFilter(tag, out filter))
-							{
-								filters.Add(filter);
-							}
-						}
-
-						// Get the stream view
-						PerforceViewMap viewMap = await perforce.GetCachedStreamViewAsync(_streamConfig, cancellationToken);
-
 						// Do the query in batches
 						int numResults = maxResults ?? Int32.MaxValue;
 						while (numResults > 0)
 						{
 							string filter = GetFilter($"//{perforce.Settings.ClientName}/...", minChange, maxChange);
 
+							const int BatchSize = 20;
+
 							List<ChangesRecord> changesRecords = await perforce.GetChangesAsync(ChangesOptions.None, null, BatchSize, ChangeStatus.Submitted, null, filter, cancellationToken);
 							if (changesRecords.Count == 0)
 							{
 								break;
 							}
 
-							List<DescribeRecord> describeRecords = await perforce.DescribeAsync(DescribeOptions.None, -1, changesRecords.Select(x => x.Number).ToArray(), cancellationToken);
+							const int MaxFiles = 1000;
+
+							List<DescribeRecord> describeRecords = await perforce.DescribeAsync(DescribeOptions.None, MaxFiles, changesRecords.Select(x => x.Number).ToArray(), cancellationToken);
 							foreach (DescribeRecord describeRecord in describeRecords)
 							{
-								List<string> files = viewMap.MapFiles(describeRecord.Files.Select(x => x.DepotFile), info.PathComparison).ToList();
-								if (files.Any(x => filters.Any(y => y.Matches(x))))
+								ICommit commit = await _perforceService.CreateCommitAsync(perforce, _streamConfig, describeRecord, MaxFiles, info, cancellationToken);
+
+								IReadOnlyList<CommitTag> commitTags = await commit.GetTagsAsync(cancellationToken);
+								if (commitTags.Intersect(tags).Any())
 								{
-									ICommit commit = await _perforceService.CreateCommitAsync(perforce, _streamConfig, describeRecord, info, cancellationToken);
 									yield return commit;
-
 									if (--numResults == 0)
 									{
 										yield break;
@@ -1276,7 +1257,7 @@ namespace Horde.Build.Perforce
 		/// <inheritdoc/>
 		public virtual ICommitCollection GetCommits(StreamConfig streamConfig)
 		{
-			return new CommitSource(this, streamConfig);
+			return new CommitSource(this, streamConfig, _logger);
 		}
 	}
 }
diff --git a/Engine/Source/Programs/Horde/Horde.Build/Perforce/PerforceServiceCache.cs b/Engine/Source/Programs/Horde/Horde.Build/Perforce/PerforceServiceCache.cs
--- a/Engine/Source/Programs/Horde/Horde.Build/Perforce/PerforceServiceCache.cs
+++ b/Engine/Source/Programs/Horde/Horde.Build/Perforce/PerforceServiceCache.cs
@@ -158,14 +158,24 @@ namespace Horde.Build.Perforce
 				return new ValueTask<IReadOnlyList<CommitTag>>(CommitTags);
 			}
 
-			public async ValueTask<IReadOnlyList<string>> GetFilesAsync(CancellationToken cancellationToken)
+			public async ValueTask<bool> MatchesFilterAsync(FileFilter filter, CancellationToken cancellationToken)
+			{
+				ICommit? other = await _owner!.GetChangeDetailsAsync(_streamConfig, Number, cancellationToken);
+				if (other == null)
+				{
+					return false;
+				}
+				return await other.MatchesFilterAsync(filter, cancellationToken);
+			}
+
+			public async ValueTask<IReadOnlyList<string>> GetFilesAsync(int maxFiles, CancellationToken cancellationToken)
 			{
 				ICommit? other = await _owner!.GetChangeDetailsAsync(_streamConfig, Number, cancellationToken);
 				if (other == null)
 				{
 					return Array.Empty<string>();
 				}
-				return await other.GetFilesAsync(cancellationToken);
+				return await other.GetFilesAsync(maxFiles, cancellationToken);
 			}
 		}
 
@@ -414,8 +424,11 @@ namespace Horde.Build.Perforce
 				// Create a buffer for files
 				List<string> files = new List<string>();
 
+				// Maximum number of files to query by default
+				const int MaxFiles = 1000;
+
 				// Describe the changes and create records for them
-				List<DescribeRecord> describeRecords = await perforce.DescribeAsync(changes.Select(x => x.Number).ToArray(), cancellationToken);
+				List<DescribeRecord> describeRecords = await perforce.DescribeAsync(DescribeOptions.None, MaxFiles, changes.Select(x => x.Number).ToArray(), cancellationToken);
 				foreach (StreamInfo streamInfo in streamInfos)
 				{
 					foreach (DescribeRecord describeRecord in describeRecords)
@@ -432,7 +445,7 @@ namespace Horde.Build.Perforce
 
 						if (files.Count > 0)
 						{
-							ICommit commit = await CreateCommitAsync(perforce, streamInfo.StreamConfig, describeRecord, info, cancellationToken);
+							ICommit commit = await CreateCommitAsync(perforce, streamInfo.StreamConfig, describeRecord, MaxFiles, info, cancellationToken);
 							_logger.LogInformation("Replicating {StreamId} commit {Change}", streamInfo.StreamConfig.Id, describeRecord.Number);
 
 							List<CommitTag> commitTags = new List<CommitTag>();
@@ -486,8 +499,8 @@ namespace Horde.Build.Perforce
 		{
 			readonly PerforceServiceCache _owner;
 
-			public CachedCommitSource(PerforceServiceCache owner, StreamConfig stream)
-				: base(owner, stream)
+			public CachedCommitSource(PerforceServiceCache owner, StreamConfig stream, ILogger logger)
+				: base(owner, stream, logger)
 			{
 				_owner = owner;
 			}
@@ -693,7 +706,7 @@ namespace Horde.Build.Perforce
 		/// <inheritdoc/>
 		public override ICommitCollection GetCommits(StreamConfig streamConfig)
 		{
-			return new CachedCommitSource(this, streamConfig);
+			return new CachedCommitSource(this, streamConfig, _logger);
 		}
 
 		#endregion
diff --git a/Engine/Source/Programs/Horde/Horde.Build/Streams/StreamsController.cs b/Engine/Source/Programs/Horde/Horde.Build/Streams/StreamsController.cs
--- a/Engine/Source/Programs/Horde/Horde.Build/Streams/StreamsController.cs
+++ b/Engine/Source/Programs/Horde/Horde.Build/Streams/StreamsController.cs
@@ -275,12 +275,13 @@ namespace Horde.Build.Streams
 		/// </summary>
 		/// <param name="streamId">The stream id</param>
 		/// <param name="changeNumber">The changelist number</param>
+		/// <param name="maxFiles">Maximum number of files to return</param>
 		/// <param name="filter">The filter to apply to the results</param>
 		/// <returns>Http result code</returns>
 		[HttpGet]
 		[Route("/api/v1/streams/{streamId}/changes/{changeNumber}")]
 		[ProducesResponseType(typeof(GetCommitResponse), 200)]
-		public async Task<ActionResult<object>> GetChangeDetailsAsync(StreamId streamId, int changeNumber, PropertyFilter? filter = null)
+		public async Task<ActionResult<object>> GetChangeDetailsAsync(StreamId streamId, int changeNumber, int maxFiles = 100, PropertyFilter? filter = null)
 		{
 			StreamConfig? streamConfig;
 			if (!_globalConfig.Value.TryGetStream(streamId, out streamConfig))
@@ -300,7 +301,7 @@ namespace Horde.Build.Streams
 
 			IUser? author = await _userCollection.GetCachedUserAsync(changeDetails.AuthorId);
 			IReadOnlyList<CommitTag> tags = await changeDetails.GetTagsAsync(HttpContext.RequestAborted);
-			IReadOnlyList<string> files = await changeDetails.GetFilesAsync(CancellationToken.None);
+			IReadOnlyList<string> files = await changeDetails.GetFilesAsync(maxFiles, CancellationToken.None);
 
 			return PropertyFilter.Apply(new GetCommitResponse(changeDetails, author!, tags, files), filter);
 		}
