Horde: Missing file.

diff --git a/Engine/Source/Programs/Horde/Horde.Build.Tests/Stubs/Services/PerforceServiceStub.cs b/Engine/Source/Programs/Horde/Horde.Build.Tests/Stubs/Services/PerforceServiceStub.cs
--- a/Engine/Source/Programs/Horde/Horde.Build.Tests/Stubs/Services/PerforceServiceStub.cs
+++ b/Engine/Source/Programs/Horde/Horde.Build.Tests/Stubs/Services/PerforceServiceStub.cs
@@ -78,7 +78,12 @@ namespace Horde.Build.Tests.Stubs.Services
 				return new ValueTask<IReadOnlyList<CommitTag>>(tags);
 			}
 
-			public ValueTask<IReadOnlyList<string>> GetFilesAsync(CancellationToken cancellationToken)
+			public ValueTask<bool> MatchesFilterAsync(FileFilter filter, CancellationToken cancellationToken)
+			{
+				throw new NotImplementedException();
+			}
+
+			public ValueTask<IReadOnlyList<string>> GetFilesAsync(int maxFiles, CancellationToken cancellationToken)
 			{
 				return new ValueTask<IReadOnlyList<string>>(Files);
 			}
