Horde: Fix tests.

diff --git a/Engine/Source/Programs/Horde/Horde.Build.Tests/Stubs/Services/PerforceServiceStub.cs b/Engine/Source/Programs/Horde/Horde.Build.Tests/Stubs/Services/PerforceServiceStub.cs
--- a/Engine/Source/Programs/Horde/Horde.Build.Tests/Stubs/Services/PerforceServiceStub.cs
+++ b/Engine/Source/Programs/Horde/Horde.Build.Tests/Stubs/Services/PerforceServiceStub.cs
@@ -80,7 +80,7 @@ namespace Horde.Build.Tests.Stubs.Services
 
 			public ValueTask<bool> MatchesFilterAsync(FileFilter filter, CancellationToken cancellationToken)
 			{
-				throw new NotImplementedException();
+				return new ValueTask<bool>(filter.ApplyTo(Files).Any());
 			}
 
 			public ValueTask<IReadOnlyList<string>> GetFilesAsync(int maxFiles, CancellationToken cancellationToken)
diff --git a/Engine/Source/Programs/Horde/Horde.Build/Jobs/Schedules/ScheduleService.cs b/Engine/Source/Programs/Horde/Horde.Build/Jobs/Schedules/ScheduleService.cs
--- a/Engine/Source/Programs/Horde/Horde.Build/Jobs/Schedules/ScheduleService.cs
+++ b/Engine/Source/Programs/Horde/Horde.Build/Jobs/Schedules/ScheduleService.cs
@@ -494,7 +494,7 @@ namespace Horde.Build.Jobs.Schedules
 			}
 			if (fileFilter != null)
 			{
-				if (await commit.MatchesFilterAsync(fileFilter, cancellationToken))
+				if (!await commit.MatchesFilterAsync(fileFilter, cancellationToken))
 				{
 					_logger.LogDebug("Not building change {Change} due to file filter", commit.Number);
 					return false;
