Horde: Use log level info for scaling messages in AWS recycling fleet manager

diff --git a/Engine/Source/Programs/Horde/Horde.Build/Agents/Fleet/Providers/AwsRecyclingFleetManager.cs b/Engine/Source/Programs/Horde/Horde.Build/Agents/Fleet/Providers/AwsRecyclingFleetManager.cs
--- a/Engine/Source/Programs/Horde/Horde.Build/Agents/Fleet/Providers/AwsRecyclingFleetManager.cs
+++ b/Engine/Source/Programs/Horde/Horde.Build/Agents/Fleet/Providers/AwsRecyclingFleetManager.cs
@@ -139,11 +139,11 @@ public sealed class AwsRecyclingFleetManager : IFleetManager
 		{
 			if (instancesToStartCount == 0)
 			{
-				_logger.LogWarning("Unable to start any instance(s). No stopped instances are available.");
+				_logger.LogInformation("Unable to start any instance(s). No stopped instances are available.");
 			}
 			else if (stoppedInstancesMissingCount > 0)
 			{
-				_logger.LogWarning("Starting {InstancesToStartCount} instance(s) but not enough stopped instances to accommodate the full pool scale-out", instancesToStartCount);
+				_logger.LogInformation("Starting {InstancesToStartCount} instance(s) but not enough stopped instances to accommodate the full pool scale-out", instancesToStartCount);
 			}
 			else
 			{
