Horde: Fix min height for job streams with few jobs
#skipci
#rnx

diff --git a/Engine/Source/Programs/Horde/HordeDashboard/src/components/JobView.tsx b/Engine/Source/Programs/Horde/HordeDashboard/src/components/JobView.tsx
--- a/Engine/Source/Programs/Horde/HordeDashboard/src/components/JobView.tsx
+++ b/Engine/Source/Programs/Horde/HordeDashboard/src/components/JobView.tsx
@@ -587,10 +587,10 @@ const JobList: React.FC<{ tab: string; filter: JobFilterSimple, controller: Call
 
    return (
       <Stack tokens={{ childrenGap: 0 }} className={detailClasses.detailsRow}>
-         <div className={detailClasses.container} style={{ width: "100%", height: 'calc(100vh - 270px)', position: 'relative', marginTop: 0 }}>
+         <div className={detailClasses.container} style={{ width: "100%", height: 'calc(100vh - 240px)', position: 'relative', marginTop: 0 }}>
             {<ScrollablePane scrollbarVisibility={ScrollbarVisibility.always} onScroll={() => { controller.setState(undefined, true) }} style={{ overflow: "visible" }}>
                {renderItems.length > 0 &&
-                  <Stack style={{ width: width, marginLeft: 4, boxShadow: "0 1.6px 3.6px 0 rgba(0,0,0,0.132), 0 0.3px 0.9px 0 rgba(0,0,0,0.108)", backgroundColor: "#FFFFFF" }}>
+                  <Stack style={{ width: width, height:'calc(100vh - 260px)', marginLeft: 4, boxShadow: "0 1.6px 3.6px 0 rgba(0,0,0,0.132), 0 0.3px 0.9px 0 rgba(0,0,0,0.108)", backgroundColor: "#FFFFFF" }}>
                      <DetailsList
                         styles={{ headerWrapper: { overflow: "hidden" } }}
                         indentWidth={0}
diff --git a/Engine/Source/Programs/Horde/HordeDashboard/src/components/JobViewIncremental.tsx b/Engine/Source/Programs/Horde/HordeDashboard/src/components/JobViewIncremental.tsx
--- a/Engine/Source/Programs/Horde/HordeDashboard/src/components/JobViewIncremental.tsx
+++ b/Engine/Source/Programs/Horde/HordeDashboard/src/components/JobViewIncremental.tsx
@@ -599,11 +599,11 @@ const JobList: React.FC<{ tab: string; filter: JobFilterSimple, controller: Call
       <Stack tokens={{ childrenGap: 0 }} className={detailClasses.detailsRow}>
          <FocusZone direction={FocusZoneDirection.vertical} >
             {commitState.target && commitState.job && <ChangeContextMenu target={commitState.target} job={commitState.job} commit={commitState.commit} onDismiss={() => setCommitState({})} />}
-            <div className={detailClasses.container} style={{ width: "100%", height: 'calc(100vh - 270px)', position: 'relative', marginTop: 0 }} data-is-scrollable={true}>
+            <div className={detailClasses.container} style={{ width: "100%", height: 'calc(100vh - 240px)', position: 'relative', marginTop: 0 }} data-is-scrollable={true}>
 
                {<ScrollablePane scrollbarVisibility={ScrollbarVisibility.always} onScroll={() => { controller.setState(undefined, true) }} style={{ overflow: "visible" }}>
                   {renderItems.length > 0 &&
-                     <Stack style={{ width: width, marginLeft: 4, boxShadow: "0 1.6px 3.6px 0 rgba(0,0,0,0.132), 0 0.3px 0.9px 0 rgba(0,0,0,0.108)", backgroundColor: "#FFFFFF" }}>
+                     <Stack style={{ width: width, marginLeft: 4,  height:'calc(100vh - 260px)', boxShadow: "0 1.6px 3.6px 0 rgba(0,0,0,0.132), 0 0.3px 0.9px 0 rgba(0,0,0,0.108)", backgroundColor: "#FFFFFF" }}>
                         <DetailsList
                            styles={{
                               headerWrapper: { overflow: "hidden" }, root: {
