Horde: Fix job view height for real
#rnx
#skipci

diff --git a/Engine/Source/Programs/Horde/HordeDashboard/src/components/JobView.tsx b/Engine/Source/Programs/Horde/HordeDashboard/src/components/JobView.tsx
--- a/Engine/Source/Programs/Horde/HordeDashboard/src/components/JobView.tsx
+++ b/Engine/Source/Programs/Horde/HordeDashboard/src/components/JobView.tsx
@@ -590,9 +590,9 @@ const JobList: React.FC<{ tab: string; filter: JobFilterSimple, controller: Call
          <div className={detailClasses.container} style={{ width: "100%", height: 'calc(100vh - 240px)', position: 'relative', marginTop: 0 }}>
             {<ScrollablePane scrollbarVisibility={ScrollbarVisibility.always} onScroll={() => { controller.setState(undefined, true) }} style={{ overflow: "visible" }}>
                {renderItems.length > 0 &&
-                  <Stack style={{ width: width, height:'calc(100vh - 260px)', marginLeft: 4, boxShadow: "0 1.6px 3.6px 0 rgba(0,0,0,0.132), 0 0.3px 0.9px 0 rgba(0,0,0,0.108)", backgroundColor: "#FFFFFF" }}>
+                  <Stack style={{ width: width, marginLeft: 4, boxShadow: "0 1.6px 3.6px 0 rgba(0,0,0,0.132), 0 0.3px 0.9px 0 rgba(0,0,0,0.108)", backgroundColor: "#FFFFFF" }}>
                      <DetailsList
-                        styles={{ headerWrapper: { overflow: "hidden" } }}
+                        styles={{ root: {paddingBottom: 32}, headerWrapper: { overflow: "hidden" } }}
                         indentWidth={0}
                         compact={false}
                         selectionMode={SelectionMode.none}
diff --git a/Engine/Source/Programs/Horde/HordeDashboard/src/components/JobViewIncremental.tsx b/Engine/Source/Programs/Horde/HordeDashboard/src/components/JobViewIncremental.tsx
--- a/Engine/Source/Programs/Horde/HordeDashboard/src/components/JobViewIncremental.tsx
+++ b/Engine/Source/Programs/Horde/HordeDashboard/src/components/JobViewIncremental.tsx
@@ -603,10 +603,11 @@ const JobList: React.FC<{ tab: string; filter: JobFilterSimple, controller: Call
 
                {<ScrollablePane scrollbarVisibility={ScrollbarVisibility.always} onScroll={() => { controller.setState(undefined, true) }} style={{ overflow: "visible" }}>
                   {renderItems.length > 0 &&
-                     <Stack style={{ width: width, marginLeft: 4,  height:'calc(100vh - 260px)', boxShadow: "0 1.6px 3.6px 0 rgba(0,0,0,0.132), 0 0.3px 0.9px 0 rgba(0,0,0,0.108)", backgroundColor: "#FFFFFF" }}>
+                     <Stack style={{ width: width, marginLeft: 4, boxShadow: "0 1.6px 3.6px 0 rgba(0,0,0,0.132), 0 0.3px 0.9px 0 rgba(0,0,0,0.108)", backgroundColor: "#FFFFFF" }}>
                         <DetailsList
                            styles={{
                               headerWrapper: { overflow: "hidden" }, root: {
+                                 paddingBottom: 32, 
                                  selectors: {
                                     '.ms-List-cell': {
                                        minHeight: lineHeight
