Horde: Add command for generating markdown documentation from config file JSON schemas.

diff --git a/Engine/Source/Programs/Horde/Horde.Build/Commands/Config/DocsCommand.cs b/Engine/Source/Programs/Horde/Horde.Build/Commands/Config/DocsCommand.cs
new file mode 100644
--- /dev/null
+++ b/Engine/Source/Programs/Horde/Horde.Build/Commands/Config/DocsCommand.cs
@@ -0,0 +1,240 @@
+// Copyright Epic Games, Inc. All Rights Reserved.
+
+using System;
+using System.Collections.Generic;
+using System.Diagnostics.CodeAnalysis;
+using System.IO;
+using System.Linq;
+using System.Text.RegularExpressions;
+using System.Threading.Tasks;
+using EpicGames.Core;
+using Horde.Build.Projects;
+using Horde.Build.Server;
+using Horde.Build.Streams;
+using Horde.Build.Utilities;
+using Microsoft.Extensions.Logging;
+
+namespace Horde.Build.Commands.Config
+{
+	[Command("config", "docs", "Writes Markdown docs for server settings")]
+	class DocsCommand : Command
+	{
+		[CommandLine]
+		DirectoryReference? _outputDir = null!;
+
+		public override async Task<int> ExecuteAsync(ILogger logger)
+		{
+			_outputDir ??= DirectoryReference.Combine(Program.AppDir, "Docs");
+			DirectoryReference.CreateDirectory(_outputDir);
+
+			JsonSchema serverSchema = Schemas.CreateSchema(typeof(ServerSettings));
+			JsonSchema globalSchema = Schemas.CreateSchema(typeof(GlobalConfig));
+			JsonSchema projectSchema = Schemas.CreateSchema(typeof(ProjectConfig));
+			JsonSchema streamSchema = Schemas.CreateSchema(typeof(StreamConfig));
+
+			JsonSchemaType[] allTypes = { serverSchema.RootType, globalSchema.RootType, projectSchema.RootType, streamSchema.RootType };
+
+			await WriteDocAsync(serverSchema.RootType, "Server Config", "Config-Server.md", allTypes);
+			await WriteDocAsync(globalSchema.RootType, "Global Config", "Config-Globals.md", allTypes);
+			await WriteDocAsync(projectSchema.RootType, "Project Config", "Config-Projects.md", allTypes);
+			await WriteDocAsync(streamSchema.RootType, "Stream Config", "Config-Streams.md", allTypes);
+
+			return 0;
+		}
+
+		class ObjectQueue
+		{
+			readonly Stack<JsonSchemaType> _stack = new Stack<JsonSchemaType>();
+			readonly HashSet<string> _visited = new HashSet<string>(StringComparer.Ordinal);
+
+			public void Add(JsonSchemaType type)
+			{
+				if (type.Name != null && _visited.Add(type.Name))
+				{
+					_stack.Push(type);
+				}
+			}
+
+			public void Ignore(JsonSchemaType type)
+			{
+				if (type.Name != null)
+				{
+					_visited.Add(type.Name);
+				}
+			}
+
+			public bool TryPop([NotNullWhen(true)] out JsonSchemaType? obj) => _stack.TryPop(out obj);
+		}
+
+		async Task WriteDocAsync(JsonSchemaType rootType, string title, string fileName, IEnumerable<JsonSchemaType> ignoreTypes)
+		{
+			FileReference file = FileReference.Combine(_outputDir!, fileName);
+			using (FileStream stream = FileReference.Open(file, FileMode.Create, FileAccess.Write))
+			{
+				using (StreamWriter writer = new StreamWriter(stream))
+				{
+					await writer.WriteLineAsync($"# {title}");
+
+					HashSet<JsonSchemaType> visitedTypes = new HashSet<JsonSchemaType>(ignoreTypes);
+					visitedTypes.Remove(rootType);
+
+					List<JsonSchemaType> types = new List<JsonSchemaType>();
+					FindCustomTypes(rootType, types, visitedTypes);
+
+					foreach (JsonSchemaType schemaType in types)
+					{
+						await writer.WriteLineAsync();
+
+						if (schemaType != rootType)
+						{
+							await writer.WriteLineAsync($"## {GetHeadingName(schemaType)}");
+							await writer.WriteLineAsync();
+						}
+
+						if (schemaType.Description != null)
+						{
+							await writer.WriteLineAsync(schemaType.Description);
+							await writer.WriteLineAsync();
+						}
+
+						if (schemaType is JsonSchemaObject schemaObj)
+						{
+							await writer.WriteLineAsync("Name | Type | Description");
+							await writer.WriteLineAsync("---- | ---- | -----------");
+
+							foreach (JsonSchemaProperty property in schemaObj.Properties)
+							{
+								string name = property.CamelCaseName;
+								string type = GetType(property.Type);
+								string description = GetMarkdownDescription(property.Description);
+								await writer.WriteLineAsync($"`{name}` | {type} | {description}");
+							}
+						}
+						else if (schemaType is JsonSchemaEnum schemaEnum)
+						{
+							await writer.WriteLineAsync("Name | Description");
+							await writer.WriteLineAsync("---- | -----------");
+
+							for(int idx = 0; idx < schemaEnum.Values.Count; idx++)
+							{
+								string name = schemaEnum.Values[idx];
+								string description = GetMarkdownDescription(schemaEnum.Descriptions[idx]);
+								await writer.WriteLineAsync($"`{name}` | {description}");
+							}
+						}
+						else
+						{
+							throw new NotImplementedException();
+						}
+					}
+				}
+			}
+
+
+			await Task.Delay(1);
+		}
+
+		static void FindCustomTypes(JsonSchemaType type, List<JsonSchemaType> types, HashSet<JsonSchemaType> visited)
+		{
+			if (visited.Add(type))
+			{
+				switch (type)
+				{
+					case JsonSchemaOneOf oneOf:
+						foreach (JsonSchemaType oneOfType in oneOf.Types)
+						{
+							FindCustomTypes(oneOfType, types, visited);
+						}
+						break;
+					case JsonSchemaArray array:
+						FindCustomTypes(array.ItemType, types, visited);
+						break;
+					case JsonSchemaEnum _:
+						if (type.Name != null)
+						{
+							types.Add(type);
+						}
+						break;
+					case JsonSchemaObject obj:
+						if (type.Name != null)
+						{
+							types.Add(type);
+						}
+						foreach (JsonSchemaProperty property in obj.Properties)
+						{
+							FindCustomTypes(property.Type, types, visited);
+						}
+						break;
+				}
+			}
+		}
+
+		static string GetType(JsonSchemaType type)
+		{
+			switch (type)
+			{
+				case JsonSchemaBoolean _:
+					return "`boolean`";
+				case JsonSchemaInteger _:
+					return "`integer`";
+				case JsonSchemaNumber _:
+					return "`number`";
+				case JsonSchemaString _:
+					return "`string`";
+				case JsonSchemaOneOf oneOf:
+					return String.Join("/", oneOf.Types.Select(x => GetType(x)));
+				case JsonSchemaArray array:
+					string elementType = GetType(array.ItemType);
+					if (elementType.EndsWith("`", StringComparison.Ordinal))
+					{
+						return elementType.Insert(elementType.Length - 1, "[]");
+					}
+					else
+					{
+						return elementType + "`[]`";
+					}
+				case JsonSchemaEnum _:
+				case JsonSchemaObject _:
+					if (type.Name == null)
+					{
+						return "`object`";
+					}
+					else
+					{
+						return $"[`{type.Name}`](#{GetAnchorName(type)})";
+					}
+				default:
+					return type.GetType().Name;
+			}
+		}
+
+		static string GetMarkdownDescription(string? description)
+		{
+			return (description ?? String.Empty).Replace("\n", "<br>", StringComparison.Ordinal);
+		}
+
+		static string GetHeadingName(JsonSchemaType type)
+		{
+			if (type.Name == null)
+			{
+				throw new NotImplementedException("Unknown type");
+			}
+
+			switch (type)
+			{
+				case JsonSchemaEnum _:
+					return $"{type.Name} (Enum)";
+				default:
+					return type.Name;
+			}
+		}
+
+		[SuppressMessage("Globalization", "CA1308:Normalize strings to uppercase")]
+		static string GetAnchorName(JsonSchemaType type)
+		{
+			string anchor = GetHeadingName(type).ToLowerInvariant();
+			anchor = Regex.Replace(anchor, @"[^a-z0-9]+", "-");
+			return anchor.Trim('-');
+		}
+	}
+}
diff --git a/Engine/Source/Programs/Shared/EpicGames.Core/JsonSchema.cs b/Engine/Source/Programs/Shared/EpicGames.Core/JsonSchema.cs
--- a/Engine/Source/Programs/Shared/EpicGames.Core/JsonSchema.cs
+++ b/Engine/Source/Programs/Shared/EpicGames.Core/JsonSchema.cs
@@ -296,13 +296,32 @@ namespace EpicGames.Core
 		/// <param name="xmlDoc"></param>
 		/// <returns></returns>
 		static string? GetPropertyDescription(Type type, string name, XmlDocument? xmlDoc)
+		{
+			return GetSummaryFromXmlDoc(type, "P", name, xmlDoc);
+		}
+
+		/// <summary>
+		/// Gets a description from Xml documentation file
+		/// </summary>
+		/// <param name="type">Type to retrieve summary for</param>
+		/// <param name="qualifier">Type of element to retrieve</param>
+		/// <param name="member">Name of the member</param>
+		/// <param name="xmlDoc">XML documentation file to search</param>
+		/// <returns>Summary string, or null if it's not available</returns>
+		static string? GetSummaryFromXmlDoc(Type type, string qualifier, string? member, XmlDocument? xmlDoc)
 		{
 			if (xmlDoc == null)
 			{
 				return null;
 			}
 
-			string selector = $"//member[@name='P:{type.FullName}.{name}']/summary";
+			string? fullName = type.FullName;
+			if (member != null)
+			{
+				fullName = $"{fullName}.{member}";
+			}
+
+			string selector = $"//member[@name='{qualifier}:{fullName}']/summary";
 
 			XmlNode? node = xmlDoc.SelectSingleNode(selector);
 			if (node == null)
@@ -336,7 +355,7 @@ namespace EpicGames.Core
 				case TypeCode.UInt64:
 					if (type.IsEnum)
 					{
-						return new JsonSchemaEnum(Enum.GetNames(type)) { Name = type.Name };
+						return CreateEnumSchemaType(type, xmlDoc);
 					}
 					else
 					{
@@ -368,6 +387,10 @@ namespace EpicGames.Core
 			{
 				return new JsonSchemaString();
 			}
+			if (type == typeof(Uri))
+			{
+				return new JsonSchemaString(JsonSchemaStringFormat.Uri);
+			}
 
 			Type[] interfaceTypes = type.GetInterfaces();
 			foreach (Type interfaceType in interfaceTypes)
@@ -415,6 +438,20 @@ namespace EpicGames.Core
 			throw new Exception($"Unknown type for schema generation: {type}");
 		}
 
+		static JsonSchemaEnum CreateEnumSchemaType(Type type, XmlDocument? xmlDoc)
+		{
+			string[] names = Enum.GetNames(type);
+			string[] descriptions = new string[names.Length];
+
+			for (int idx = 0; idx < names.Length; idx++)
+			{
+				descriptions[idx] = GetSummaryFromXmlDoc(type, "F", names[idx], xmlDoc) ?? String.Empty;
+			}
+
+			string? enumDescription = GetSummaryFromXmlDoc(type, "T", null, xmlDoc);
+			return new JsonSchemaEnum(names, descriptions) { Name = type.Name, Description = enumDescription };
+		}
+
 		static void SetOneOfProperties(JsonSchemaOneOf obj, Type type, Type[] knownTypes, Dictionary<Type, JsonSchemaType> typeCache, XmlDocument? xmlDoc)
 		{
 			obj.Name = type.Name;
@@ -425,7 +462,7 @@ namespace EpicGames.Core
 				if (attribute != null)
 				{
 					JsonSchemaObject knownObject = new JsonSchemaObject();
-					knownObject.Properties.Add(new JsonSchemaProperty("type", "Type discriminator", new JsonSchemaEnum(new[] { attribute.Name })));
+					knownObject.Properties.Add(new JsonSchemaProperty("type", "Type discriminator", new JsonSchemaEnum(new[] { attribute.Name }, new[] { "Identifier for the derived type" })));
 					SetObjectProperties(knownObject, knownType, typeCache, xmlDoc);
 					obj.Types.Add(knownObject);
 				}
@@ -435,6 +472,7 @@ namespace EpicGames.Core
 		static void SetObjectProperties(JsonSchemaObject obj, Type type, Dictionary<Type, JsonSchemaType> typeCache, XmlDocument? xmlDoc)
 		{
 			obj.Name = type.Name;
+			obj.Description = GetSummaryFromXmlDoc(type, "T", null, xmlDoc);
 
 			PropertyInfo[] properties = type.GetProperties(BindingFlags.Instance | BindingFlags.Public);
 			foreach (PropertyInfo property in properties)
diff --git a/Engine/Source/Programs/Shared/EpicGames.Core/JsonSchemaType.cs b/Engine/Source/Programs/Shared/EpicGames.Core/JsonSchemaType.cs
--- a/Engine/Source/Programs/Shared/EpicGames.Core/JsonSchemaType.cs
+++ b/Engine/Source/Programs/Shared/EpicGames.Core/JsonSchemaType.cs
@@ -15,6 +15,11 @@ namespace EpicGames.Core
 		/// </summary>
 		public string? Name { get; set; }
 
+		/// <summary>
+		/// Description of this type
+		/// </summary>
+		public string? Description { get; set; }
+
 		/// <summary>
 		/// Write this type to the given archive
 		/// </summary>
@@ -154,12 +159,18 @@ namespace EpicGames.Core
 		/// </summary>
 		public List<string> Values { get; } = new List<string>();
 
+		/// <summary>
+		/// Descriptions for each enum value
+		/// </summary>
+		public List<string> Descriptions { get; } = new List<string>();
+
 		/// <summary>
 		/// Constructor
 		/// </summary>
-		public JsonSchemaEnum(IEnumerable<string> values)
+		public JsonSchemaEnum(IEnumerable<string> values, IEnumerable<string> descriptions)
 		{
 			Values.AddRange(values);
+			Descriptions.AddRange(descriptions);
 		}
 
 		/// <inheritdoc/>
@@ -212,6 +223,11 @@ namespace EpicGames.Core
 		/// </summary>
 		public string Name { get; set; }
 
+		/// <summary>
+		/// The camelcase name for this property
+		/// </summary>
+		public string CamelCaseName => Char.ToLowerInvariant(Name[0]) + Name.Substring(1);
+
 		/// <summary>
 		/// Description of the property
 		/// </summary>
@@ -238,8 +254,7 @@ namespace EpicGames.Core
 		/// <param name="writer"></param>
 		public void Write(IJsonSchemaWriter writer)
 		{
-			string camelCaseName = Char.ToLowerInvariant(Name[0]) + Name.Substring(1);
-			writer.WriteStartObject(camelCaseName);
+			writer.WriteStartObject(CamelCaseName);
 			writer.WriteType(Type);
 			if (Description != null)
 			{
