Horde: Tweak documentation filenames.

diff --git a/Engine/Source/Programs/Horde/Horde.Build/Commands/Config/DocsCommand.cs b/Engine/Source/Programs/Horde/Horde.Build/Commands/Config/DocsCommand.cs
--- a/Engine/Source/Programs/Horde/Horde.Build/Commands/Config/DocsCommand.cs
+++ b/Engine/Source/Programs/Horde/Horde.Build/Commands/Config/DocsCommand.cs
@@ -20,12 +20,11 @@ namespace Horde.Build.Commands.Config
 	class DocsCommand : Command
 	{
 		[CommandLine]
-		DirectoryReference? _outputDir = null!;
+		public DirectoryReference OutputDir { get; set; } = DirectoryReference.Combine(Program.AppDir, "Docs");
 
 		public override async Task<int> ExecuteAsync(ILogger logger)
 		{
-			_outputDir ??= DirectoryReference.Combine(Program.AppDir, "Docs");
-			DirectoryReference.CreateDirectory(_outputDir);
+			DirectoryReference.CreateDirectory(OutputDir);
 
 			JsonSchema serverSchema = Schemas.CreateSchema(typeof(ServerSettings));
 			JsonSchema globalSchema = Schemas.CreateSchema(typeof(GlobalConfig));
@@ -34,10 +33,10 @@ namespace Horde.Build.Commands.Config
 
 			JsonSchemaType[] allTypes = { serverSchema.RootType, globalSchema.RootType, projectSchema.RootType, streamSchema.RootType };
 
-			await WriteDocAsync(serverSchema.RootType, "Server Config", "Config-Server.md", allTypes);
-			await WriteDocAsync(globalSchema.RootType, "Global Config", "Config-Globals.md", allTypes);
-			await WriteDocAsync(projectSchema.RootType, "Project Config", "Config-Projects.md", allTypes);
-			await WriteDocAsync(streamSchema.RootType, "Stream Config", "Config-Streams.md", allTypes);
+			await WriteDocAsync(serverSchema.RootType, "Server Config Reference", "Config-Server.md", allTypes);
+			await WriteDocAsync(globalSchema.RootType, "Global Config Reference", "Config-Global.md", allTypes);
+			await WriteDocAsync(projectSchema.RootType, "Project Config Reference", "Config-Project.md", allTypes);
+			await WriteDocAsync(streamSchema.RootType, "Stream Config Reference", "Config-Stream.md", allTypes);
 
 			return 0;
 		}
@@ -68,7 +67,7 @@ namespace Horde.Build.Commands.Config
 
 		async Task WriteDocAsync(JsonSchemaType rootType, string title, string fileName, IEnumerable<JsonSchemaType> ignoreTypes)
 		{
-			FileReference file = FileReference.Combine(_outputDir!, fileName);
+			FileReference file = FileReference.Combine(OutputDir, fileName);
 			using (FileStream stream = FileReference.Open(file, FileMode.Create, FileAccess.Write))
 			{
 				using (StreamWriter writer = new StreamWriter(stream))
