Horde: Support dictionaries in markdown documentation.

diff --git a/Engine/Source/Programs/Horde/Horde.Build/Commands/Config/DocsCommand.cs b/Engine/Source/Programs/Horde/Horde.Build/Commands/Config/DocsCommand.cs
--- a/Engine/Source/Programs/Horde/Horde.Build/Commands/Config/DocsCommand.cs
+++ b/Engine/Source/Programs/Horde/Horde.Build/Commands/Config/DocsCommand.cs
@@ -7,6 +7,7 @@ using System.IO;
 using System.Linq;
 using System.Text.RegularExpressions;
 using System.Threading.Tasks;
+using Amazon.S3.Model.Internal.MarshallTransformations;
 using EpicGames.Core;
 using Horde.Build.Projects;
 using Horde.Build.Server;
@@ -22,21 +23,26 @@ namespace Horde.Build.Commands.Config
 		[CommandLine]
 		public DirectoryReference OutputDir { get; set; } = DirectoryReference.Combine(Program.AppDir, "Docs");
 
+		class AppSettings
+		{
+			public ServerSettings Horde { get; set; } = new ServerSettings();
+		}
+
 		public override async Task<int> ExecuteAsync(ILogger logger)
 		{
 			DirectoryReference.CreateDirectory(OutputDir);
 
-			JsonSchema serverSchema = Schemas.CreateSchema(typeof(ServerSettings));
+			JsonSchema serverSchema = Schemas.CreateSchema(typeof(AppSettings));
 			JsonSchema globalSchema = Schemas.CreateSchema(typeof(GlobalConfig));
 			JsonSchema projectSchema = Schemas.CreateSchema(typeof(ProjectConfig));
 			JsonSchema streamSchema = Schemas.CreateSchema(typeof(StreamConfig));
 
 			JsonSchemaType[] allTypes = { serverSchema.RootType, globalSchema.RootType, projectSchema.RootType, streamSchema.RootType };
 
-			await WriteDocAsync(serverSchema.RootType, "Server Config Reference", "Config-Server.md", allTypes);
-			await WriteDocAsync(globalSchema.RootType, "Global Config Reference", "Config-Global.md", allTypes);
-			await WriteDocAsync(projectSchema.RootType, "Project Config Reference", "Config-Project.md", allTypes);
-			await WriteDocAsync(streamSchema.RootType, "Stream Config Reference", "Config-Stream.md", allTypes);
+			await WriteDocAsync(serverSchema.RootType, "appsettings.json", "Config-Server.md", allTypes);
+			await WriteDocAsync(globalSchema.RootType, "Globals.json", "Config-Global.md", allTypes);
+			await WriteDocAsync(projectSchema.RootType, "*.project.json", "Config-Project.md", allTypes);
+			await WriteDocAsync(streamSchema.RootType, "*.stream.json", "Config-Stream.md", allTypes);
 
 			return 0;
 		}
@@ -92,7 +98,8 @@ namespace Horde.Build.Commands.Config
 
 						if (schemaType.Description != null)
 						{
-							await writer.WriteLineAsync(schemaType.Description);
+							string description = Regex.Replace(schemaType.Description, "\r?\n", Environment.NewLine);
+							await writer.WriteLineAsync(description);
 							await writer.WriteLineAsync();
 						}
 
@@ -104,7 +111,7 @@ namespace Horde.Build.Commands.Config
 							foreach (JsonSchemaProperty property in schemaObj.Properties)
 							{
 								string name = property.CamelCaseName;
-								string type = GetType(property.Type);
+								string type = GetMarkdownType(property.Type);
 								string description = GetMarkdownDescription(property.Description);
 								await writer.WriteLineAsync($"`{name}` | {type} | {description}");
 							}
@@ -159,6 +166,10 @@ namespace Horde.Build.Commands.Config
 						{
 							types.Add(type);
 						}
+						if (obj.AdditionalProperties != null)
+						{
+							FindCustomTypes(obj.AdditionalProperties, types, visited);
+						}
 						foreach (JsonSchemaProperty property in obj.Properties)
 						{
 							FindCustomTypes(property.Type, types, visited);
@@ -168,7 +179,7 @@ namespace Horde.Build.Commands.Config
 			}
 		}
 
-		static string GetType(JsonSchemaType type)
+		static string GetMarkdownType(JsonSchemaType type)
 		{
 			switch (type)
 			{
@@ -181,9 +192,9 @@ namespace Horde.Build.Commands.Config
 				case JsonSchemaString _:
 					return "`string`";
 				case JsonSchemaOneOf oneOf:
-					return String.Join("/", oneOf.Types.Select(x => GetType(x)));
+					return String.Join("/", oneOf.Types.Select(x => GetMarkdownType(x)));
 				case JsonSchemaArray array:
-					string elementType = GetType(array.ItemType);
+					string elementType = GetMarkdownType(array.ItemType);
 					if (elementType.EndsWith("`", StringComparison.Ordinal))
 					{
 						return elementType.Insert(elementType.Length - 1, "[]");
@@ -192,9 +203,22 @@ namespace Horde.Build.Commands.Config
 					{
 						return elementType + "`[]`";
 					}
-				case JsonSchemaEnum _:
-				case JsonSchemaObject _:
+
+				case JsonSchemaEnum en:
 					if (type.Name == null)
+					{
+						return String.Join("<br>", en.Values);
+					}
+					else
+					{
+						return $"[`{type.Name}`](#{GetAnchorName(type)})";
+					}
+				case JsonSchemaObject obj:
+					if (obj.AdditionalProperties != null)
+					{
+						return "`string` `->` " + GetMarkdownType(obj.AdditionalProperties);
+					}
+					else if (type.Name == null)
 					{
 						return "`object`";
 					}
