Horde: Remove XGE log parsing from agent

Switching to *.ib_mon files for extracting data can now be done. See CL 23909576.

diff --git a/Engine/Source/Programs/Horde/Horde.Agent.Tests/LogParserTests.cs b/Engine/Source/Programs/Horde/Horde.Agent.Tests/LogParserTests.cs
--- a/Engine/Source/Programs/Horde/Horde.Agent.Tests/LogParserTests.cs
+++ b/Engine/Source/Programs/Horde/Horde.Agent.Tests/LogParserTests.cs
@@ -573,86 +573,6 @@ namespace Horde.Agent.Tests
 			CheckEventGroup(logEvents, 0, 1, LogLevel.Information, KnownLogEvents.Systemic_SlowDDC);
 		}
 
-		[TestMethod]
-		public void XoreaxErrorMatcher()
-		{
-			string[] lines =
-			{
-				@"--------------------Build System Warning---------------------------------------",
-				@"Failed to connect to Coordinator:",
-				@"    All builds will run in standalone mode.",
-				@"--------------------Build System Warning (Agent 'Cloud_1p5rlm3rq_10 (Core #9)')",
-				@"Remote tasks distribution:",
-				@"    Tasks execution is impeded due to low agent responsiveness",
-				@"-------------------------------------------------------------------------------",
-				@"    LogXGEController: Warning: XGE's background service (BuildService.exe) is not running - service is likely disabled on this machine.",
-				@"BUILD FAILED: Command failed (Result:1): C:\Program Files (x86)\IncrediBuild\xgConsole.exe ""d:\build\Sync\Engine\Programs\AutomationTool\Saved\Logs\UAT_XGE.xml"" /Rebuild /NoLogo /ShowAgent /ShowTime"
-			};
-
-			List<LogEvent> logEvents = Parse(lines);
-			Assert.AreEqual(8, logEvents.Count);
-			CheckEventGroup(logEvents.Slice(0, 3), 0, 3, LogLevel.Information, KnownLogEvents.Systemic_Xge_Standalone);
-			CheckEventGroup(logEvents.Slice(3, 3), 3, 3, LogLevel.Information, KnownLogEvents.Systemic_Xge);
-			CheckEventGroup(logEvents.Slice(6, 1), 7, 1, LogLevel.Information, KnownLogEvents.Systemic_Xge_ServiceNotRunning);
-			CheckEventGroup(logEvents.Slice(7, 1), 8, 1, LogLevel.Error, KnownLogEvents.Systemic_Xge_BuildFailed);
-		}
-		
-		[TestMethod]
-		public void XoreaxTaskMetadataMatcher()
-		{
-			string[] lines =
-			{
-				@"MiMalloc.c (0:01.17 at +0:11)",
-				@"SomeFile.cpp (4:31.12 at +0:16)",
-				@"Default.rc2 (Agent 'SomeAgentName (Core #7)', 2:34.56 at +55:18)",
-				@" (1:01.12 at +58:03)",
-			};
-
-			List<LogEvent> logEvents = Parse(lines);
-			Assert.AreEqual(4, logEvents.Count);
-			CheckEventGroup(logEvents.Slice(0, 1), 0, 1, LogLevel.Information, KnownLogEvents.Systemic_Xge_TaskMetadata);
-			CheckEventGroup(logEvents.Slice(1, 1), 1, 1, LogLevel.Information, KnownLogEvents.Systemic_Xge_TaskMetadata);
-			CheckEventGroup(logEvents.Slice(2, 1), 2, 1, LogLevel.Information, KnownLogEvents.Systemic_Xge_TaskMetadata);
-			CheckEventGroup(logEvents.Slice(3, 1), 3, 1, LogLevel.Information, KnownLogEvents.Systemic_Xge_TaskMetadata);
-
-			Assert.AreEqual("MiMalloc.c", logEvents[0].GetProperty<string>("name"));
-			Assert.AreEqual(1170, logEvents[0].GetProperty<int>("duration"));
-			Assert.AreEqual(11000, logEvents[0].GetProperty<int>("startTime"));
-			
-			Assert.AreEqual("SomeFile.cpp", logEvents[1].GetProperty<string>("name"));
-			Assert.AreEqual(271120, logEvents[1].GetProperty<int>("duration"));
-			Assert.AreEqual(16000, logEvents[1].GetProperty<int>("startTime"));
-			
-			Assert.AreEqual("Default.rc2", logEvents[2].GetProperty<string>("name"));
-			Assert.AreEqual("SomeAgentName (Core #7)", logEvents[2].GetProperty<string>("agent"));
-			Assert.AreEqual(154560, logEvents[2].GetProperty<int>("duration"));
-			Assert.AreEqual(3318000, logEvents[2].GetProperty<int>("startTime"));
-			
-			Assert.AreEqual("", logEvents[3].GetProperty<string>("name"));
-			Assert.AreEqual(61120, logEvents[3].GetProperty<int>("duration"));
-			Assert.AreEqual(3483000, logEvents[3].GetProperty<int>("startTime"));
-		}
-
-		private class FakeLogSink : ILogEventSink
-		{
-			public List<LogEvent> ReceivedLogEvents { get; } = new();
-			public void ProcessEvent(LogEvent logEvent)
-			{
-				ReceivedLogEvents.Add(logEvent);
-			}
-		}
-		
-		[TestMethod]
-		public void AdditionalLogSinks()
-		{
-			FakeLogSink fakeLogSink = new();
-			using LogParser parser = new (NullLogger.Instance, new List<string>(), new List<ILogEventSink>() { fakeLogSink });
-
-			parser.WriteLine(@"MiMalloc.c (0:01.17 at +0:11)");
-			
-			Assert.AreEqual(1, fakeLogSink.ReceivedLogEvents.Count);
-			Assert.AreEqual("MiMalloc.c", fakeLogSink.ReceivedLogEvents[0].GetProperty<string>("name"));
-		}
 
 		[TestMethod]
 		public void LogChannelMatcher()
diff --git a/Engine/Source/Programs/Horde/Horde.Agent/Execution/JobExecutor.cs b/Engine/Source/Programs/Horde/Horde.Agent/Execution/JobExecutor.cs
--- a/Engine/Source/Programs/Horde/Horde.Agent/Execution/JobExecutor.cs
+++ b/Engine/Source/Programs/Horde/Horde.Agent/Execution/JobExecutor.cs
@@ -1200,29 +1200,8 @@ namespace Horde.Agent.Execution
 
 			List<string> ignorePatterns = await ReadIgnorePatternsAsync(workspaceDir, jobLogger);
 
-			InterceptingLogger interceptedJobLogger = new (jobLogger, new [] { KnownLogEvents.Systemic_Xge_TaskMetadata.Id }, (level, id, state, exception) =>
-			{
-				if (state is JsonLogEvent jsonLogEvent)
-				{
-					try
-					{
-						LogEvent logEvent = LogEvent.Read(jsonLogEvent.Data.Span);
-						if (logEvent.TryGetProperty("agent", out string? agentName))
-						{
-							string taskName = logEvent.GetProperty<string>("name");
-							int duration = logEvent.GetProperty<int>("duration");
-							_logger.LogInformation("Executed XGE task {XgeTaskName} on agent {XgeAgent} for {Duration} ms", taskName, agentName, duration);	
-						}
-					}
-					catch (Exception e)
-					{
-						_logger.LogWarning(e, "Failed to log XGE task execution");
-					}
-				}
-			});
-
 			int exitCode;
-			using (LogParser filter = new LogParser(interceptedJobLogger, ignorePatterns))
+			using (LogParser filter = new LogParser(jobLogger, ignorePatterns))
 			{
 				await ExecuteCleanupScriptAsync(cleanupScript, filter, jobLogger);
 				try
