Horde: Add options to enable the new temp storage and logger backends for individual jobs, inheritable by stream and project.

diff --git a/Engine/Source/Programs/Horde/Horde.Agent.Tests/AwsLambdaListenerCommandTest.cs b/Engine/Source/Programs/Horde/Horde.Agent.Tests/AwsLambdaListenerCommandTest.cs
--- a/Engine/Source/Programs/Horde/Horde.Agent.Tests/AwsLambdaListenerCommandTest.cs
+++ b/Engine/Source/Programs/Horde/Horde.Agent.Tests/AwsLambdaListenerCommandTest.cs
@@ -35,7 +35,7 @@ namespace Horde.Agent.Tests
 			_serviceCollection.AddLogging();
 			_serviceCollection.AddSingleton<CapabilitiesService>();
 			_serviceCollection.AddSingleton<LeaseHandler, TestTaskHandler>();
-			_serviceCollection.AddSingleton<ISessionFactoryService>(sp => new FakeServerSessionFactory(_fakeServer));
+			_serviceCollection.AddSingleton<ISessionFactory>(sp => new FakeServerSessionFactory(_fakeServer));
 		}
 
 		public void Dispose()
diff --git a/Engine/Source/Programs/Horde/Horde.Agent.Tests/Execution/WorkspaceExecutorTest.cs b/Engine/Source/Programs/Horde/Horde.Agent.Tests/Execution/WorkspaceExecutorTest.cs
--- a/Engine/Source/Programs/Horde/Horde.Agent.Tests/Execution/WorkspaceExecutorTest.cs
+++ b/Engine/Source/Programs/Horde/Horde.Agent.Tests/Execution/WorkspaceExecutorTest.cs
@@ -8,6 +8,8 @@ using System.Threading.Tasks;
 using EpicGames.Core;
 using Horde.Agent.Execution;
 using Horde.Agent.Services;
+using Horde.Common;
+using HordeCommon.Rpc.Tasks;
 using Microsoft.Extensions.Logging;
 using Microsoft.Extensions.Logging.Abstractions;
 using Microsoft.VisualStudio.TestTools.UnitTesting;
@@ -46,8 +48,9 @@ public sealed class WorkspaceExecutorTest : IDisposable
 		_autoSdkWorkspace.SetFile(1, "HostWin64/Android/base.h", "base");
 		_workspace.SetFile(1, "main.cpp", "main");
 		_workspace.SetFile(1, "foo/bar/baz.h", "baz");
-		
-		_executor = new (_session, JobId, "batch1", AgentType, null, _workspace, null!, NullLogger.Instance);
+
+		JobExecutorOptions executorOptions = new JobExecutorOptions(_session, null!, JobId, "batch1", AgentType, new JobOptions());
+		_executor = new (executorOptions, null, _workspace, NullLogger.Instance);
 	}
 
 	public void Dispose()
@@ -68,7 +71,8 @@ public sealed class WorkspaceExecutorTest : IDisposable
 	[TestMethod]
 	public async Task RegularAndAutoSdkWorkspace()
 	{
-		WorkspaceExecutor executor = new (_session, JobId, "batch1", AgentType, _autoSdkWorkspace, _workspace, null!, NullLogger.Instance);
+		JobExecutorOptions executorOptions = new JobExecutorOptions(_session, null!, JobId, "batch1", AgentType, new JobOptions());
+		WorkspaceExecutor executor = new (executorOptions, _autoSdkWorkspace, _workspace, NullLogger.Instance);
 
 		await executor.InitializeAsync(_logger, CancellationToken.None);
 		AssertWorkspaceFile(_autoSdkWorkspace, "HostWin64/Android/base.h", "base");
@@ -80,7 +84,8 @@ public sealed class WorkspaceExecutorTest : IDisposable
 	[TestMethod]
 	public async Task EnvVars()
 	{
-		WorkspaceExecutor executor = new (_session, JobId, "batch1", AgentType, _autoSdkWorkspace, _workspace, null!, NullLogger.Instance);
+		JobExecutorOptions executorOptions = new JobExecutorOptions(_session, null!, JobId, "batch1", AgentType, new JobOptions());
+		WorkspaceExecutor executor = new (executorOptions, _autoSdkWorkspace, _workspace, NullLogger.Instance);
 		await executor.InitializeAsync(_logger, CancellationToken.None);
 
 		WorkspaceMaterializerSettings settings = await _workspace.GetSettingsAsync(CancellationToken.None);
@@ -103,7 +108,9 @@ public sealed class WorkspaceExecutorTest : IDisposable
 	{
 		_server.AddJob("jobPreflight", StreamId, 1, 1000);		
 		_workspace.SetFile(1000, "New/Feature/Foo.cs", "foo");
-		WorkspaceExecutor executor = new (_session, "jobPreflight", "batch1", AgentType, null, _workspace, null!, NullLogger.Instance);
+
+		JobExecutorOptions executorOptions = new JobExecutorOptions(_session, null!, "jobPreflight", "batch1", AgentType, new JobOptions());
+		WorkspaceExecutor executor = new (executorOptions, null, _workspace, NullLogger.Instance);
 		
 		await executor.InitializeAsync(_logger, CancellationToken.None);
 		AssertWorkspaceFile(_workspace, "main.cpp", "main");
@@ -116,7 +123,8 @@ public sealed class WorkspaceExecutorTest : IDisposable
 	public async Task JobWithNoChange()
 	{
 		_server.AddJob("jobNoChange", StreamId, 0, 0);		
-		WorkspaceExecutor executor = new (_session, "jobNoChange", "batch1", AgentType, null, _workspace, null!, NullLogger.Instance);
+		JobExecutorOptions executorOptions = new JobExecutorOptions(_session, null!, "jobNoChange", "batch1", AgentType, new JobOptions());
+		WorkspaceExecutor executor = new (executorOptions, null, _workspace, NullLogger.Instance);
 		await Assert.ThrowsExceptionAsync<WorkspaceMaterializationException>(() => executor.InitializeAsync(_logger, CancellationToken.None));
 	}
 
diff --git a/Engine/Source/Programs/Horde/Horde.Agent.Tests/Helper.cs b/Engine/Source/Programs/Horde/Horde.Agent.Tests/Helper.cs
--- a/Engine/Source/Programs/Horde/Horde.Agent.Tests/Helper.cs
+++ b/Engine/Source/Programs/Horde/Horde.Agent.Tests/Helper.cs
@@ -12,6 +12,7 @@ using Horde.Agent.Services;
 using Horde.Agent.Utility;
 using HordeCommon;
 using HordeCommon.Rpc;
+using HordeCommon.Rpc.Messages;
 using HordeCommon.Rpc.Tasks;
 using Microsoft.Extensions.Logging;
 using Microsoft.Extensions.Logging.Abstractions;
@@ -181,59 +182,47 @@ namespace Horde.Agent.Tests
 		}
 	}
 
-	class SimpleTestExecutor : JobExecutor
+	class SimpleTestExecutor : IJobExecutor
 	{
 		public const string Name = "Simple";
 
 		private readonly Func<BeginStepResponse, ILogger, CancellationToken, Task<JobStepOutcome>> _func;
 
 		public SimpleTestExecutor(Func<BeginStepResponse, ILogger, CancellationToken, Task<JobStepOutcome>> func)
-			: base(null!, null!, null!, null!, null!, NullLogger.Instance)
 		{
 			_func = func;
 		}
 
-		public override Task InitializeAsync(ILogger logger, CancellationToken cancellationToken)
+		public Task InitializeAsync(ILogger logger, CancellationToken cancellationToken)
 		{
 			logger.LogDebug("SimpleTestExecutor.InitializeAsync()");
 			return Task.CompletedTask;
 		}
 
-		public override Task<JobStepOutcome> RunAsync(BeginStepResponse step, ILogger logger,
-			CancellationToken cancellationToken)
+		public Task<JobStepOutcome> RunAsync(BeginStepResponse step, ILogger logger, CancellationToken cancellationToken)
 		{
 			logger.LogDebug("SimpleTestExecutor.RunAsync(Step: {Step})", step);
 			return _func(step, logger, cancellationToken);
 		}
 
-		public override Task FinalizeAsync(ILogger logger, CancellationToken cancellationToken)
+		public Task FinalizeAsync(ILogger logger, CancellationToken cancellationToken)
 		{
 			logger.LogDebug("SimpleTestExecutor.FinalizeAsync()");
 			return Task.CompletedTask;
 		}
-
-		protected override Task<bool> SetupAsync(BeginStepResponse step, ILogger logger, CancellationToken cancellationToken)
-		{
-			throw new NotImplementedException();
-		}
-
-		protected override Task<bool> ExecuteAsync(BeginStepResponse step, ILogger logger, CancellationToken cancellationToken)
-		{
-			throw new NotImplementedException();
-		}
 	}
 
-	class SimpleTestExecutorFactory : JobExecutorFactory
+	class SimpleTestExecutorFactory : IJobExecutorFactory
 	{
-		readonly JobExecutor _executor;
+		readonly IJobExecutor _executor;
 
-		public override string Name => SimpleTestExecutor.Name;
+		public string Name => SimpleTestExecutor.Name;
 
-		public SimpleTestExecutorFactory(JobExecutor executor)
+		public SimpleTestExecutorFactory(IJobExecutor executor)
 		{
 			_executor = executor;
 		}
 
-		public override JobExecutor CreateExecutor(ISession session, ExecuteJobTask executeJobTask, BeginBatchResponse beginBatchResponse) => _executor;
+		public IJobExecutor CreateExecutor(AgentWorkspace? workspaceInfo, AgentWorkspace? autoSdkWorkspaceInfo, JobExecutorOptions options) => _executor;
 	}
 }
\ No newline at end of file
diff --git a/Engine/Source/Programs/Horde/Horde.Agent.Tests/WorkerServiceTest.cs b/Engine/Source/Programs/Horde/Horde.Agent.Tests/WorkerServiceTest.cs
--- a/Engine/Source/Programs/Horde/Horde.Agent.Tests/WorkerServiceTest.cs
+++ b/Engine/Source/Programs/Horde/Horde.Agent.Tests/WorkerServiceTest.cs
@@ -10,6 +10,7 @@ using System.Threading.Channels;
 using System.Threading.Tasks;
 using Amazon.EC2.Model;
 using EpicGames.Core;
+using EpicGames.Horde.Storage;
 using Google.Protobuf.WellKnownTypes;
 using Grpc.Core;
 using Grpc.Net.Client;
@@ -19,6 +20,7 @@ using Horde.Agent.Leases.Handlers;
 using Horde.Agent.Parser;
 using Horde.Agent.Services;
 using Horde.Agent.Utility;
+using Horde.Common;
 using HordeCommon;
 using HordeCommon.Rpc;
 using HordeCommon.Rpc.Messages;
@@ -54,10 +56,15 @@ namespace Horde.Agent.Tests
 
 		class FakeServerLoggerFactory : IServerLoggerFactory
 		{
-			public IServerLogger CreateLogger(ISession session, string logId, string? jobId, string? batchId, string? stepId, bool? warnings = null) => new FakeServerLogger();
+			public IServerLogger CreateLogger(ISession session, string logId, string? jobId, string? batchId, string? stepId, bool? warnings = null, bool? useNewLogger = null) => new FakeServerLogger();
 		}
 
-		internal static JobExecutor NullExecutor = new SimpleTestExecutor(async (step, logger, cancellationToken) =>
+		class FakeSessionStorageFactory : IServerStorageFactory
+		{
+			public IStorageClient CreateStorageClient(ISession session) => null!;
+		}
+
+		internal static IJobExecutor NullExecutor = new SimpleTestExecutor(async (step, logger, cancellationToken) =>
 		{
 			await Task.Delay(1, cancellationToken);
 			return JobStepOutcome.Success;
@@ -68,6 +75,7 @@ namespace Horde.Agent.Tests
 			_serviceCollection = new ServiceCollection();
 			_serviceCollection.AddLogging();
 			_serviceCollection.AddSingleton<IServerLoggerFactory, FakeServerLoggerFactory>();
+			_serviceCollection.AddSingleton<IServerStorageFactory, FakeSessionStorageFactory>();
 
 			_serviceCollection.Configure<AgentSettings>(settings =>
 			{
@@ -96,7 +104,7 @@ namespace Horde.Agent.Tests
 				using CancellationTokenSource cancelSource = new CancellationTokenSource();
 				using CancellationTokenSource stepCancelSource = new CancellationTokenSource();
 
-				JobExecutor executor = new SimpleTestExecutor(async (stepResponse, logger, cancelToken) =>
+				IJobExecutor executor = new SimpleTestExecutor(async (stepResponse, logger, cancelToken) =>
 				{
 					cancelSource.CancelAfter(10);
 					await Task.Delay(5000, cancelToken);
@@ -111,7 +119,7 @@ namespace Horde.Agent.Tests
 				using CancellationTokenSource cancelSource = new CancellationTokenSource();
 				using CancellationTokenSource stepCancelSource = new CancellationTokenSource();
 
-				JobExecutor executor = new SimpleTestExecutor(async (stepResponse, logger, cancelToken) =>
+				IJobExecutor executor = new SimpleTestExecutor(async (stepResponse, logger, cancelToken) =>
 				{
 					stepCancelSource.CancelAfter(10);
 					await Task.Delay(5000, cancelToken);
@@ -135,7 +143,7 @@ namespace Horde.Agent.Tests
 			executeJobTask.BatchId = "batchId1";
 			executeJobTask.LogId = "logId1";
 			executeJobTask.JobName = "jobName1";
-			executeJobTask.Executor = SimpleTestExecutor.Name;
+			executeJobTask.JobOptions = new JobOptions { Executor = SimpleTestExecutor.Name };
 			executeJobTask.AutoSdkWorkspace = new AgentWorkspace();
 			executeJobTask.Workspace = new AgentWorkspace();
 
@@ -152,13 +160,13 @@ namespace Horde.Agent.Tests
 			GetStepResponse step2Res = new GetStepResponse(JobStepOutcome.Unspecified, JobStepState.Unspecified, true);
 			client.GetStepResponses[step2Req] = step2Res;
 
-			JobExecutor executor = new SimpleTestExecutor(async (step, logger, cancelToken) =>
+			SimpleTestExecutor executor = new SimpleTestExecutor(async (step, logger, cancelToken) =>
 			{
 				await Task.Delay(50, cancelToken);
 				return JobStepOutcome.Success;
 			});
 
-			_serviceCollection.AddSingleton<JobExecutorFactory>(x => new SimpleTestExecutorFactory(executor));
+			_serviceCollection.AddSingleton<IJobExecutorFactory>(x => new SimpleTestExecutorFactory(executor));
 			using ServiceProvider serviceProvider = _serviceCollection.BuildServiceProvider();
 
 			JobHandler jobHandler = serviceProvider.GetRequiredService<JobHandler>();
@@ -180,13 +188,13 @@ namespace Horde.Agent.Tests
 		[TestMethod]
 		public async Task PollForStepAbortFailureTest()
 		{
-			JobExecutor executor = new SimpleTestExecutor(async (step, logger, cancelToken) =>
+			IJobExecutor executor = new SimpleTestExecutor(async (step, logger, cancelToken) =>
 			{
 				await Task.Delay(50, cancelToken);
 				return JobStepOutcome.Success;
 			});
 
-			_serviceCollection.AddSingleton<JobExecutorFactory>(x => new SimpleTestExecutorFactory(executor));
+			_serviceCollection.AddSingleton<IJobExecutorFactory>(x => new SimpleTestExecutorFactory(executor));
 			using ServiceProvider serviceProvider = _serviceCollection.BuildServiceProvider();
 
 			JobHandler jobHandler = serviceProvider.GetRequiredService<JobHandler>();
@@ -218,13 +226,13 @@ namespace Horde.Agent.Tests
 		[TestMethod]
 		public async Task Shutdown()
 		{
-			JobExecutor executor = new SimpleTestExecutor(async (step, logger, cancellationToken) =>
+			IJobExecutor executor = new SimpleTestExecutor(async (step, logger, cancellationToken) =>
 			{
 				await Task.Delay(50, cancellationToken);
 				return JobStepOutcome.Success;
 			});
 
-			_serviceCollection.AddSingleton<JobExecutorFactory>(x => new SimpleTestExecutorFactory(executor));
+			_serviceCollection.AddSingleton<IJobExecutorFactory>(x => new SimpleTestExecutorFactory(executor));
 			using ServiceProvider serviceProvider = _serviceCollection.BuildServiceProvider();
 
 			using CancellationTokenSource cts = new ();
@@ -242,7 +250,7 @@ namespace Horde.Agent.Tests
 		}
 	}
 
-	internal class FakeServerSessionFactory : ISessionFactoryService
+	internal class FakeServerSessionFactory : ISessionFactory
 	{
 		readonly FakeHordeRpcServer _fakeServer;
 
diff --git a/Engine/Source/Programs/Horde/Horde.Agent/AgentSettings.cs b/Engine/Source/Programs/Horde/Horde.Agent/AgentSettings.cs
--- a/Engine/Source/Programs/Horde/Horde.Agent/AgentSettings.cs
+++ b/Engine/Source/Programs/Horde/Horde.Agent/AgentSettings.cs
@@ -186,11 +186,6 @@ namespace Horde.Agent
 		/// </summary>
 		public bool EnableAwsEc2Support { get; set; } = true;
 
-		/// <summary>
-		/// Whether to enable the new storage-API based log system 
-		/// </summary>
-		public bool EnableNewLogger { get; set; }
-
 		/// <summary>
 		/// Option to use a local storage client rather than connecting through the server. Primarily for convenience when debugging / iterating locally.
 		/// </summary>
diff --git a/Engine/Source/Programs/Horde/Horde.Agent/Commands/Compute/AwsLambdaListenerCommand.cs b/Engine/Source/Programs/Horde/Horde.Agent/Commands/Compute/AwsLambdaListenerCommand.cs
--- a/Engine/Source/Programs/Horde/Horde.Agent/Commands/Compute/AwsLambdaListenerCommand.cs
+++ b/Engine/Source/Programs/Horde/Horde.Agent/Commands/Compute/AwsLambdaListenerCommand.cs
@@ -53,7 +53,7 @@ namespace Horde.Agent.Commands.Compute
 
 			TimeSpan maxWaitTime = TimeSpan.FromMilliseconds(request.MaxWaitTimeForLeaseMs);
 
-			await using ISession session = await _serviceProvider.GetRequiredService<ISessionFactoryService>().CreateAsync(cancellationToken);
+			await using ISession session = await _serviceProvider.GetRequiredService<ISessionFactory>().CreateAsync(cancellationToken);
 
 			LeaseManager manager = new LeaseManager(session, _serviceProvider);
 			bool isLeaseActive = false;
diff --git a/Engine/Source/Programs/Horde/Horde.Agent/Execution/JobExecutor.cs b/Engine/Source/Programs/Horde/Horde.Agent/Execution/JobExecutor.cs
--- a/Engine/Source/Programs/Horde/Horde.Agent/Execution/JobExecutor.cs
+++ b/Engine/Source/Programs/Horde/Horde.Agent/Execution/JobExecutor.cs
@@ -22,19 +22,47 @@ using Grpc.Core;
 using Horde.Agent.Parser;
 using Horde.Agent.Services;
 using Horde.Agent.Utility;
+using Horde.Common;
 using Horde.Storage.Utility;
 using HordeCommon;
 using HordeCommon.Rpc;
+using HordeCommon.Rpc.Messages;
 using HordeCommon.Rpc.Tasks;
 using Microsoft.Extensions.Caching.Memory;
 using Microsoft.Extensions.Logging;
-using Microsoft.Extensions.Options;
 using OpenTracing;
 using OpenTracing.Util;
 
 namespace Horde.Agent.Execution
 {
-	abstract class JobExecutor
+	interface IJobExecutor
+	{
+		Task InitializeAsync(ILogger logger, CancellationToken cancellationToken);
+		Task<JobStepOutcome> RunAsync(BeginStepResponse step, ILogger logger, CancellationToken cancellationToken);
+		Task FinalizeAsync(ILogger logger, CancellationToken cancellationToken);
+	}
+
+	class JobExecutorOptions
+	{
+		public ISession Session { get; }
+		public IStorageClient Storage { get; }
+		public string JobId { get; }
+		public string BatchId { get; }
+		public string AgentType { get; }
+		public JobOptions JobOptions { get; }
+
+		public JobExecutorOptions(ISession session, IStorageClient storage, string jobId, string batchId, string agentType, JobOptions jobOptions)
+		{
+			Session = session;
+			Storage = storage;
+			JobId = jobId;
+			BatchId = batchId;
+			AgentType = agentType;
+			JobOptions = jobOptions;
+		}
+	}
+
+	abstract class JobExecutor : IJobExecutor
 	{
 		protected class ExportedNode
 		{
@@ -149,12 +177,10 @@ namespace Horde.Agent.Execution
 		protected string _batchId;
 		protected string _agentTypeName;
 
-		protected IHttpClientFactory _httpClientFactory;
-		
 		/// <summary>
 		/// Logger for the local agent process (as opposed to job logger)
 		/// </summary>
-		readonly protected ILogger _logger;
+		protected readonly ILogger _logger;
 
 		protected GetJobResponse _job;
 		protected GetStreamResponse _stream;
@@ -165,22 +191,25 @@ namespace Horde.Agent.Execution
 		protected bool _compileAutomationTool = true;
 
 		protected readonly ISession _session;
+		protected readonly IStorageClient _storage;
+		protected readonly JobOptions _jobOptions;
+
 		protected IRpcConnection RpcConnection => _session.RpcConnection;
 		protected Dictionary<string, string> _remapAgentTypes = new Dictionary<string, string>(StringComparer.OrdinalIgnoreCase);
 
 		protected Dictionary<string, string> _envVars = new Dictionary<string, string>(StringComparer.OrdinalIgnoreCase);
 
-		public bool UseNewLogger { get; private set; }
-
-		public JobExecutor(ISession session, string jobId, string batchId, string agentTypeName, IHttpClientFactory httpClientFactory, ILogger logger)
+		public JobExecutor(JobExecutorOptions options, ILogger logger)
 		{
-			_session = session;
+			_session = options.Session;
+			_storage = options.Storage;
+
+			_jobId = options.JobId;
+			_batchId = options.BatchId;
+			_agentTypeName = options.AgentType;
 
-			_jobId = jobId;
-			_batchId = batchId;
-			_agentTypeName = agentTypeName;
+			_jobOptions = options.JobOptions;
 
-			_httpClientFactory = httpClientFactory;
 			_logger = logger;
 
 			_job = null!;
@@ -246,10 +275,6 @@ namespace Horde.Agent.Execution
 				{
 					_preprocessScript = true;
 				}
-				else if (argument.Equals("-UseNewLogger", StringComparison.OrdinalIgnoreCase))
-				{
-					UseNewLogger = true;
-				}
 				else if (argument.StartsWith(TargetArgumentPrefix, StringComparison.OrdinalIgnoreCase))
 				{
 					_targets.Add(argument.Substring(TargetArgumentPrefix.Length));
@@ -652,22 +677,17 @@ namespace Horde.Agent.Execution
 			arguments.AppendArgument("-SingleNode=", step.Name);
 //			Arguments.AppendArgument("-TokenSignature=", JobId.ToString());
 
-			bool manageSharedStorage = false;
 			foreach (string additionalArgument in _additionalArguments)
 			{
-				if (additionalArgument.Equals("-ManageSharedStorage", StringComparison.OrdinalIgnoreCase))
-				{
-					manageSharedStorage = true;
-				}
-				else if (!_preprocessScript || !additionalArgument.StartsWith("-set:", StringComparison.OrdinalIgnoreCase))
+				if (!_preprocessScript || !additionalArgument.StartsWith("-set:", StringComparison.OrdinalIgnoreCase))
 				{
 					arguments.AppendArgument(additionalArgument);
 				}
 			}
 
-			if (manageSharedStorage && sharedStorageDir != null)
+			if (_jobOptions.UseNewTempStorage ?? false)
 			{
-				return await ExecuteWithTempStorageAsync(step, workspaceDir, sharedStorageDir, arguments.ToString(), logger, cancellationToken);
+				return await ExecuteWithTempStorageAsync(step, workspaceDir, arguments.ToString(), logger, cancellationToken);
 			}
 			else
 			{
@@ -679,53 +699,12 @@ namespace Horde.Agent.Execution
 			}
 		}
 
-		private async Task<bool> ExecuteWithTempStorageAsync(BeginStepResponse step, DirectoryReference workspaceDir, DirectoryReference sharedStorageDir, string arguments, ILogger logger, CancellationToken cancellationToken)
+		private async Task<bool> ExecuteWithTempStorageAsync(BeginStepResponse step, DirectoryReference workspaceDir, string arguments, ILogger logger, CancellationToken cancellationToken)
 		{
 			DirectoryReference manifestDir = DirectoryReference.Combine(workspaceDir, "Engine", "Saved", "BuildGraph");
 
-			// Create the storage client
 			using MemoryCache cache = new MemoryCache(new MemoryCacheOptions { });
-
-			using HttpClient httpClient = _httpClientFactory.CreateClient(HttpStorageClient.HttpClientName);
-			httpClient.BaseAddress = new Uri($"{_session.ServerUrl}/api/v1/storage/default");
-			httpClient.DefaultRequestHeaders.Authorization = new AuthenticationHeaderValue("Bearer", _session.Token);
-
-			using HttpClient redirectHttpClient = _httpClientFactory.CreateClient(HttpStorageClient.HttpClientName);
-
-			HttpStorageClient storage = new HttpStorageClient(CreateHttpClient, CreateHttpRedirectClient, logger);
-			//			FileStorageClient storage = new FileStorageClient(DirectoryReference.Combine(sharedStorageDir, "bundles"), cache, logger);
-			TreeReader reader = new TreeReader(storage, cache, logger);
-			logger.LogInformation("Using Horde-managed shared storage via {SharedStorageDir}", sharedStorageDir);
-
-			ILogger nextLogger = logger;
-			for(; ;)
-			{
-				if (nextLogger is JsonRpcLogger)
-				{
-					break;
-				}
-				else if (nextLogger is PerforceLogger perforceLogger)
-				{
-					nextLogger = perforceLogger.Inner;
-				}
-				else if (nextLogger is ForwardingLogger forwardingLogger)
-				{
-					nextLogger = forwardingLogger.Loggers[0];
-				}
-				else if (nextLogger is DefaultLoggerIndentHandler indentLogger)
-				{
-					nextLogger = indentLogger.Inner;
-				}
-				else
-				{
-					throw new Exception($"Unknown logger type: {nextLogger.GetType().Name}");
-				}
-			}
-
-			JsonRpcLogger rpcLogger = (JsonRpcLogger)nextLogger;
-			await using JsonRpcAndStorageLogSink sink = new JsonRpcAndStorageLogSink(RpcConnection, rpcLogger._logId, rpcLogger._sink, storage, logger);
-			await using JsonRpcLogger newRpcLogger = new JsonRpcLogger(sink, rpcLogger._logId, rpcLogger._warnings, rpcLogger._inner);
-			logger = newRpcLogger;
+			TreeReader reader = new TreeReader(_storage, cache, _logger);
 
 			// Create the mapping of tag names to file sets
 			Dictionary<string, HashSet<FileReference>> tagNameToFileSet = new Dictionary<string, HashSet<FileReference>>();
@@ -781,7 +760,7 @@ namespace Horde.Agent.Execution
 			}
 
 			// Run UAT
-			if (await ExecuteAutomationToolAsync(step, workspaceDir, arguments, storage, logger, cancellationToken) != 0)
+			if (await ExecuteAutomationToolAsync(step, workspaceDir, arguments, _storage, logger, cancellationToken) != 0)
 			{
 				return false;
 			}
@@ -833,7 +812,7 @@ namespace Horde.Agent.Execution
 
 			// Find a block name for all new outputs
 			Dictionary<FileReference, string> fileToBlockName = new Dictionary<FileReference, string>();
-			for(int idx = 0; idx < step.OutputNames.Count; idx++)
+			for (int idx = 0; idx < step.OutputNames.Count; idx++)
 			{
 				string tagName = step.OutputNames[idx];
 
@@ -888,7 +867,7 @@ namespace Horde.Agent.Execution
 				RefName refName = TempStorage.GetRefNameForNode(RefPrefix, step.Name);
 
 				TreeOptions treeOptions = new TreeOptions();
-				using TreeWriter treeWriter = new TreeWriter(storage, treeOptions, refName.Text);
+				using TreeWriter treeWriter = new TreeWriter(_storage, treeOptions, refName.Text);
 
 				TempStorageNode outputNode = new TempStorageNode();
 
@@ -932,19 +911,6 @@ namespace Horde.Agent.Execution
 			return true;
 		}
 
-		HttpClient CreateHttpClient()
-		{
-			HttpClient httpClient = _httpClientFactory.CreateClient(HttpStorageClient.HttpClientName);
-			httpClient.BaseAddress = new Uri($"{_session.ServerUrl}/api/v1/storage/default");
-			httpClient.DefaultRequestHeaders.Authorization = new AuthenticationHeaderValue("Bearer", _session.Token);
-			return httpClient;
-		}
-
-		HttpClient CreateHttpRedirectClient()
-		{
-			return _httpClientFactory.CreateClient(HttpStorageClient.HttpClientName);
-		}
-
 		protected async Task<int> ExecuteAutomationToolAsync(BeginStepResponse step, DirectoryReference workspaceDir, string arguments, IStorageClient? store, ILogger logger, CancellationToken cancellationToken)
 		{
 			int result;
@@ -1490,10 +1456,10 @@ namespace Horde.Agent.Execution
 		}
 	}
 
-	abstract class JobExecutorFactory
+	interface IJobExecutorFactory
 	{
-		public abstract string Name { get; }
+		string Name { get; }
 
-		public abstract JobExecutor CreateExecutor(ISession session, ExecuteJobTask executeJobTask, BeginBatchResponse beginBatchResponse);
+		IJobExecutor CreateExecutor(AgentWorkspace workspaceInfo, AgentWorkspace? autoSdkWorkspaceInfo, JobExecutorOptions options);
 	}
 }
diff --git a/Engine/Source/Programs/Horde/Horde.Agent/Execution/LocalExecutor.cs b/Engine/Source/Programs/Horde/Horde.Agent/Execution/LocalExecutor.cs
--- a/Engine/Source/Programs/Horde/Horde.Agent/Execution/LocalExecutor.cs
+++ b/Engine/Source/Programs/Horde/Horde.Agent/Execution/LocalExecutor.cs
@@ -6,9 +6,11 @@ using System.Reflection;
 using System.Threading;
 using System.Threading.Tasks;
 using EpicGames.Core;
+using EpicGames.Horde.Storage;
 using Horde.Agent.Services;
 using Horde.Agent.Utility;
 using HordeCommon.Rpc;
+using HordeCommon.Rpc.Messages;
 using HordeCommon.Rpc.Tasks;
 using Microsoft.Extensions.Logging;
 using Microsoft.Extensions.Options;
@@ -20,8 +22,8 @@ namespace Horde.Agent.Execution
 		private readonly LocalExecutorSettings _settings;
 		private readonly DirectoryReference _localWorkspaceDir;
 
-		public LocalExecutor(ISession session, string jobId, string batchId, string agentTypeName, LocalExecutorSettings settings, IHttpClientFactory httpClientFactory, ILogger logger) 
-			: base(session, jobId, batchId, agentTypeName, httpClientFactory, logger)
+		public LocalExecutor(JobExecutorOptions options, LocalExecutorSettings settings, ILogger logger) 
+			: base(options, logger)
 		{
 			_settings = settings;
 			if (settings.WorkspaceDir == null)
@@ -70,24 +72,22 @@ namespace Horde.Agent.Execution
 		}
 	}
 
-	class LocalExecutorFactory : JobExecutorFactory
+	class LocalExecutorFactory : IJobExecutorFactory
 	{
 		readonly LocalExecutorSettings _settings;
-		readonly IHttpClientFactory _httpClientFactory;
 		readonly ILogger<LocalExecutor> _logger;
 
-		public override string Name => "Local";
+		public string Name => "Local";
 
-		public LocalExecutorFactory(IOptions<LocalExecutorSettings> settings, IHttpClientFactory httpClientFactory, ILogger<LocalExecutor> logger)
+		public LocalExecutorFactory(IOptions<LocalExecutorSettings> settings, ILogger<LocalExecutor> logger)
 		{
 			_settings = settings.Value;
-			_httpClientFactory = httpClientFactory;
 			_logger = logger;
 		}
 
-		public override JobExecutor CreateExecutor(ISession session, ExecuteJobTask executeJobTask, BeginBatchResponse beginBatchResponse)
+		public IJobExecutor CreateExecutor(AgentWorkspace workspaceInfo, AgentWorkspace? autoSdkWorkspaceInfo, JobExecutorOptions options)
 		{
-			return new LocalExecutor(session, executeJobTask.JobId, executeJobTask.BatchId, beginBatchResponse.AgentType, _settings, _httpClientFactory, _logger);
+			return new LocalExecutor(options, _settings, _logger);
 		}
 	}
 }
diff --git a/Engine/Source/Programs/Horde/Horde.Agent/Execution/PerforceExecutor.cs b/Engine/Source/Programs/Horde/Horde.Agent/Execution/PerforceExecutor.cs
--- a/Engine/Source/Programs/Horde/Horde.Agent/Execution/PerforceExecutor.cs
+++ b/Engine/Source/Programs/Horde/Horde.Agent/Execution/PerforceExecutor.cs
@@ -12,13 +12,13 @@ using System.Threading;
 using System.Threading.Tasks;
 using System.Web;
 using EpicGames.Core;
+using EpicGames.Horde.Storage;
 using EpicGames.Perforce;
 using EpicGames.Perforce.Managed;
 using Horde.Agent.Services;
 using Horde.Agent.Utility;
 using HordeCommon.Rpc;
 using HordeCommon.Rpc.Messages;
-using HordeCommon.Rpc.Tasks;
 using Microsoft.Extensions.Logging;
 using OpenTracing;
 using OpenTracing.Util;
@@ -29,20 +29,20 @@ namespace Horde.Agent.Execution
 	{
 		public const string Name = "Perforce";
 
-		protected AgentWorkspace? _autoSdkWorkspaceInfo;
 		protected AgentWorkspace _workspaceInfo;
+		protected AgentWorkspace? _autoSdkWorkspaceInfo;
 		protected DirectoryReference _rootDir;
 		protected DirectoryReference? _sharedStorageDir;
 
 		protected WorkspaceInfo? _autoSdkWorkspace;
 		protected WorkspaceInfo _workspace;
 
-		public PerforceExecutor(ISession session, string jobId, string batchId, string agentTypeName, AgentWorkspace? autoSdkWorkspaceInfo, AgentWorkspace workspaceInfo, DirectoryReference rootDir, IHttpClientFactory httpClientFactory, ILogger logger)
-			: base(session, jobId, batchId, agentTypeName, httpClientFactory, logger)
+		public PerforceExecutor(AgentWorkspace workspaceInfo, AgentWorkspace? autoSdkWorkspaceInfo, JobExecutorOptions options, ILogger logger)
+			: base(options, logger)
 		{
-			_autoSdkWorkspaceInfo = autoSdkWorkspaceInfo;
 			_workspaceInfo = workspaceInfo;
-			_rootDir = rootDir;
+			_autoSdkWorkspaceInfo = autoSdkWorkspaceInfo;
+			_rootDir = options.Session.WorkingDir;
 
 			_workspace = null!;
 		}
@@ -433,22 +433,20 @@ namespace Horde.Agent.Execution
 		}
 	}
 
-	class PerforceExecutorFactory : JobExecutorFactory
+	class PerforceExecutorFactory : IJobExecutorFactory
 	{
-		readonly IHttpClientFactory _httpClientFactory;
 		readonly ILogger<PerforceExecutor> _logger;
 
-		public override string Name => PerforceExecutor.Name;
+		public string Name => PerforceExecutor.Name;
 
-		public PerforceExecutorFactory(IHttpClientFactory httpClientFactory, ILogger<PerforceExecutor> logger)
+		public PerforceExecutorFactory(ILogger<PerforceExecutor> logger)
 		{
-			_httpClientFactory = httpClientFactory;
 			_logger = logger;
 		}
 
-		public override JobExecutor CreateExecutor(ISession session, ExecuteJobTask executeJobTask, BeginBatchResponse beginBatchResponse)
+		public IJobExecutor CreateExecutor(AgentWorkspace workspaceInfo, AgentWorkspace? autoSdkWorkspaceInfo, JobExecutorOptions options)
 		{
-			return new PerforceExecutor(session, executeJobTask.JobId, executeJobTask.BatchId, beginBatchResponse.AgentType, executeJobTask.AutoSdkWorkspace, executeJobTask.Workspace, session.WorkingDir, _httpClientFactory, _logger);
+			return new PerforceExecutor(workspaceInfo, autoSdkWorkspaceInfo, options, _logger);
 		}
 	}
 }
diff --git a/Engine/Source/Programs/Horde/Horde.Agent/Execution/TestExecutor.cs b/Engine/Source/Programs/Horde/Horde.Agent/Execution/TestExecutor.cs
--- a/Engine/Source/Programs/Horde/Horde.Agent/Execution/TestExecutor.cs
+++ b/Engine/Source/Programs/Horde/Horde.Agent/Execution/TestExecutor.cs
@@ -8,11 +8,13 @@ using System.Reflection;
 using System.Threading;
 using System.Threading.Tasks;
 using EpicGames.Core;
+using EpicGames.Horde.Storage;
 using Horde.Agent.Parser;
 using Horde.Agent.Services;
 using Horde.Agent.Utility;
 using HordeCommon;
 using HordeCommon.Rpc;
+using HordeCommon.Rpc.Messages;
 using HordeCommon.Rpc.Tasks;
 using Microsoft.Extensions.Logging;
 
@@ -22,8 +24,8 @@ namespace Horde.Agent.Execution
 	{
 		public const string Name = "Test";
 
-		public TestExecutor(ISession session, string jobId, string batchId, string agentTypeName, IHttpClientFactory httpClientFactory, ILogger logger)
-			: base(session, jobId, batchId, agentTypeName, httpClientFactory, logger)
+		public TestExecutor(JobExecutorOptions options, ILogger logger)
+			: base(options, logger)
 		{
 		}
 
@@ -190,22 +192,20 @@ namespace Horde.Agent.Execution
 		}
 	}
 
-	class TestExecutorFactory : JobExecutorFactory
+	class TestExecutorFactory : IJobExecutorFactory
 	{
-		readonly IHttpClientFactory _httpClientFactory;
 		readonly ILogger<TestExecutor> _logger;
 
-		public override string Name => TestExecutor.Name;
+		public string Name => TestExecutor.Name;
 
-		public TestExecutorFactory(IHttpClientFactory httpClientFactory, ILogger<TestExecutor> logger)
+		public TestExecutorFactory(ILogger<TestExecutor> logger)
 		{
-			_httpClientFactory = httpClientFactory;
 			_logger = logger;
 		}
 
-		public override JobExecutor CreateExecutor(ISession session, ExecuteJobTask executeJobTask, BeginBatchResponse beginBatchResponse)
+		public IJobExecutor CreateExecutor(AgentWorkspace workspaceInfo, AgentWorkspace? autoSdkWorkspaceInfo, JobExecutorOptions options)
 		{
-			return new TestExecutor(session, executeJobTask.JobId, executeJobTask.BatchId, beginBatchResponse.AgentType, _httpClientFactory, _logger);
+			return new TestExecutor(options, _logger);
 		}
 	}
 }
diff --git a/Engine/Source/Programs/Horde/Horde.Agent/Execution/WorkspaceExecutor.cs b/Engine/Source/Programs/Horde/Horde.Agent/Execution/WorkspaceExecutor.cs
--- a/Engine/Source/Programs/Horde/Horde.Agent/Execution/WorkspaceExecutor.cs
+++ b/Engine/Source/Programs/Horde/Horde.Agent/Execution/WorkspaceExecutor.cs
@@ -1,14 +1,12 @@
 // Copyright Epic Games, Inc. All Rights Reserved.
 
 using System;
-using System.Net.Http;
 using System.Text.RegularExpressions;
 using System.Threading;
 using System.Threading.Tasks;
 using EpicGames.Core;
-using Horde.Agent.Services;
 using HordeCommon.Rpc;
-using HordeCommon.Rpc.Tasks;
+using HordeCommon.Rpc.Messages;
 using Microsoft.Extensions.Logging;
 using OpenTracing;
 using OpenTracing.Util;
@@ -23,8 +21,8 @@ namespace Horde.Agent.Execution
 		private readonly IWorkspaceMaterializer? _autoSdkWorkspace;
 		private readonly IWorkspaceMaterializer _workspace;
 
-		public WorkspaceExecutor(ISession session, string jobId, string batchId, string agentTypeName, IWorkspaceMaterializer? autoSdkWorkspace, IWorkspaceMaterializer workspace, IHttpClientFactory httpClientFactory, ILogger logger)
-			: base(session, jobId, batchId, agentTypeName, httpClientFactory, logger)
+		public WorkspaceExecutor(JobExecutorOptions options, IWorkspaceMaterializer? autoSdkWorkspace, IWorkspaceMaterializer workspace, ILogger logger)
+			: base(options, logger)
 		{
 			_autoSdkWorkspace = autoSdkWorkspace;
 			_workspace = workspace;
@@ -129,18 +127,15 @@ namespace Horde.Agent.Execution
 		}
 	}
 
-	class WorkspaceExecutorFactory : JobExecutorFactory
+	class WorkspaceExecutorFactory : IJobExecutorFactory
 	{
-		readonly IHttpClientFactory _httpClientFactory;
+		public string Name => WorkspaceExecutor.Name;
 
-		public override string Name => WorkspaceExecutor.Name;
-
-		public WorkspaceExecutorFactory(IHttpClientFactory httpClientFactory)
+		public WorkspaceExecutorFactory()
 		{
-			_httpClientFactory = httpClientFactory;
 		}
 
-		public override JobExecutor CreateExecutor(ISession session, ExecuteJobTask executeJobTask, BeginBatchResponse beginBatchResponse)
+		public IJobExecutor CreateExecutor(AgentWorkspace workspaceInfo, AgentWorkspace? autoSdkWorkspaceInfo, JobExecutorOptions options)
 		{
 			throw new NotImplementedException("WorkspaceExecutor cannot be instantiated as there are no IWorkspaceMaterializer implementations yet");
 		}
diff --git a/Engine/Source/Programs/Horde/Horde.Agent/Leases/Handlers/ConformHandler.cs b/Engine/Source/Programs/Horde/Horde.Agent/Leases/Handlers/ConformHandler.cs
--- a/Engine/Source/Programs/Horde/Horde.Agent/Leases/Handlers/ConformHandler.cs
+++ b/Engine/Source/Programs/Horde/Horde.Agent/Leases/Handlers/ConformHandler.cs
@@ -34,7 +34,7 @@ namespace Horde.Agent.Leases.Handlers
 		/// <inheritdoc/>
 		public override async Task<LeaseResult> ExecuteAsync(ISession session, string leaseId, ConformTask conformTask, CancellationToken cancellationToken)
 		{
-			await using IServerLogger conformLogger = _serverLoggerFactory.CreateLogger(session, conformTask.LogId, null);
+			await using IServerLogger conformLogger = _serverLoggerFactory.CreateLogger(session, conformTask.LogId, null, null);
 
 			conformLogger.LogInformation("Conforming, lease {LeaseId}", leaseId);
 			await session.TerminateProcessesAsync(conformLogger, cancellationToken);
diff --git a/Engine/Source/Programs/Horde/Horde.Agent/Leases/Handlers/JobHandler.cs b/Engine/Source/Programs/Horde/Horde.Agent/Leases/Handlers/JobHandler.cs
--- a/Engine/Source/Programs/Horde/Horde.Agent/Leases/Handlers/JobHandler.cs
+++ b/Engine/Source/Programs/Horde/Horde.Agent/Leases/Handlers/JobHandler.cs
@@ -10,13 +10,16 @@ using System.Runtime.InteropServices;
 using System.Threading;
 using System.Threading.Tasks;
 using EpicGames.Core;
+using EpicGames.Horde.Storage;
 using Grpc.Core;
 using Horde.Agent.Execution;
 using Horde.Agent.Parser;
 using Horde.Agent.Services;
 using Horde.Agent.Utility;
+using Horde.Common;
 using HordeCommon;
 using HordeCommon.Rpc;
+using HordeCommon.Rpc.Messages;
 using HordeCommon.Rpc.Tasks;
 using Microsoft.Extensions.Logging;
 using Microsoft.Extensions.Options;
@@ -33,18 +36,20 @@ namespace Horde.Agent.Leases.Handlers
 		/// </summary>
 		internal TimeSpan _stepAbortPollInterval = TimeSpan.FromSeconds(5);
 
-		readonly IEnumerable<JobExecutorFactory> _executorFactories;
+		readonly IEnumerable<IJobExecutorFactory> _executorFactories;
 		readonly AgentSettings _settings;
+		readonly IServerStorageFactory _storageClientFactory;
 		readonly IServerLoggerFactory _serverLoggerFactory;
 		readonly ILogger _defaultLogger;
 
 		/// <summary>
 		/// Constructor
 		/// </summary>
-		public JobHandler(IEnumerable<JobExecutorFactory> executorFactories, IOptions<AgentSettings> settings, IServerLoggerFactory serverLoggerFactory, ILogger<JobHandler> defaultLogger)
+		public JobHandler(IEnumerable<IJobExecutorFactory> executorFactories, IOptions<AgentSettings> settings, IServerStorageFactory storageClientFactory, IServerLoggerFactory serverLoggerFactory, ILogger<JobHandler> defaultLogger)
 		{
 			_executorFactories = executorFactories;
 			_settings = settings.Value;
+			_storageClientFactory = storageClientFactory;
 			_serverLoggerFactory = serverLoggerFactory;
 			_defaultLogger = defaultLogger;
 		}
@@ -52,7 +57,9 @@ namespace Horde.Agent.Leases.Handlers
 		/// <inheritdoc/>
 		public override async Task<LeaseResult> ExecuteAsync(ISession session, string leaseId, ExecuteJobTask executeTask, CancellationToken cancellationToken)
 		{
-			await using IServerLogger logger = _serverLoggerFactory.CreateLogger(session, executeTask.LogId, executeTask.JobId, executeTask.BatchId, null);
+			// Create a storage client for this session
+			JobOptions jobOptions = executeTask.JobOptions;
+			await using IServerLogger logger = _serverLoggerFactory.CreateLogger(session, executeTask.LogId, executeTask.JobId, executeTask.BatchId, null, null, jobOptions.UseNewLogStorage);
 
 			logger.LogInformation("Executing job \"{JobName}\", jobId {JobId}, batchId {BatchId}, leaseId {LeaseId}", executeTask.JobName, executeTask.JobId, executeTask.BatchId, leaseId);
 			GlobalTracer.Instance.ActiveSpan?.SetTag("jobId", executeTask.JobId.ToString());
@@ -65,7 +72,9 @@ namespace Horde.Agent.Leases.Handlers
 			// Execute the batch
 			try
 			{
-				await ExecuteBatchAsync(session, leaseId, executeTask, batch, logger, cancellationToken);
+				IStorageClient storage = _storageClientFactory.CreateStorageClient(session);
+				JobExecutorOptions options = new JobExecutorOptions(session, storage, executeTask.JobId, executeTask.BatchId, batch.AgentType, jobOptions);
+				await ExecuteBatchAsync(session, leaseId, executeTask.Workspace, executeTask.AutoSdkWorkspace, options, logger, cancellationToken);
 			}
 			catch (Exception ex)
 			{
@@ -96,24 +105,24 @@ namespace Horde.Agent.Leases.Handlers
 		/// <summary>
 		/// Executes a batch
 		/// </summary>
-		async Task ExecuteBatchAsync(ISession session, string leaseId, ExecuteJobTask executeTask, BeginBatchResponse batch, ILogger batchLogger, CancellationToken cancellationToken)
+		async Task ExecuteBatchAsync(ISession session, string leaseId, AgentWorkspace workspaceInfo, AgentWorkspace? autoSdkWorkspaceInfo, JobExecutorOptions options, ILogger batchLogger, CancellationToken cancellationToken)
 		{
 			IRpcConnection rpcClient = session.RpcConnection;
 
-			batchLogger.LogInformation("Executing batch {BatchId} using {Executor} executor", executeTask.BatchId, _settings.Executor.ToString());
+			batchLogger.LogInformation("Executing batch {BatchId} using {Executor} executor", options.BatchId, _settings.Executor.ToString());
 			await session.TerminateProcessesAsync(batchLogger, cancellationToken);
 
 			// Create an executor for this job
-			string executorName = String.IsNullOrEmpty(executeTask.Executor) ? _settings.Executor : executeTask.Executor;
+			string executorName = String.IsNullOrEmpty(options.JobOptions.Executor) ? _settings.Executor : options.JobOptions.Executor;
 
-			JobExecutorFactory? executorFactory = _executorFactories.FirstOrDefault(x => x.Name.Equals(executorName, StringComparison.OrdinalIgnoreCase));
+			IJobExecutorFactory? executorFactory = _executorFactories.FirstOrDefault(x => x.Name.Equals(executorName, StringComparison.OrdinalIgnoreCase));
 			if (executorFactory == null)
 			{
 				batchLogger.LogError("Unable to find executor '{ExecutorName}'", executorName);
 				return;
 			}
 
-			JobExecutor executor = executorFactory.CreateExecutor(session, executeTask, batch);
+			IJobExecutor executor = executorFactory.CreateExecutor(workspaceInfo, autoSdkWorkspaceInfo, options);
 
 			// Try to initialize the executor
 			batchLogger.LogInformation("Initializing...");
@@ -129,7 +138,7 @@ namespace Horde.Agent.Leases.Handlers
 				for (; ; )
 				{
 					// Get the next step to execute
-					BeginStepResponse step = await rpcClient.InvokeAsync((HordeRpc.HordeRpcClient x) => x.BeginStepAsync(new BeginStepRequest(executeTask.JobId, executeTask.BatchId, leaseId), null, null, cancellationToken), cancellationToken);
+					BeginStepResponse step = await rpcClient.InvokeAsync((HordeRpc.HordeRpcClient x) => x.BeginStepAsync(new BeginStepRequest(options.JobId, options.BatchId, leaseId), null, null, cancellationToken), cancellationToken);
 					if (step.State == BeginStepResponse.Types.Result.Waiting)
 					{
 						batchLogger.LogInformation("Waiting for dependency to be ready");
@@ -174,7 +183,7 @@ namespace Horde.Agent.Leases.Handlers
 					// Print the new state
 					Stopwatch stepTimer = Stopwatch.StartNew();
 
-					batchLogger.LogInformation("Starting job {JobId}, batch {BatchId}, step {StepId} (Drive Space Left: {DriveSpaceRemaining} GB)", executeTask.JobId, executeTask.BatchId, step.StepId, availableFreeSpace.ToString("F1"));
+					batchLogger.LogInformation("Starting job {JobId}, batch {BatchId}, step {StepId} (Drive Space Left: {DriveSpaceRemaining} GB)", options.JobId, options.BatchId, step.StepId, availableFreeSpace.ToString("F1"));
 
 					// Create a trace span
 					using IScope scope = GlobalTracer.Instance.BuildSpan("Execute").WithResourceName(step.Name).StartActive();
@@ -190,7 +199,7 @@ namespace Horde.Agent.Leases.Handlers
 					{
 						// Start writing to the log file
 #pragma warning disable CA2000 // Dispose objects before losing scope
-						await using (IServerLogger stepLogger = _serverLoggerFactory.CreateLogger(session, step.LogId, executeTask.JobId, executeTask.BatchId, step.StepId, step.Warnings))
+						await using (IServerLogger stepLogger = _serverLoggerFactory.CreateLogger(session, step.LogId, options.JobId, options.BatchId, step.StepId, step.Warnings, options.JobOptions.UseNewLogStorage))
 						{
 							// Execute the task
 							ILogger forwardingLogger = new DefaultLoggerIndentHandler(stepLogger);
@@ -202,7 +211,7 @@ namespace Horde.Agent.Leases.Handlers
 							using CancellationTokenSource stepPollCancelSource = new CancellationTokenSource();
 							using CancellationTokenSource stepAbortSource = new CancellationTokenSource();
 							TaskCompletionSource<bool> stepFinishedSource = new TaskCompletionSource<bool>();
-							Task stepPollTask = Task.Run(() => PollForStepAbort(rpcClient, executeTask.JobId, executeTask.BatchId, step.StepId, stepAbortSource, stepFinishedSource.Task, stepPollCancelSource.Token), cancellationToken);
+							Task stepPollTask = Task.Run(() => PollForStepAbort(rpcClient, options.JobId, options.BatchId, step.StepId, stepAbortSource, stepFinishedSource.Task, stepPollCancelSource.Token), cancellationToken);
 
 							try
 							{
@@ -231,7 +240,7 @@ namespace Horde.Agent.Leases.Handlers
 
 						// Update the server with the outcome from the step
 						batchLogger.LogInformation("Marking step as complete (Outcome={Outcome}, State={StepState})", stepOutcome, stepState);
-						await rpcClient.InvokeAsync((HordeRpc.HordeRpcClient x) => x.UpdateStepAsync(new UpdateStepRequest(executeTask.JobId, executeTask.BatchId, step.StepId, stepState, stepOutcome), null, null, cancellationToken), cancellationToken);
+						await rpcClient.InvokeAsync((HordeRpc.HordeRpcClient x) => x.UpdateStepAsync(new UpdateStepRequest(options.JobId, options.BatchId, step.StepId, stepState, stepOutcome), null, null, cancellationToken), cancellationToken);
 					}
 
 					// Print the finishing state
@@ -269,7 +278,7 @@ namespace Horde.Agent.Leases.Handlers
 		/// <param name="cancellationToken">Cancellation token to abort the batch</param>
 		/// <param name="stepCancellationToken">Cancellation token to abort only this individual step</param>
 		/// <returns>Async task</returns>
-		internal static async Task<(JobStepOutcome, JobStepState)> ExecuteStepAsync(JobExecutor executor, BeginStepResponse step, ILogger stepLogger, CancellationToken cancellationToken, CancellationToken stepCancellationToken)
+		internal static async Task<(JobStepOutcome, JobStepState)> ExecuteStepAsync(IJobExecutor executor, BeginStepResponse step, ILogger stepLogger, CancellationToken cancellationToken, CancellationToken stepCancellationToken)
 		{
 			using CancellationTokenSource combined = CancellationTokenSource.CreateLinkedTokenSource(cancellationToken, stepCancellationToken);
 			try
diff --git a/Engine/Source/Programs/Horde/Horde.Agent/Leases/Handlers/UpgradeHandler.cs b/Engine/Source/Programs/Horde/Horde.Agent/Leases/Handlers/UpgradeHandler.cs
--- a/Engine/Source/Programs/Horde/Horde.Agent/Leases/Handlers/UpgradeHandler.cs
+++ b/Engine/Source/Programs/Horde/Horde.Agent/Leases/Handlers/UpgradeHandler.cs
@@ -33,7 +33,7 @@ namespace Horde.Agent.Leases.Handlers
 		/// <inheritdoc/>
 		public override async Task<LeaseResult> ExecuteAsync(ISession session, string leaseId, UpgradeTask task, CancellationToken cancellationToken)
 		{
-			await using IServerLogger logger = _serverLoggerFactory.CreateLogger(session, task.LogId);
+			await using IServerLogger logger = _serverLoggerFactory.CreateLogger(session, task.LogId, null, null);
 
 			string requiredVersion = task.SoftwareId;
 
diff --git a/Engine/Source/Programs/Horde/Horde.Agent/Services/ServerLoggerFactory.cs b/Engine/Source/Programs/Horde/Horde.Agent/Services/ServerLoggerFactory.cs
--- a/Engine/Source/Programs/Horde/Horde.Agent/Services/ServerLoggerFactory.cs
+++ b/Engine/Source/Programs/Horde/Horde.Agent/Services/ServerLoggerFactory.cs
@@ -1,12 +1,8 @@
 // Copyright Epic Games, Inc. All Rights Reserved.
 
-using System;
-using Horde.Agent.Parser;
 using Microsoft.Extensions.Logging;
-using EpicGames.Horde.Storage.Backends;
-using System.Net.Http;
-using Microsoft.Extensions.Options;
 using Horde.Agent.Utility;
+using EpicGames.Horde.Storage;
 
 namespace Horde.Agent.Services
 {
@@ -24,8 +20,9 @@ namespace Horde.Agent.Services
 		/// <param name="batchId">Job batch id</param>
 		/// <param name="stepId">The job step id</param>
 		/// <param name="warnings">Whether to suppress warnings</param>
+		/// <param name="useNewLogger">Whether to enable the new logger backend</param>
 		/// <returns>New logger instance</returns>
-		IServerLogger CreateLogger(ISession session, string logId, string? jobId, string? batchId, string? stepId, bool? warnings = null);
+		IServerLogger CreateLogger(ISession session, string logId, string? jobId, string? batchId, string? stepId, bool? warnings, bool? useNewLogger);
 	}
 
 	/// <summary>
@@ -40,10 +37,11 @@ namespace Horde.Agent.Services
 		/// <param name="session">The current session</param>
 		/// <param name="logId">The log identifier</param>
 		/// <param name="warnings">Whether to suppress warnings</param>
+		/// <param name="useNewLogger">Whether to enable the new logger backend</param>
 		/// <returns>New logger instance</returns>
-		public static IServerLogger CreateLogger(this IServerLoggerFactory service, ISession session, string logId, bool? warnings = null)
+		public static IServerLogger CreateLogger(this IServerLoggerFactory service, ISession session, string logId, bool? warnings, bool? useNewLogger)
 		{
-			return service.CreateLogger(session, logId, null, null, null, warnings);
+			return service.CreateLogger(session, logId, null, null, null, warnings, useNewLogger);
 		}
 	}
 
@@ -52,28 +50,26 @@ namespace Horde.Agent.Services
 	/// </summary>
 	class ServerLoggerFactory : IServerLoggerFactory
 	{
-		readonly IHttpClientFactory _httpClientFactory;
-		readonly IOptions<AgentSettings> _settings;
+		readonly IServerStorageFactory _storageClientFactory;
 		readonly ILogger _logger;
 
 		/// <summary>
 		/// Constructor
 		/// </summary>
-		public ServerLoggerFactory(IHttpClientFactory httpClientFactory, IOptions<AgentSettings> settings, ILogger<ServerLoggerFactory> logger)
+		public ServerLoggerFactory(IServerStorageFactory storageClientFactory, ILogger<ServerLoggerFactory> logger)
 		{
-			_httpClientFactory = httpClientFactory;
-			_settings = settings;
+			_storageClientFactory = storageClientFactory;
 			_logger = logger;
 		}
 
 		/// <inheritdoc/>
-		public IServerLogger CreateLogger(ISession session, string logId, string? jobId, string? batchId, string? stepId, bool? warnings = null)
+		public IServerLogger CreateLogger(ISession session, string logId, string? jobId, string? batchId, string? stepId, bool? warnings, bool? useNewLogger)
 		{
 #pragma warning disable CA2000 // Dispose objects before losing scope
 			IJsonRpcLogSink sink = new JsonRpcLogSink(session.RpcConnection, jobId, batchId, stepId, _logger);
-			if (_settings.Value.EnableNewLogger)
+			if (useNewLogger ?? false)
 			{
-				HttpStorageClient storageClient = new HttpStorageClient(_httpClientFactory,  new Uri(session.ServerUrl, $"api/v1/logs/{logId}/"), session.Token, _logger);
+				IStorageClient storageClient = _storageClientFactory.CreateStorageClient(session);
 				sink = new JsonRpcAndStorageLogSink(session.RpcConnection, logId, sink, storageClient, _logger);
 			}
 			return new JsonRpcLogger(sink, logId, warnings, _logger);
diff --git a/Engine/Source/Programs/Horde/Horde.Agent/Services/ServerStorageFactory.cs b/Engine/Source/Programs/Horde/Horde.Agent/Services/ServerStorageFactory.cs
new file mode 100644
--- /dev/null
+++ b/Engine/Source/Programs/Horde/Horde.Agent/Services/ServerStorageFactory.cs
@@ -0,0 +1,61 @@
+// Copyright Epic Games, Inc. All Rights Reserved.
+
+using System;
+using System.Net.Http;
+using System.Net.Http.Headers;
+using EpicGames.Horde.Storage;
+using EpicGames.Horde.Storage.Backends;
+using Microsoft.Extensions.Logging;
+
+namespace Horde.Agent.Services
+{
+	/// <summary>
+	/// Interface for creating <see cref="IStorageClient"/> instances
+	/// </summary>
+	interface IServerStorageFactory
+	{
+		/// <summary>
+		/// Creates a storage client which authenticates using tokens for the current session
+		/// </summary>
+		/// <param name="session">The current session</param>
+		/// <returns>New logger instance</returns>
+		IStorageClient CreateStorageClient(ISession session);
+	}
+
+	/// <summary>
+	/// Implementation of <see cref="IServerStorageFactory"/> which constructs objects that communicate via HTTP
+	/// </summary>
+	class HttpServerStorageFactory : IServerStorageFactory
+	{
+		readonly IHttpClientFactory _httpClientFactory;
+		readonly ILogger _logger;
+
+		/// <summary>
+		/// Constructor
+		/// </summary>
+		public HttpServerStorageFactory(IHttpClientFactory httpClientFactory, ILogger<HttpStorageClient> logger)
+		{
+			_httpClientFactory = httpClientFactory;
+			_logger = logger;
+		}
+
+		/// <inheritdoc/>
+		public IStorageClient CreateStorageClient(ISession session)
+		{
+			return new HttpStorageClient(() => CreateHttpClient(session), CreateHttpRedirectClient, _logger);
+		}
+
+		HttpClient CreateHttpClient(ISession session)
+		{
+			HttpClient httpClient = _httpClientFactory.CreateClient(HttpStorageClient.HttpClientName);
+			httpClient.BaseAddress = new Uri($"{session.ServerUrl}/api/v1/storage/default");
+			httpClient.DefaultRequestHeaders.Authorization = new AuthenticationHeaderValue("Bearer", session.Token);
+			return httpClient;
+		}
+
+		HttpClient CreateHttpRedirectClient()
+		{
+			return _httpClientFactory.CreateClient(HttpStorageClient.HttpClientName);
+		}
+	}
+}
diff --git a/Engine/Source/Programs/Horde/Horde.Agent/Services/SessionFactory.cs b/Engine/Source/Programs/Horde/Horde.Agent/Services/SessionFactory.cs
new file mode 100644
--- /dev/null
+++ b/Engine/Source/Programs/Horde/Horde.Agent/Services/SessionFactory.cs
@@ -0,0 +1,279 @@
+// Copyright Epic Games, Inc. All Rights Reserved.
+
+using System;
+using System.Collections.Generic;
+using System.Linq;
+using System.Reflection;
+using System.Runtime.InteropServices;
+using System.Threading;
+using System.Threading.Tasks;
+using EpicGames.Core;
+using Grpc.Net.Client;
+using Horde.Agent.Utility;
+using HordeCommon;
+using HordeCommon.Rpc;
+using HordeCommon.Rpc.Messages;
+using Microsoft.Extensions.Logging;
+using Microsoft.Extensions.Options;
+
+namespace Horde.Agent.Services
+{
+	/// <summary>
+	/// Interface for about the current session. 
+	/// </summary>
+	interface ISession : IAsyncDisposable
+	{
+		/// <summary>
+		/// URL of the server
+		/// </summary>
+		Uri ServerUrl { get; }
+
+		/// <summary>
+		/// The agent identifier
+		/// </summary>
+		string AgentId { get; }
+
+		/// <summary>
+		/// Identifier for the current session
+		/// </summary>
+		string SessionId { get; }
+
+		/// <summary>
+		/// Token to use for connection to the server
+		/// </summary>
+		string Token { get; }
+
+		/// <summary>
+		/// Connection to the server
+		/// </summary>
+		IRpcConnection RpcConnection { get; }
+
+		/// <summary>
+		/// Working directory for sandboxes etc..
+		/// </summary>
+		DirectoryReference WorkingDir { get; }
+
+		/// <summary>
+		/// Terminate all processes in running in the working directory
+		/// </summary>
+		/// <param name="logger">Logger for any diagnostic messages</param>
+		/// <param name="cancellationToken">Cancellation token for the operation</param>
+		Task TerminateProcessesAsync(ILogger logger, CancellationToken cancellationToken);
+	}
+
+	/// <summary>
+	/// Interface for a factory to create sessions
+	/// </summary>
+	interface ISessionFactory
+	{
+		/// <summary>
+		/// Creates a new session
+		/// </summary>
+		/// <param name="cancellationToken">Cancellation token for the operation</param>
+		/// <returns>New session information</returns>
+		public Task<ISession> CreateAsync(CancellationToken cancellationToken);
+	}
+
+	/// <summary>
+	/// Information about the current session. 
+	/// </summary>
+	sealed class Session : ISession
+	{
+		/// <summary>
+		/// URL of the server
+		/// </summary>
+		public Uri ServerUrl { get; }
+
+		/// <summary>
+		/// The agent identifier
+		/// </summary>
+		public string AgentId { get; }
+
+		/// <summary>
+		/// Identifier for the current session
+		/// </summary>
+		public string SessionId { get; }
+
+		/// <summary>
+		/// Token to use for connection to the server
+		/// </summary>
+		public string Token { get; }
+
+		/// <summary>
+		/// Connection to the server
+		/// </summary>
+		public IRpcConnection RpcConnection { get; }
+
+		/// <summary>
+		/// Working directory for sandboxes etc..
+		/// </summary>
+		public DirectoryReference WorkingDir { get; }
+
+		readonly HashSet<string> _processNamesToTerminate;
+		readonly ILogger _logger;
+
+		/// <summary>
+		/// Constructor
+		/// </summary>
+		public Session(Uri serverUrl, string agentId, string sessionId, string token, IRpcConnection rpcConnection, DirectoryReference workingDir, IEnumerable<string> processNamesToTerminate, ILogger logger)
+		{
+			ServerUrl = serverUrl;
+			AgentId = agentId;
+			SessionId = sessionId;
+			Token = token;
+			RpcConnection = rpcConnection;
+			WorkingDir = workingDir;
+
+			_processNamesToTerminate = new HashSet<string>(processNamesToTerminate, StringComparer.OrdinalIgnoreCase);
+			_logger = logger;
+		}
+
+		/// <summary>
+		/// Creates a new agent session
+		/// </summary>
+		/// <param name="capabilitiesService"></param>
+		/// <param name="grpcService"></param>
+		/// <param name="settings"></param>
+		/// <param name="logger"></param>
+		/// <param name="cancellationToken">Indicates that the service is trying to stop</param>
+		/// <returns>Async task</returns>
+		public static async Task<Session> CreateAsync(CapabilitiesService capabilitiesService, GrpcService grpcService, IOptions<AgentSettings> settings, ILogger logger, CancellationToken cancellationToken)
+		{
+			AgentSettings currentSettings = settings.Value;
+
+			// Get the working directory
+			if (currentSettings.WorkingDir == null)
+			{
+				throw new Exception("WorkingDir is not set. Unable to run service.");
+			}
+
+			DirectoryReference baseDir = new FileReference(Assembly.GetExecutingAssembly().Location).Directory;
+			DirectoryReference workingDir = DirectoryReference.Combine(baseDir, currentSettings.WorkingDir);
+			logger.LogInformation("WorkingDir: {WorkingDir}", workingDir);
+			DirectoryReference.CreateDirectory(workingDir);
+
+			// Print the server info
+			ServerProfile serverProfile = currentSettings.GetCurrentServerProfile();
+			logger.LogInformation("Server: {Server}", serverProfile.Url);
+
+			// Show the worker capabilities
+			AgentCapabilities capabilities = await capabilitiesService.GetCapabilitiesAsync(workingDir);
+			if (capabilities.Properties.Count > 0)
+			{
+				logger.LogInformation("Global:");
+				foreach (string property in capabilities.Properties)
+				{
+					logger.LogInformation("  {AgentProperty}", property);
+				}
+			}
+			foreach (DeviceCapabilities device in capabilities.Devices)
+			{
+				logger.LogInformation("{DeviceName} Device:", device.Handle);
+				foreach (string property in device.Properties)
+				{
+					logger.LogInformation("   {DeviceProperty}", property);
+				}
+			}
+
+			// Mount all the necessary network shares. Currently only supported on Windows.
+			if (currentSettings.ShareMountingEnabled && RuntimeInformation.IsOSPlatform(OSPlatform.Windows))
+			{
+				foreach (MountNetworkShare share in currentSettings.Shares)
+				{
+					if (share.MountPoint != null && share.RemotePath != null)
+					{
+						logger.LogInformation("Mounting {RemotePath} as {MountPoint}", share.RemotePath, share.MountPoint);
+						NetworkShare.Mount(share.MountPoint, share.RemotePath);
+					}
+				}
+			}
+
+			// Create the session
+			CreateSessionResponse createSessionResponse;
+			using (GrpcChannel channel = grpcService.CreateGrpcChannel(serverProfile.Token))
+			{
+				HordeRpc.HordeRpcClient rpcClient = new HordeRpc.HordeRpcClient(channel);
+
+				// Create the session information
+				CreateSessionRequest sessionRequest = new CreateSessionRequest();
+				sessionRequest.Name = currentSettings.GetAgentName();
+				sessionRequest.Status = AgentStatus.Ok;
+				sessionRequest.Capabilities = capabilities;
+				sessionRequest.Version = Program.Version;
+
+				// Create a session
+				createSessionResponse = await rpcClient.CreateSessionAsync(sessionRequest, null, null, cancellationToken);
+				logger.LogInformation("Created session. AgentName={AgentName} SessionId={SessionId}", currentSettings.GetAgentName(), createSessionResponse.SessionId);
+			}
+
+			Func<GrpcChannel> createGrpcChannel = () => grpcService.CreateGrpcChannel(createSessionResponse.Token);
+
+			// Open a connection to the server
+#pragma warning disable CA2000 // False positive; ownership is transferred to new Session object.
+			IRpcConnection rpcConnection = new RpcConnection(createGrpcChannel, logger);
+			return new Session(serverProfile.Url, createSessionResponse.AgentId, createSessionResponse.SessionId, createSessionResponse.Token, rpcConnection, workingDir, currentSettings.ProcessNamesToTerminate, logger);
+#pragma warning restore CA2000
+		}
+
+		/// <summary>
+		/// Dispose of the current session
+		/// </summary>
+		/// <returns></returns>
+		public async ValueTask DisposeAsync()
+		{
+			await RpcConnection.DisposeAsync();
+		}
+
+		/// <inheritdoc/>
+		public Task TerminateProcessesAsync(ILogger logger, CancellationToken cancellationToken)
+		{
+			// Terminate child processes from any previous runs
+			ProcessUtils.TerminateProcesses(ShouldTerminateProcess, logger, cancellationToken);
+			return Task.CompletedTask;
+		}
+
+		/// <summary>
+		/// Callback for determining whether a process should be terminated
+		/// </summary>
+		bool ShouldTerminateProcess(FileReference imageFile)
+		{
+			if (imageFile.IsUnderDirectory(WorkingDir))
+			{
+				return true;
+			}
+
+			string fileName = imageFile.GetFileName();
+			if (_processNamesToTerminate.Contains(fileName))
+			{
+				return true;
+			}
+
+			return false;
+		}
+	}
+
+	/// <summary>
+	/// Creates session objects
+	/// </summary>
+	class SessionFactory : ISessionFactory
+	{
+		readonly CapabilitiesService _capabilitiesService;
+		readonly GrpcService _grpcService;
+		readonly IOptions<AgentSettings> _settings;
+		readonly ILogger _logger;
+
+		/// <summary>
+		/// Constructor
+		/// </summary>
+		public SessionFactory(CapabilitiesService capabilitiesService, GrpcService grpcService, IOptions<AgentSettings> settings, ILogger<SessionFactory> logger)
+		{
+			_capabilitiesService = capabilitiesService;
+			_grpcService = grpcService;
+			_settings = settings;
+			_logger = logger;
+		}
+
+		/// <inheritdoc/>
+		public async Task<ISession> CreateAsync(CancellationToken cancellationToken) => await Session.CreateAsync(_capabilitiesService, _grpcService, _settings, _logger, cancellationToken);
+	}
+}
diff --git a/Engine/Source/Programs/Horde/Horde.Agent/Services/SessionFactoryService.cs b/Engine/Source/Programs/Horde/Horde.Agent/Services/SessionFactoryService.cs
deleted file mode 100644
diff --git a/Engine/Source/Programs/Horde/Horde.Agent/Services/WorkerService.cs b/Engine/Source/Programs/Horde/Horde.Agent/Services/WorkerService.cs
--- a/Engine/Source/Programs/Horde/Horde.Agent/Services/WorkerService.cs
+++ b/Engine/Source/Programs/Horde/Horde.Agent/Services/WorkerService.cs
@@ -22,7 +22,7 @@ namespace Horde.Agent.Services
 	{
 		readonly ILogger _logger;
 		readonly IOptions<AgentSettings> _settings;
-		readonly SessionFactoryService _sessionFactoryService;
+		readonly ISessionFactory _sessionFactory;
 		readonly CapabilitiesService _capabilitiesService;
 		readonly LeaseHandler[] _leaseHandlers;
 		readonly IServiceProvider _serviceProvider;
@@ -30,10 +30,10 @@ namespace Horde.Agent.Services
 		/// <summary>
 		/// Constructor. Registers with the server and starts accepting connections.
 		/// </summary>
-		public WorkerService(IOptions<AgentSettings> settings, SessionFactoryService sessionFactoryService, CapabilitiesService capabilitiesService, IEnumerable<LeaseHandler> leaseHandlers, IServiceProvider serviceProvider, ILogger<WorkerService> logger)
+		public WorkerService(IOptions<AgentSettings> settings, ISessionFactory sessionFactory, CapabilitiesService capabilitiesService, IEnumerable<LeaseHandler> leaseHandlers, IServiceProvider serviceProvider, ILogger<WorkerService> logger)
 		{
 			_settings = settings;
-			_sessionFactoryService = sessionFactoryService;
+			_sessionFactory = sessionFactory;
 			_capabilitiesService = capabilitiesService;
 			_logger = logger;
 			_leaseHandlers = leaseHandlers.ToArray();
@@ -82,7 +82,7 @@ namespace Horde.Agent.Services
 						using Mutex singleInstanceMutex = new(false, "Global\\HordeAgent-DB828ACB-0AA5-4D32-A62A-21D4429B1014");
 						await WaitForMutexAsync(singleInstanceMutex, stoppingToken);
 
-						await using (ISession session = await _sessionFactoryService.CreateAsync(stoppingToken))
+						await using (ISession session = await _sessionFactory.CreateAsync(stoppingToken))
 						{
 							LeaseManager leaseManager = new LeaseManager(session, _capabilitiesService, _leaseHandlers, _settings, _logger);
 							result = await leaseManager.RunAsync(false, stoppingToken);
diff --git a/Engine/Source/Programs/Horde/Horde.Build/Jobs/IJob.cs b/Engine/Source/Programs/Horde/Horde.Build/Jobs/IJob.cs
--- a/Engine/Source/Programs/Horde/Horde.Build/Jobs/IJob.cs
+++ b/Engine/Source/Programs/Horde/Horde.Build/Jobs/IJob.cs
@@ -19,7 +19,9 @@ using Horde.Build.Streams;
 using Horde.Build.Ugs;
 using Horde.Build.Users;
 using Horde.Build.Utilities;
+using Horde.Common;
 using HordeCommon;
+using HordeCommon.Rpc.Tasks;
 using MongoDB.Bson;
 using MongoDB.Driver;
 
@@ -654,9 +656,9 @@ namespace Horde.Build.Jobs
 		public DateTime CreateTimeUtc { get; }
 
 		/// <summary>
-		/// Executor to use for this job
+		/// Options for executing the job
 		/// </summary>
-		public string? Executor { get; }
+		public JobOptions? JobOptions { get; }
 
 		/// <summary>
 		/// Largest value of the CombinedPriority value for batches in the ready state.
diff --git a/Engine/Source/Programs/Horde/Horde.Build/Jobs/IJobCollection.cs b/Engine/Source/Programs/Horde/Horde.Build/Jobs/IJobCollection.cs
--- a/Engine/Source/Programs/Horde/Horde.Build/Jobs/IJobCollection.cs
+++ b/Engine/Source/Programs/Horde/Horde.Build/Jobs/IJobCollection.cs
@@ -13,7 +13,9 @@ using Horde.Build.Logs;
 using Horde.Build.Streams;
 using Horde.Build.Users;
 using Horde.Build.Utilities;
+using Horde.Common;
 using HordeCommon;
+using HordeCommon.Rpc.Tasks;
 using MongoDB.Bson;
 
 namespace Horde.Build.Jobs
@@ -73,9 +75,9 @@ namespace Horde.Build.Jobs
 		public bool? PromoteIssuesByDefault { get; set; }
 
 		/// <summary>
-		/// The executor to use
+		/// Settings for agents executing the job
 		/// </summary>
-		public string? Executor { get; set; }
+		public JobOptions? JobOptions { get; set; }
 
 		/// <summary>
 		/// List of downstream job triggers
@@ -124,7 +126,7 @@ namespace Horde.Build.Jobs
 		/// </summary>
 		public CreateJobOptions(TemplateRefConfig templateRefConfig)
 		{
-			Executor = templateRefConfig.Executor;
+			JobOptions = templateRefConfig.JobOptions;
 			PromoteIssuesByDefault = templateRefConfig.PromoteIssuesByDefault;
 			if (templateRefConfig.ChainedJobs != null)
 			{
diff --git a/Engine/Source/Programs/Horde/Horde.Build/Jobs/JobCollection.cs b/Engine/Source/Programs/Horde/Horde.Build/Jobs/JobCollection.cs
--- a/Engine/Source/Programs/Horde/Horde.Build/Jobs/JobCollection.cs
+++ b/Engine/Source/Programs/Horde/Horde.Build/Jobs/JobCollection.cs
@@ -19,7 +19,9 @@ using Horde.Build.Streams;
 using Horde.Build.Telemetry;
 using Horde.Build.Users;
 using Horde.Build.Utilities;
+using Horde.Common;
 using HordeCommon;
+using HordeCommon.Rpc.Tasks;
 using Microsoft.Extensions.Logging;
 using MongoDB.Bson;
 using MongoDB.Bson.IO;
@@ -253,7 +255,7 @@ namespace Horde.Build.Jobs
 			public DateTimeOffset? CreateTime { get; set; }
 
 			[BsonIgnoreIfNull]
-			public string? Executor { get; set; }
+			public JobOptions? JobOptions { get; set; }
 
 			[BsonIgnoreIfNull]
 			public DateTime? CreateTimeUtc { get; set; }
@@ -323,7 +325,7 @@ namespace Horde.Build.Jobs
 				AutoSubmit = options.AutoSubmit ?? false;
 				UpdateIssues = options.UpdateIssues ?? (options.StartedByUserId == null && ( options.PreflightChange == 0 || options.PreflightChange == null));
 				PromoteIssuesByDefault = options.PromoteIssuesByDefault ?? false;
-				Executor = options.Executor;
+				JobOptions = options.JobOptions;
 				CreateTimeUtc = createTimeUtc;
 				ChainedJobs.AddRange(options.JobTriggers.Select(x => new ChainedJobDocument(x)));
 				ShowUgsBadges = options.ShowUgsBadges;
diff --git a/Engine/Source/Programs/Horde/Horde.Build/Jobs/JobTaskSource.cs b/Engine/Source/Programs/Horde/Horde.Build/Jobs/JobTaskSource.cs
--- a/Engine/Source/Programs/Horde/Horde.Build/Jobs/JobTaskSource.cs
+++ b/Engine/Source/Programs/Horde/Horde.Build/Jobs/JobTaskSource.cs
@@ -843,7 +843,7 @@ namespace Horde.Build.Jobs
 			task.BatchId = batch.Id.ToString();
 			task.LogId = logId.ToString();
 			task.JobName = leaseName.ToString();
-			task.Executor = job.Executor ?? String.Empty;
+			task.JobOptions = job.JobOptions;
 
 			List<HordeCommon.Rpc.Messages.AgentWorkspace> workspaces = new ();
 
diff --git a/Engine/Source/Programs/Horde/Horde.Build/Jobs/Templates/TemplateMessages.cs b/Engine/Source/Programs/Horde/Horde.Build/Jobs/Templates/TemplateMessages.cs
--- a/Engine/Source/Programs/Horde/Horde.Build/Jobs/Templates/TemplateMessages.cs
+++ b/Engine/Source/Programs/Horde/Horde.Build/Jobs/Templates/TemplateMessages.cs
@@ -3,10 +3,13 @@
 using System;
 using System.Collections.Generic;
 using System.ComponentModel.DataAnnotations;
+using System.Reflection;
 using System.Text.Json.Serialization;
 using EpicGames.Core;
 using Horde.Build.Streams;
+using Horde.Common;
 using HordeCommon;
+using HordeCommon.Rpc.Tasks;
 
 namespace Horde.Build.Jobs.Templates
 {
@@ -453,9 +456,9 @@ namespace Horde.Build.Jobs.Templates
 		public List<ParameterData> Parameters { get; set; } = new List<ParameterData>();
 
 		/// <summary>
-		/// The executor to use for this job. Mainly used for debugging.
+		/// Default settings for jobs
 		/// </summary>
-		public string? Executor { get; set; }
+		public JobOptions JobOptions { get; set; } = new JobOptions();
 
 		/// <summary>
 		/// The cached hash of this template.
diff --git a/Engine/Source/Programs/Horde/Horde.Build/Projects/ProjectConfig.cs b/Engine/Source/Programs/Horde/Horde.Build/Projects/ProjectConfig.cs
--- a/Engine/Source/Programs/Horde/Horde.Build/Projects/ProjectConfig.cs
+++ b/Engine/Source/Programs/Horde/Horde.Build/Projects/ProjectConfig.cs
@@ -10,9 +10,12 @@ using EpicGames.Core;
 using Horde.Build.Acls;
 using Horde.Build.Agents.Pools;
 using Horde.Build.Configuration;
+using Horde.Build.Jobs.Templates;
 using Horde.Build.Server;
 using Horde.Build.Streams;
 using Horde.Build.Utilities;
+using Horde.Common;
+using HordeCommon.Rpc.Tasks;
 
 namespace Horde.Build.Projects
 {
@@ -74,6 +77,11 @@ namespace Horde.Build.Projects
 		/// </summary>
 		public List<ProjectCategoryConfig> Categories { get; set; } = new List<ProjectCategoryConfig>();
 
+		/// <summary>
+		/// Default settings for executing jobs
+		/// </summary>
+		public JobOptions JobOptions { get; set; } = new JobOptions();
+
 		/// <summary>
 		/// List of streams
 		/// </summary>
diff --git a/Engine/Source/Programs/Horde/Horde.Build/Server/GlobalConfig.cs b/Engine/Source/Programs/Horde/Horde.Build/Server/GlobalConfig.cs
--- a/Engine/Source/Programs/Horde/Horde.Build/Server/GlobalConfig.cs
+++ b/Engine/Source/Programs/Horde/Horde.Build/Server/GlobalConfig.cs
@@ -23,7 +23,6 @@ using Horde.Build.Streams;
 using Horde.Build.Tools;
 using Horde.Build.Users;
 using Horde.Build.Utilities;
-using Microsoft.Extensions.Options;
 
 namespace Horde.Build.Server
 {
diff --git a/Engine/Source/Programs/Horde/Horde.Build/Streams/StreamConfig.cs b/Engine/Source/Programs/Horde/Horde.Build/Streams/StreamConfig.cs
--- a/Engine/Source/Programs/Horde/Horde.Build/Streams/StreamConfig.cs
+++ b/Engine/Source/Programs/Horde/Horde.Build/Streams/StreamConfig.cs
@@ -21,7 +21,9 @@ using Horde.Build.Perforce;
 using Horde.Build.Projects;
 using Horde.Build.Server;
 using Horde.Build.Utilities;
+using Horde.Common;
 using HordeCommon;
+using HordeCommon.Rpc.Tasks;
 
 namespace Horde.Build.Streams
 {
@@ -140,6 +142,11 @@ namespace Horde.Build.Streams
 		/// </summary>
 		public string? TriageChannel { get; set; }
 
+		/// <summary>
+		/// Default settings for executing jobs
+		/// </summary>
+		public JobOptions JobOptions { get; set; } = new JobOptions();
+
 		/// <summary>
 		/// Legacy name for the default preflight template
 		/// </summary>
@@ -229,6 +236,7 @@ namespace Horde.Build.Streams
 		{
 			Id = id;
 			ProjectConfig = projectConfig;
+			JobOptions.MergeDefaults(projectConfig.JobOptions);
 
 			foreach (TemplateRefConfig template in Templates)
 			{
@@ -738,6 +746,7 @@ namespace Horde.Build.Streams
 		public void PostLoad(StreamConfig streamConfig)
 		{
 			StreamConfig = streamConfig;
+			JobOptions.MergeDefaults(streamConfig.JobOptions);
 
 			if (Id.IsEmpty)
 			{
