Horde: Support extracting markdown docs from agent.

diff --git a/Engine/Source/Programs/Horde/Horde.Build/Commands/Config/DocsCommand.cs b/Engine/Source/Programs/Horde/Horde.Build/Commands/Config/DocsCommand.cs
--- a/Engine/Source/Programs/Horde/Horde.Build/Commands/Config/DocsCommand.cs
+++ b/Engine/Source/Programs/Horde/Horde.Build/Commands/Config/DocsCommand.cs
@@ -5,9 +5,9 @@ using System.Collections.Generic;
 using System.Diagnostics.CodeAnalysis;
 using System.IO;
 using System.Linq;
+using System.Reflection;
 using System.Text.RegularExpressions;
 using System.Threading.Tasks;
-using Amazon.S3.Model.Internal.MarshallTransformations;
 using EpicGames.Core;
 using Horde.Build.Projects;
 using Horde.Build.Server;
@@ -23,11 +23,16 @@ namespace Horde.Build.Commands.Config
 		[CommandLine]
 		public DirectoryReference OutputDir { get; set; } = DirectoryReference.Combine(Program.AppDir, "Docs");
 
+		[CommandLine]
+		public FileReference? Agent { get; set; }
+
 		class AppSettings
 		{
 			public ServerSettings Horde { get; set; } = new ServerSettings();
 		}
 
+		record class PageInfo(string Title, string FileName);
+
 		public override async Task<int> ExecuteAsync(ILogger logger)
 		{
 			DirectoryReference.CreateDirectory(OutputDir);
@@ -37,12 +42,32 @@ namespace Horde.Build.Commands.Config
 			JsonSchema projectSchema = Schemas.CreateSchema(typeof(ProjectConfig));
 			JsonSchema streamSchema = Schemas.CreateSchema(typeof(StreamConfig));
 
-			JsonSchemaType[] allTypes = { serverSchema.RootType, globalSchema.RootType, projectSchema.RootType, streamSchema.RootType };
+			Dictionary<JsonSchemaType, PageInfo> typeToPageInfo = new Dictionary<JsonSchemaType, PageInfo>
+			{
+				[serverSchema.RootType] = new PageInfo("appsettings.json (Server)", "Config-Schema-Server.md"),
+				[globalSchema.RootType] = new PageInfo("Globals.json", "Config-Schema-Globals.md"),
+				[projectSchema.RootType] = new PageInfo("*.project.json", "Config-Schema-Projects.md"),
+				[streamSchema.RootType] = new PageInfo("*.stream.json", "Config-Schema-Streams.md"),
+			};
+
+			if (Agent == null)
+			{
+				logger.LogWarning("Not generating documentation for agent settings. Specify path to agent assembly with -Agent=...");
+			}
+			else
+			{
+				Assembly agentAssembly = Assembly.LoadFile(Agent.FullName);
+				Type agentSettingsType = agentAssembly.GetType("Horde.Agent.AgentSettings")!;
+
+				JsonSchema agentSchema = Schemas.CreateSchema(agentSettingsType);
+				await WriteDocAsync(agentSchema.RootType, "appsettings.json (Agent)", "Config-Schema-Agent.md", new Dictionary<string, string>());
+			}
 
-			await WriteDocAsync(serverSchema.RootType, "appsettings.json", "Config-Server.md", allTypes);
-			await WriteDocAsync(globalSchema.RootType, "Globals.json", "Config-Global.md", allTypes);
-			await WriteDocAsync(projectSchema.RootType, "*.project.json", "Config-Project.md", allTypes);
-			await WriteDocAsync(streamSchema.RootType, "*.stream.json", "Config-Stream.md", allTypes);
+			Dictionary<string, string> typeNameToPageName = typeToPageInfo.ToDictionary(x => x.Key.Name!, x => x.Value.FileName, StringComparer.Ordinal);
+			foreach ((JsonSchemaType type, PageInfo pageInfo) in typeToPageInfo)
+			{
+				await WriteDocAsync(type, pageInfo.Title, pageInfo.FileName, typeNameToPageName);
+			}
 
 			return 0;
 		}
@@ -71,31 +96,33 @@ namespace Horde.Build.Commands.Config
 			public bool TryPop([NotNullWhen(true)] out JsonSchemaType? obj) => _stack.TryPop(out obj);
 		}
 
-		async Task WriteDocAsync(JsonSchemaType rootType, string title, string fileName, IEnumerable<JsonSchemaType> ignoreTypes)
+		async Task WriteDocAsync(JsonSchemaType rootType, string title, string fileName, Dictionary<string, string> typeNameToLink)
 		{
 			FileReference file = FileReference.Combine(OutputDir, fileName);
 			using (FileStream stream = FileReference.Open(file, FileMode.Create, FileAccess.Write))
 			{
 				using (StreamWriter writer = new StreamWriter(stream))
 				{
+					await writer.WriteLineAsync($"[Horde](../README.md) > [Configuration](Config.md) > {title}");
+					await writer.WriteLineAsync();
 					await writer.WriteLineAsync($"# {title}");
 
-					HashSet<JsonSchemaType> visitedTypes = new HashSet<JsonSchemaType>(ignoreTypes);
-					visitedTypes.Remove(rootType);
+					HashSet<string> visitedTypeNames = new HashSet<string>(typeNameToLink.Keys, StringComparer.Ordinal);
+					visitedTypeNames.Remove(rootType.Name!);
 
 					List<JsonSchemaType> types = new List<JsonSchemaType>();
-					FindCustomTypes(rootType, types, visitedTypes);
+					FindCustomTypes(rootType, types, visitedTypeNames);
 
 					foreach (JsonSchemaType schemaType in types)
 					{
-						await writer.WriteLineAsync();
-
 						if (schemaType != rootType)
 						{
-							await writer.WriteLineAsync($"## {GetHeadingName(schemaType)}");
 							await writer.WriteLineAsync();
+							await writer.WriteLineAsync($"## {GetHeadingName(schemaType)}");
 						}
 
+						await writer.WriteLineAsync();
+
 						if (schemaType.Description != null)
 						{
 							string description = Regex.Replace(schemaType.Description, "\r?\n", Environment.NewLine);
@@ -111,7 +138,7 @@ namespace Horde.Build.Commands.Config
 							foreach (JsonSchemaProperty property in schemaObj.Properties)
 							{
 								string name = property.CamelCaseName;
-								string type = GetMarkdownType(property.Type);
+								string type = GetMarkdownType(property.Type, typeNameToLink);
 								string description = GetMarkdownDescription(property.Description);
 								await writer.WriteLineAsync($"`{name}` | {type} | {description}");
 							}
@@ -135,51 +162,52 @@ namespace Horde.Build.Commands.Config
 					}
 				}
 			}
-
-
-			await Task.Delay(1);
 		}
 
-		static void FindCustomTypes(JsonSchemaType type, List<JsonSchemaType> types, HashSet<JsonSchemaType> visited)
+		static void FindCustomTypes(JsonSchemaType type, List<JsonSchemaType> types, HashSet<string> visitedTypeNames)
 		{
-			if (visited.Add(type))
+			switch (type)
 			{
-				switch (type)
-				{
-					case JsonSchemaOneOf oneOf:
-						foreach (JsonSchemaType oneOfType in oneOf.Types)
-						{
-							FindCustomTypes(oneOfType, types, visited);
-						}
-						break;
-					case JsonSchemaArray array:
-						FindCustomTypes(array.ItemType, types, visited);
-						break;
-					case JsonSchemaEnum _:
-						if (type.Name != null)
-						{
-							types.Add(type);
-						}
-						break;
-					case JsonSchemaObject obj:
-						if (type.Name != null)
+				case JsonSchemaOneOf oneOf:
+					foreach (JsonSchemaType oneOfType in oneOf.Types)
+					{
+						FindCustomTypes(oneOfType, types, visitedTypeNames);
+					}
+					break;
+				case JsonSchemaArray array:
+					FindCustomTypes(array.ItemType, types, visitedTypeNames);
+					break;
+				case JsonSchemaEnum _:
+					if (type.Name != null && visitedTypeNames.Add(type.Name))
+					{
+						types.Add(type);
+					}
+					break;
+				case JsonSchemaObject obj:
+					if (type.Name != null)
+					{
+						if (visitedTypeNames.Add(type.Name))
 						{
 							types.Add(type);
 						}
-						if (obj.AdditionalProperties != null)
-						{
-							FindCustomTypes(obj.AdditionalProperties, types, visited);
-						}
-						foreach (JsonSchemaProperty property in obj.Properties)
+						else
 						{
-							FindCustomTypes(property.Type, types, visited);
+							break;
 						}
-						break;
-				}
+					}
+					if (obj.AdditionalProperties != null)
+					{
+						FindCustomTypes(obj.AdditionalProperties, types, visitedTypeNames);
+					}
+					foreach (JsonSchemaProperty property in obj.Properties)
+					{
+						FindCustomTypes(property.Type, types, visitedTypeNames);
+					}
+					break;
 			}
 		}
 
-		static string GetMarkdownType(JsonSchemaType type)
+		static string GetMarkdownType(JsonSchemaType type, Dictionary<string, string> typeNameToLink)
 		{
 			switch (type)
 			{
@@ -192,9 +220,9 @@ namespace Horde.Build.Commands.Config
 				case JsonSchemaString _:
 					return "`string`";
 				case JsonSchemaOneOf oneOf:
-					return String.Join("/", oneOf.Types.Select(x => GetMarkdownType(x)));
+					return String.Join("/", oneOf.Types.Select(x => GetMarkdownType(x, typeNameToLink)));
 				case JsonSchemaArray array:
-					string elementType = GetMarkdownType(array.ItemType);
+					string elementType = GetMarkdownType(array.ItemType, typeNameToLink);
 					if (elementType.EndsWith("`", StringComparison.Ordinal))
 					{
 						return elementType.Insert(elementType.Length - 1, "[]");
@@ -211,12 +239,12 @@ namespace Horde.Build.Commands.Config
 					}
 					else
 					{
-						return $"[`{type.Name}`](#{GetAnchorName(type)})";
+						return $"[`{type.Name}`]({GetLink(type, typeNameToLink)})";
 					}
 				case JsonSchemaObject obj:
 					if (obj.AdditionalProperties != null)
 					{
-						return "`string` `->` " + GetMarkdownType(obj.AdditionalProperties);
+						return "`string` `->` " + GetMarkdownType(obj.AdditionalProperties, typeNameToLink);
 					}
 					else if (type.Name == null)
 					{
@@ -224,7 +252,7 @@ namespace Horde.Build.Commands.Config
 					}
 					else
 					{
-						return $"[`{type.Name}`](#{GetAnchorName(type)})";
+						return $"[`{type.Name}`]({GetLink(type, typeNameToLink)})";
 					}
 				default:
 					return type.GetType().Name;
@@ -252,12 +280,25 @@ namespace Horde.Build.Commands.Config
 			}
 		}
 
+		static string GetLink(JsonSchemaType type, Dictionary<string, string> typeNameToLink)
+		{
+			string? link;
+			if (type.Name != null && typeNameToLink.TryGetValue(type.Name, out link))
+			{
+				return link;
+			}
+			else
+			{
+				return GetAnchorName(type);
+			}
+		}
+
 		[SuppressMessage("Globalization", "CA1308:Normalize strings to uppercase")]
 		static string GetAnchorName(JsonSchemaType type)
 		{
 			string anchor = GetHeadingName(type).ToLowerInvariant();
 			anchor = Regex.Replace(anchor, @"[^a-z0-9]+", "-");
-			return anchor.Trim('-');
+			return "#" + anchor.Trim('-');
 		}
 	}
 }
