Horde: Missing file.

diff --git a/Engine/Source/Programs/Horde/Horde.Build/Protos/horde/tasks/job_task.cs b/Engine/Source/Programs/Horde/Horde.Build/Protos/horde/tasks/job_task.cs
new file mode 100644
--- /dev/null
+++ b/Engine/Source/Programs/Horde/Horde.Build/Protos/horde/tasks/job_task.cs
@@ -0,0 +1,17 @@
+// Copyright Epic Games, Inc. All Rights Reserved.
+
+namespace HordeCommon.Rpc.Tasks
+{
+	partial class JobOptions
+	{
+		/// <summary>
+		/// Merge settings from another JobOptions object
+		/// </summary>
+		public void MergeDefaults(JobOptions other)
+		{
+			Executor ??= other.Executor;
+			UseNewLogStorage ??= other.UseNewLogStorage;
+			UseNewTempStorage ??= other.UseNewTempStorage;
+		}
+	}
+}
