Horde: Missing file.

diff --git a/Engine/Source/Programs/Horde/Horde.Agent/Program.cs b/Engine/Source/Programs/Horde/Horde.Agent/Program.cs
--- a/Engine/Source/Programs/Horde/Horde.Agent/Program.cs
+++ b/Engine/Source/Programs/Horde/Horde.Agent/Program.cs
@@ -181,9 +181,9 @@ namespace Horde.Agent
 
 			services.AddSingleton<GrpcService>();
 
-			services.AddSingleton<JobExecutorFactory, PerforceExecutorFactory>();
-			services.AddSingleton<JobExecutorFactory, LocalExecutorFactory>();
-			services.AddSingleton<JobExecutorFactory, TestExecutorFactory>();
+			services.AddSingleton<IJobExecutorFactory, PerforceExecutorFactory>();
+			services.AddSingleton<IJobExecutorFactory, LocalExecutorFactory>();
+			services.AddSingleton<IJobExecutorFactory, TestExecutorFactory>();
 
 			services.AddSingleton<LeaseHandler, ComputeHandler>();
 			services.AddSingleton<LeaseHandler, ComputeHandlerV2>();
@@ -194,7 +194,7 @@ namespace Horde.Agent
 			services.AddSingleton<LeaseHandler, UpgradeHandler>();
 
 			services.AddSingleton<CapabilitiesService>();
-			services.AddSingleton<SessionFactoryService>();
+			services.AddSingleton<ISessionFactory, SessionFactory>();
 			services.AddSingleton<IServerLoggerFactory, ServerLoggerFactory>();
 			services.AddHostedService<WorkerService>();
 
