Horde: Move TemplateConfig class into a separate file.

diff --git a/Engine/Source/Programs/Horde/Horde.Build/Jobs/Templates/TemplateConfig.cs b/Engine/Source/Programs/Horde/Horde.Build/Jobs/Templates/TemplateConfig.cs
new file mode 100644
--- /dev/null
+++ b/Engine/Source/Programs/Horde/Horde.Build/Jobs/Templates/TemplateConfig.cs
@@ -0,0 +1,467 @@
+// Copyright Epic Games, Inc. All Rights Reserved.
+
+using System;
+using System.Collections.Generic;
+using System.ComponentModel.DataAnnotations;
+using System.Text.Json.Serialization;
+using EpicGames.Core;
+using Horde.Build.Streams;
+using HordeCommon;
+using HordeCommon.Rpc.Tasks;
+
+namespace Horde.Build.Jobs.Templates
+{
+	/// <summary>
+	/// Parameters to create a new template
+	/// </summary>
+	public class TemplateConfig
+	{
+		/// <summary>
+		/// Name for the new template
+		/// </summary>
+		[Required]
+		public string Name { get; set; } = null!;
+
+		/// <summary>
+		/// Description for the template
+		/// </summary>
+		public string? Description { get; set; }
+
+		/// <summary>
+		/// Default priority for this job
+		/// </summary>
+		public Priority? Priority { get; set; }
+
+		/// <summary>
+		/// Whether to allow preflights of this template
+		/// </summary>
+		public bool AllowPreflights { get; set; } = true;
+
+		/// <summary>
+		/// Whether issues should be updated for all jobs using this template
+		/// </summary>
+		public bool UpdateIssues { get; set; } = false;
+
+		/// <summary>
+		/// Whether issues should be promoted by default for this template, promoted issues will generate user notifications 
+		/// </summary>
+		public bool PromoteIssuesByDefault { get; set; } = false;
+
+		/// <summary>
+		/// Initial agent type to parse the buildgraph script on
+		/// </summary>
+		public string? InitialAgentType { get; set; }
+
+		/// <summary>
+		/// Path to a file within the stream to submit to generate a new changelist for jobs
+		/// </summary>
+		public string? SubmitNewChange { get; set; }
+
+		/// <summary>
+		/// Description for new changelists
+		/// </summary>
+		public string? SubmitDescription { get; set; }
+
+		/// <summary>
+		/// Default change to build at. Each object has a condition parameter which can evaluated by the server to determine which change to use.
+		/// </summary>
+		public List<ChangeQueryConfig>? DefaultChange { get; set; }
+
+		/// <summary>
+		/// Fixed arguments for the new job
+		/// </summary>
+		public List<string> Arguments { get; set; } = new List<string>();
+
+		/// <summary>
+		/// Parameters for this template
+		/// </summary>
+		public List<ParameterData> Parameters { get; set; } = new List<ParameterData>();
+
+		/// <summary>
+		/// Default settings for jobs
+		/// </summary>
+		public JobOptions JobOptions { get; set; } = new JobOptions();
+
+		/// <summary>
+		/// The cached hash of this template.
+		/// </summary>
+		[JsonIgnore]
+		internal ContentHash? CachedHash { get; set; }
+	}
+
+	/// <summary>
+	/// Base class for template parameters
+	/// </summary>
+	[JsonKnownTypes(typeof(GroupParameterData), typeof(TextParameterData), typeof(ListParameterData), typeof(BoolParameterData))]
+	public abstract class ParameterData
+	{
+		/// <summary>
+		/// Convert to a parameter object
+		/// </summary>
+		/// <returns><see cref="Parameter"/> object</returns>
+		public abstract Parameter ToModel();
+	}
+
+	/// <summary>
+	/// Describes how to render a group parameter
+	/// </summary>
+	public enum GroupParameterStyle
+	{
+		/// <summary>
+		/// Separate tab on the form
+		/// </summary>
+		Tab,
+
+		/// <summary>
+		/// Section with heading
+		/// </summary>
+		Section,
+	}
+
+	/// <summary>
+	/// Used to group a number of other parameters
+	/// </summary>
+	[JsonDiscriminator("Group")]
+	public class GroupParameterData : ParameterData
+	{
+		/// <summary>
+		/// Label to display next to this parameter
+		/// </summary>
+		public string Label { get; set; }
+
+		/// <summary>
+		/// How to display this group
+		/// </summary>
+		public GroupParameterStyle Style { get; set; }
+
+		/// <summary>
+		/// List of child parameters
+		/// </summary>
+		public List<ParameterData> Children { get; set; }
+
+		/// <summary>
+		/// Private constructor for serialization
+		/// </summary>
+		public GroupParameterData()
+		{
+			Label = null!;
+			Children = null!;
+		}
+
+		/// <summary>
+		/// Constructor
+		/// </summary>
+		/// <param name="label">Name of the group</param>
+		/// <param name="style">How to display this group</param>
+		/// <param name="children">List of child parameters</param>
+		public GroupParameterData(string label, GroupParameterStyle style, List<ParameterData> children)
+		{
+			Label = label;
+			Style = style;
+			Children = children;
+		}
+
+		/// <summary>
+		/// Converts this data to a model object
+		/// </summary>
+		/// <returns>New <see cref="GroupParameter"/> object</returns>
+		public override Parameter ToModel()
+		{
+			return new GroupParameter(Label, Style, Children.ConvertAll(x => x.ToModel()));
+		}
+	}
+
+	/// <summary>
+	/// Free-form text entry parameter
+	/// </summary>
+	[JsonDiscriminator("Text")]
+	public class TextParameterData : ParameterData
+	{
+		/// <summary>
+		/// Name of the parameter associated with this parameter.
+		/// </summary>
+		public string Label { get; set; }
+
+		/// <summary>
+		/// Argument to pass to the executor
+		/// </summary>
+		public string Argument { get; set; }
+
+		/// <summary>
+		/// Default value for this argument
+		/// </summary>
+		public string Default { get; set; }
+
+		/// <summary>
+		/// Hint text for this parameter
+		/// </summary>
+		public string? Hint { get; set; }
+
+		/// <summary>
+		/// Regex used to validate this parameter
+		/// </summary>
+		public string? Validation { get; set; }
+
+		/// <summary>
+		/// Message displayed if validation fails, informing user of valid values.
+		/// </summary>
+		public string? ValidationError { get; set; }
+
+		/// <summary>
+		/// Tool-tip text to display
+		/// </summary>
+		public string? ToolTip { get; set; }
+
+		/// <summary>
+		/// Private constructor for serialization
+		/// </summary>
+		public TextParameterData()
+		{
+			Label = null!;
+			Argument = null!;
+			Default = String.Empty;
+		}
+
+		/// <summary>
+		/// Constructor
+		/// </summary>
+		/// <param name="label">Label to show next to the parameter</param>
+		/// <param name="argument">Argument to pass this value with</param>
+		/// <param name="defaultValue">Default value for this parameter</param>
+		/// <param name="hint">Hint text to display for this parameter</param>
+		/// <param name="validation">Regex used to validate entries</param>
+		/// <param name="validationError">Message displayed to explain validation issues</param>
+		/// <param name="toolTip">Tool tip text to display</param>
+		public TextParameterData(string label, string argument, string defaultValue, string? hint, string? validation, string? validationError, string? toolTip)
+		{
+			Label = label;
+			Argument = argument;
+			Default = defaultValue;
+			Hint = hint;
+			Validation = validation;
+			ValidationError = validationError;
+			ToolTip = toolTip;
+		}
+
+		/// <summary>
+		/// Converts this data to a model object
+		/// </summary>
+		/// <returns>New <see cref="TextParameter"/> object</returns>
+		public override Parameter ToModel()
+		{
+			return new TextParameter(Label, Argument, Default, Hint, Validation, ValidationError, ToolTip);
+		}
+	}
+
+	/// <summary>
+	/// Style of list parameter
+	/// </summary>
+	public enum ListParameterStyle
+	{
+		/// <summary>
+		/// Regular drop-down list. One item is always selected.
+		/// </summary>
+		List,
+
+		/// <summary>
+		/// Drop-down list with checkboxes
+		/// </summary>
+		MultiList,
+
+		/// <summary>
+		/// Tag picker from list of options
+		/// </summary>
+		TagPicker,
+	}
+
+	/// <summary>
+	/// Possible option for a list parameter
+	/// </summary>
+	public class ListParameterItemData
+	{
+		/// <summary>
+		/// Optional group heading to display this entry under, if the picker style supports it.
+		/// </summary>
+		public string? Group { get; set; }
+
+		/// <summary>
+		/// Name of the parameter associated with this list.
+		/// </summary>
+		public string Text { get; set; }
+
+		/// <summary>
+		/// Argument to pass with this parameter.
+		/// </summary>
+		public string? ArgumentIfEnabled { get; set; }
+
+		/// <summary>
+		/// Argument to pass with this parameter.
+		/// </summary>
+		public string? ArgumentIfDisabled { get; set; }
+
+		/// <summary>
+		/// Whether this item is selected by default
+		/// </summary>
+		public bool Default { get; set; }
+
+		/// <summary>
+		/// Private constructor for serialization
+		/// </summary>
+		public ListParameterItemData()
+		{
+			Text = null!;
+		}
+
+		/// <summary>
+		/// Constructor
+		/// </summary>
+		/// <param name="group">The group to put this parameter in</param>
+		/// <param name="text">Text to display for this option</param>
+		/// <param name="argumentIfEnabled">Argument to pass for this item if it's enabled</param>
+		/// <param name="argumentIfDisabled">Argument to pass for this item if it's enabled</param>
+		/// <param name="defaultValue">Whether this item is selected by default</param>
+		public ListParameterItemData(string? group, string text, string? argumentIfEnabled, string? argumentIfDisabled, bool defaultValue)
+		{
+			Group = group;
+			Text = text;
+			ArgumentIfEnabled = argumentIfEnabled;
+			ArgumentIfDisabled = argumentIfDisabled;
+			Default = defaultValue;
+		}
+
+		/// <summary>
+		/// Converts this data to a model object
+		/// </summary>
+		/// <returns>New <see cref="ListParameterItem"/> object</returns>
+		public ListParameterItem ToModel()
+		{
+			return new ListParameterItem(Group, Text, ArgumentIfEnabled, ArgumentIfDisabled, Default);
+		}
+	}
+
+	/// <summary>
+	/// Allows the user to select a value from a constrained list of choices
+	/// </summary>
+	[JsonDiscriminator("List")]
+	public class ListParameterData : ParameterData
+	{
+		/// <summary>
+		/// Label to display next to this parameter. Defaults to the parameter name.
+		/// </summary>
+		public string Label { get; set; }
+
+		/// <summary>
+		/// The type of list parameter
+		/// </summary>
+		public ListParameterStyle Style { get; set; }
+
+		/// <summary>
+		/// List of values to display in the list
+		/// </summary>
+		public List<ListParameterItemData> Items { get; set; }
+
+		/// <summary>
+		/// Tool tip text to display
+		/// </summary>
+		public string? ToolTip { get; set; }
+
+		/// <summary>
+		/// Private constructor
+		/// </summary>
+		public ListParameterData()
+		{
+			Label = null!;
+			Items = null!;
+		}
+
+		/// <summary>
+		/// List of possible values
+		/// </summary>
+		/// <param name="label">Label to show next to this parameter</param>
+		/// <param name="style">Type of picker to show</param>
+		/// <param name="items">Entries for this list</param>
+		/// <param name="toolTip">Tool tip text to display</param>
+		public ListParameterData(string label, ListParameterStyle style, List<ListParameterItemData> items, string? toolTip)
+		{
+			Label = label;
+			Style = style;
+			Items = items;
+			ToolTip = toolTip;
+		}
+
+		/// <summary>
+		/// Converts this data to a model object
+		/// </summary>
+		/// <returns>New <see cref="ListParameter"/> object</returns>
+		public override Parameter ToModel()
+		{
+			return new ListParameter(Label, Style, Items.ConvertAll(x => x.ToModel()), ToolTip);
+		}
+	}
+
+	/// <summary>
+	/// Allows the user to toggle an option on or off
+	/// </summary>
+	[JsonDiscriminator("Bool")]
+	public class BoolParameterData : ParameterData
+	{
+		/// <summary>
+		/// Name of the parameter associated with this parameter.
+		/// </summary>
+		public string Label { get; set; }
+
+		/// <summary>
+		/// Value if enabled
+		/// </summary>
+		public string? ArgumentIfEnabled { get; set; }
+
+		/// <summary>
+		/// Value if disabled
+		/// </summary>
+		public string? ArgumentIfDisabled { get; set; }
+
+		/// <summary>
+		/// Whether this argument is enabled by default
+		/// </summary>
+		public bool Default { get; set; }
+
+		/// <summary>
+		/// Tool tip text to display
+		/// </summary>
+		public string? ToolTip { get; set; }
+
+		/// <summary>
+		/// Private constructor for serialization
+		/// </summary>
+		public BoolParameterData()
+		{
+			Label = null!;
+		}
+
+		/// <summary>
+		/// Constructor
+		/// </summary>
+		/// <param name="label">Label to show next to this parameter</param>
+		/// <param name="argumentIfEnabled">Value if enabled</param>
+		/// <param name="argumentIfDisabled">Value if disabled</param>
+		/// <param name="defaultValue">Whether this option is enabled by default</param>
+		/// <param name="toolTip">The tool tip text to display</param>
+		public BoolParameterData(string label, string? argumentIfEnabled, string? argumentIfDisabled, bool defaultValue, string? toolTip)
+		{
+			Label = label;
+			ArgumentIfEnabled = argumentIfEnabled;
+			ArgumentIfDisabled = argumentIfDisabled;
+			Default = defaultValue;
+			ToolTip = toolTip;
+		}
+
+		/// <summary>
+		/// Converts this data to a model object
+		/// </summary>
+		/// <returns>New <see cref="BoolParameter"/> object</returns>
+		public override Parameter ToModel()
+		{
+			return new BoolParameter(Label, ArgumentIfEnabled, ArgumentIfDisabled, Default, ToolTip);
+		}
+	}
+}
diff --git a/Engine/Source/Programs/Horde/Horde.Build/Jobs/Templates/TemplateMessages.cs b/Engine/Source/Programs/Horde/Horde.Build/Jobs/Templates/TemplateMessages.cs
--- a/Engine/Source/Programs/Horde/Horde.Build/Jobs/Templates/TemplateMessages.cs
+++ b/Engine/Source/Programs/Horde/Horde.Build/Jobs/Templates/TemplateMessages.cs
@@ -13,460 +13,6 @@ using HordeCommon.Rpc.Tasks;
 
 namespace Horde.Build.Jobs.Templates
 {
-	/// <summary>
-	/// Base class for template parameters
-	/// </summary>
-	[JsonKnownTypes(typeof(GroupParameterData), typeof(TextParameterData), typeof(ListParameterData), typeof(BoolParameterData))]
-	public abstract class ParameterData
-	{
-		/// <summary>
-		/// Convert to a parameter object
-		/// </summary>
-		/// <returns><see cref="Parameter"/> object</returns>
-		public abstract Parameter ToModel();
-	}
-
-	/// <summary>
-	/// Describes how to render a group parameter
-	/// </summary>
-	public enum GroupParameterStyle
-	{
-		/// <summary>
-		/// Separate tab on the form
-		/// </summary>
-		Tab,
-
-		/// <summary>
-		/// Section with heading
-		/// </summary>
-		Section,
-	}
-
-	/// <summary>
-	/// Used to group a number of other parameters
-	/// </summary>
-	[JsonDiscriminator("Group")]
-	public class GroupParameterData : ParameterData
-	{
-		/// <summary>
-		/// Label to display next to this parameter
-		/// </summary>
-		public string Label { get; set; }
-
-		/// <summary>
-		/// How to display this group
-		/// </summary>
-		public GroupParameterStyle Style { get; set; }
-
-		/// <summary>
-		/// List of child parameters
-		/// </summary>
-		public List<ParameterData> Children { get; set; }
-
-		/// <summary>
-		/// Private constructor for serialization
-		/// </summary>
-		public GroupParameterData()
-		{
-			Label = null!;
-			Children = null!;
-		}
-
-		/// <summary>
-		/// Constructor
-		/// </summary>
-		/// <param name="label">Name of the group</param>
-		/// <param name="style">How to display this group</param>
-		/// <param name="children">List of child parameters</param>
-		public GroupParameterData(string label, GroupParameterStyle style, List<ParameterData> children)
-		{
-			Label = label;
-			Style = style;
-			Children = children;
-		}
-
-		/// <summary>
-		/// Converts this data to a model object
-		/// </summary>
-		/// <returns>New <see cref="GroupParameter"/> object</returns>
-		public override Parameter ToModel()
-		{
-			return new GroupParameter(Label, Style, Children.ConvertAll(x => x.ToModel()));
-		}
-	}
-
-	/// <summary>
-	/// Free-form text entry parameter
-	/// </summary>
-	[JsonDiscriminator("Text")]
-	public class TextParameterData : ParameterData
-	{
-		/// <summary>
-		/// Name of the parameter associated with this parameter.
-		/// </summary>
-		public string Label { get; set; }
-
-		/// <summary>
-		/// Argument to pass to the executor
-		/// </summary>
-		public string Argument { get; set; }
-
-		/// <summary>
-		/// Default value for this argument
-		/// </summary>
-		public string Default { get; set; }
-
-		/// <summary>
-		/// Hint text for this parameter
-		/// </summary>
-		public string? Hint { get; set; }
-
-		/// <summary>
-		/// Regex used to validate this parameter
-		/// </summary>
-		public string? Validation { get; set; }
-
-		/// <summary>
-		/// Message displayed if validation fails, informing user of valid values.
-		/// </summary>
-		public string? ValidationError { get; set; }
-
-		/// <summary>
-		/// Tool-tip text to display
-		/// </summary>
-		public string? ToolTip { get; set; }
-
-		/// <summary>
-		/// Private constructor for serialization
-		/// </summary>
-		public TextParameterData()
-		{
-			Label = null!;
-			Argument = null!;
-			Default = String.Empty;
-		}
-
-		/// <summary>
-		/// Constructor
-		/// </summary>
-		/// <param name="label">Label to show next to the parameter</param>
-		/// <param name="argument">Argument to pass this value with</param>
-		/// <param name="defaultValue">Default value for this parameter</param>
-		/// <param name="hint">Hint text to display for this parameter</param>
-		/// <param name="validation">Regex used to validate entries</param>
-		/// <param name="validationError">Message displayed to explain validation issues</param>
-		/// <param name="toolTip">Tool tip text to display</param>
-		public TextParameterData(string label, string argument, string defaultValue, string? hint, string? validation, string? validationError, string? toolTip)
-		{
-			Label = label;
-			Argument = argument;
-			Default = defaultValue;
-			Hint = hint;
-			Validation = validation;
-			ValidationError = validationError;
-			ToolTip = toolTip;
-		}
-
-		/// <summary>
-		/// Converts this data to a model object
-		/// </summary>
-		/// <returns>New <see cref="TextParameter"/> object</returns>
-		public override Parameter ToModel()
-		{
-			return new TextParameter(Label, Argument, Default, Hint, Validation, ValidationError, ToolTip);
-		}
-	}
-
-	/// <summary>
-	/// Style of list parameter
-	/// </summary>
-	public enum ListParameterStyle
-	{
-		/// <summary>
-		/// Regular drop-down list. One item is always selected.
-		/// </summary>
-		List,
-
-		/// <summary>
-		/// Drop-down list with checkboxes
-		/// </summary>
-		MultiList,
-
-		/// <summary>
-		/// Tag picker from list of options
-		/// </summary>
-		TagPicker,
-	}
-
-	/// <summary>
-	/// Possible option for a list parameter
-	/// </summary>
-	public class ListParameterItemData
-	{
-		/// <summary>
-		/// Optional group heading to display this entry under, if the picker style supports it.
-		/// </summary>
-		public string? Group { get; set; }
-
-		/// <summary>
-		/// Name of the parameter associated with this list.
-		/// </summary>
-		public string Text { get; set; }
-
-		/// <summary>
-		/// Argument to pass with this parameter.
-		/// </summary>
-		public string? ArgumentIfEnabled { get; set; }
-
-		/// <summary>
-		/// Argument to pass with this parameter.
-		/// </summary>
-		public string? ArgumentIfDisabled { get; set; }
-
-		/// <summary>
-		/// Whether this item is selected by default
-		/// </summary>
-		public bool Default { get; set; }
-
-		/// <summary>
-		/// Private constructor for serialization
-		/// </summary>
-		public ListParameterItemData()
-		{
-			Text = null!;
-		}
-
-		/// <summary>
-		/// Constructor
-		/// </summary>
-		/// <param name="group">The group to put this parameter in</param>
-		/// <param name="text">Text to display for this option</param>
-		/// <param name="argumentIfEnabled">Argument to pass for this item if it's enabled</param>
-		/// <param name="argumentIfDisabled">Argument to pass for this item if it's enabled</param>
-		/// <param name="defaultValue">Whether this item is selected by default</param>
-		public ListParameterItemData(string? group, string text, string? argumentIfEnabled, string? argumentIfDisabled, bool defaultValue)
-		{
-			Group = group;
-			Text = text;
-			ArgumentIfEnabled = argumentIfEnabled;
-			ArgumentIfDisabled = argumentIfDisabled;
-			Default = defaultValue;
-		}
-
-		/// <summary>
-		/// Converts this data to a model object
-		/// </summary>
-		/// <returns>New <see cref="ListParameterItem"/> object</returns>
-		public ListParameterItem ToModel()
-		{
-			return new ListParameterItem(Group, Text, ArgumentIfEnabled, ArgumentIfDisabled, Default);
-		}
-	}
-
-	/// <summary>
-	/// Allows the user to select a value from a constrained list of choices
-	/// </summary>
-	[JsonDiscriminator("List")]
-	public class ListParameterData : ParameterData
-	{
-		/// <summary>
-		/// Label to display next to this parameter. Defaults to the parameter name.
-		/// </summary>
-		public string Label { get; set; }
-
-		/// <summary>
-		/// The type of list parameter
-		/// </summary>
-		public ListParameterStyle Style { get; set; }
-
-		/// <summary>
-		/// List of values to display in the list
-		/// </summary>
-		public List<ListParameterItemData> Items { get; set; }
-
-		/// <summary>
-		/// Tool tip text to display
-		/// </summary>
-		public string? ToolTip { get; set; }
-
-		/// <summary>
-		/// Private constructor
-		/// </summary>
-		public ListParameterData()
-		{
-			Label = null!;
-			Items = null!;
-		}
-
-		/// <summary>
-		/// List of possible values
-		/// </summary>
-		/// <param name="label">Label to show next to this parameter</param>
-		/// <param name="style">Type of picker to show</param>
-		/// <param name="items">Entries for this list</param>
-		/// <param name="toolTip">Tool tip text to display</param>
-		public ListParameterData(string label, ListParameterStyle style, List<ListParameterItemData> items, string? toolTip)
-		{
-			Label = label;
-			Style = style;
-			Items = items;
-			ToolTip = toolTip;
-		}
-
-		/// <summary>
-		/// Converts this data to a model object
-		/// </summary>
-		/// <returns>New <see cref="ListParameter"/> object</returns>
-		public override Parameter ToModel()
-		{
-			return new ListParameter(Label, Style, Items.ConvertAll(x => x.ToModel()), ToolTip);
-		}
-	}
-
-	/// <summary>
-	/// Allows the user to toggle an option on or off
-	/// </summary>
-	[JsonDiscriminator("Bool")]
-	public class BoolParameterData : ParameterData
-	{
-		/// <summary>
-		/// Name of the parameter associated with this parameter.
-		/// </summary>
-		public string Label { get; set; }
-
-		/// <summary>
-		/// Value if enabled
-		/// </summary>
-		public string? ArgumentIfEnabled { get; set; }
-
-		/// <summary>
-		/// Value if disabled
-		/// </summary>
-		public string? ArgumentIfDisabled { get; set; }
-
-		/// <summary>
-		/// Whether this argument is enabled by default
-		/// </summary>
-		public bool Default { get; set; }
-
-		/// <summary>
-		/// Tool tip text to display
-		/// </summary>
-		public string? ToolTip { get; set; }
-
-		/// <summary>
-		/// Private constructor for serialization
-		/// </summary>
-		public BoolParameterData()
-		{
-			Label = null!;
-		}
-
-		/// <summary>
-		/// Constructor
-		/// </summary>
-		/// <param name="label">Label to show next to this parameter</param>
-		/// <param name="argumentIfEnabled">Value if enabled</param>
-		/// <param name="argumentIfDisabled">Value if disabled</param>
-		/// <param name="defaultValue">Whether this option is enabled by default</param>
-		/// <param name="toolTip">The tool tip text to display</param>
-		public BoolParameterData(string label, string? argumentIfEnabled, string? argumentIfDisabled, bool defaultValue, string? toolTip)
-		{
-			Label = label;
-			ArgumentIfEnabled = argumentIfEnabled;
-			ArgumentIfDisabled = argumentIfDisabled;
-			Default = defaultValue;
-			ToolTip = toolTip;
-		}
-
-		/// <summary>
-		/// Converts this data to a model object
-		/// </summary>
-		/// <returns>New <see cref="BoolParameter"/> object</returns>
-		public override Parameter ToModel()
-		{
-			return new BoolParameter(Label, ArgumentIfEnabled, ArgumentIfDisabled, Default, ToolTip);
-		}
-	}
-
-	/// <summary>
-	/// Parameters to create a new template
-	/// </summary>
-	public class TemplateConfig
-	{
-		/// <summary>
-		/// Name for the new template
-		/// </summary>
-		[Required]
-		public string Name { get; set; } = null!;
-
-		/// <summary>
-		/// Description for the template
-		/// </summary>
-		public string? Description { get; set; }
-
-		/// <summary>
-		/// Default priority for this job
-		/// </summary>
-		public Priority? Priority { get; set; }
-
-		/// <summary>
-		/// Whether to allow preflights of this template
-		/// </summary>
-		public bool AllowPreflights { get; set; } = true;
-
-		/// <summary>
-		/// Whether issues should be updated for all jobs using this template
-		/// </summary>
-		public bool UpdateIssues { get; set; } = false;
-
-		/// <summary>
-		/// Whether issues should be promoted by default for this template, promoted issues will generate user notifications 
-		/// </summary>
-		public bool PromoteIssuesByDefault { get; set; } = false;
-
-		/// <summary>
-		/// Initial agent type to parse the buildgraph script on
-		/// </summary>
-		public string? InitialAgentType { get; set; }
-
-		/// <summary>
-		/// Path to a file within the stream to submit to generate a new changelist for jobs
-		/// </summary>
-		public string? SubmitNewChange { get; set; }
-
-		/// <summary>
-		/// Description for new changelists
-		/// </summary>
-		public string? SubmitDescription { get; set; }
-
-		/// <summary>
-		/// Default change to build at. Each object has a condition parameter which can evaluated by the server to determine which change to use.
-		/// </summary>
-		public List<ChangeQueryConfig>? DefaultChange { get; set; }
-
-		/// <summary>
-		/// Fixed arguments for the new job
-		/// </summary>
-		public List<string> Arguments { get; set; } = new List<string>();
-
-		/// <summary>
-		/// Parameters for this template
-		/// </summary>
-		public List<ParameterData> Parameters { get; set; } = new List<ParameterData>();
-
-		/// <summary>
-		/// Default settings for jobs
-		/// </summary>
-		public JobOptions JobOptions { get; set; } = new JobOptions();
-
-		/// <summary>
-		/// The cached hash of this template.
-		/// </summary>
-		[JsonIgnore]
-		internal ContentHash? CachedHash { get; set; }
-	}
-
 	/// <summary>
 	/// Response describing a template
 	/// </summary>
