Horde: Print out more diagnostics about the selected job options.

diff --git a/Engine/Source/Programs/Horde/Horde.Agent/Leases/Handlers/JobHandler.cs b/Engine/Source/Programs/Horde/Horde.Agent/Leases/Handlers/JobHandler.cs
--- a/Engine/Source/Programs/Horde/Horde.Agent/Leases/Handlers/JobHandler.cs
+++ b/Engine/Source/Programs/Horde/Horde.Agent/Leases/Handlers/JobHandler.cs
@@ -66,6 +66,8 @@ namespace Horde.Agent.Leases.Handlers
 			GlobalTracer.Instance.ActiveSpan?.SetTag("jobName", executeTask.JobName.ToString());
 			GlobalTracer.Instance.ActiveSpan?.SetTag("batchId", executeTask.BatchId.ToString());
 
+			logger.LogInformation("Executor: {Name}, UseNewLogStorage: {UseNewLogStorage}, UseNewTempStorage: {UseNewTempStorage}", jobOptions.Executor, jobOptions.UseNewLogStorage ?? false, jobOptions.UseNewTempStorage ?? false);
+
 			// Start executing the current batch
 			BeginBatchResponse batch = await session.RpcConnection.InvokeAsync<HordeRpc.HordeRpcClient, BeginBatchResponse>(x => x.BeginBatchAsync(new BeginBatchRequest(executeTask.JobId, executeTask.BatchId, leaseId), null, null, cancellationToken), cancellationToken);
 
