Horde: Missing file.

diff --git a/Engine/Source/Programs/Horde/Horde.Agent/Program.cs b/Engine/Source/Programs/Horde/Horde.Agent/Program.cs
--- a/Engine/Source/Programs/Horde/Horde.Agent/Program.cs
+++ b/Engine/Source/Programs/Horde/Horde.Agent/Program.cs
@@ -196,6 +196,7 @@ namespace Horde.Agent
 			services.AddSingleton<CapabilitiesService>();
 			services.AddSingleton<ISessionFactory, SessionFactory>();
 			services.AddSingleton<IServerLoggerFactory, ServerLoggerFactory>();
+			services.AddSingleton<IServerStorageFactory, HttpServerStorageFactory>();
 			services.AddHostedService<WorkerService>();
 
 			services.AddSingleton<IStorageClientFactory, StorageClientFactory>();
