Horde: Allow specifying the base URL for storage uploads, allowing use of the dedicated logs endpoint.

diff --git a/Engine/Source/Programs/Horde/Horde.Agent.Tests/WorkerServiceTest.cs b/Engine/Source/Programs/Horde/Horde.Agent.Tests/WorkerServiceTest.cs
--- a/Engine/Source/Programs/Horde/Horde.Agent.Tests/WorkerServiceTest.cs
+++ b/Engine/Source/Programs/Horde/Horde.Agent.Tests/WorkerServiceTest.cs
@@ -61,7 +61,7 @@ namespace Horde.Agent.Tests
 
 		class FakeSessionStorageFactory : IServerStorageFactory
 		{
-			public IStorageClient CreateStorageClient(ISession session) => null!;
+			public IStorageClient CreateStorageClient(ISession session, string basePath) => null!;
 		}
 
 		internal static IJobExecutor NullExecutor = new SimpleTestExecutor(async (step, logger, cancellationToken) =>
diff --git a/Engine/Source/Programs/Horde/Horde.Agent/Execution/JobExecutor.cs b/Engine/Source/Programs/Horde/Horde.Agent/Execution/JobExecutor.cs
--- a/Engine/Source/Programs/Horde/Horde.Agent/Execution/JobExecutor.cs
+++ b/Engine/Source/Programs/Horde/Horde.Agent/Execution/JobExecutor.cs
@@ -45,16 +45,16 @@ namespace Horde.Agent.Execution
 	class JobExecutorOptions
 	{
 		public ISession Session { get; }
-		public IStorageClient Storage { get; }
+		public IServerStorageFactory StorageFactory { get; }
 		public string JobId { get; }
 		public string BatchId { get; }
 		public string AgentType { get; }
 		public JobOptions JobOptions { get; }
 
-		public JobExecutorOptions(ISession session, IStorageClient storage, string jobId, string batchId, string agentType, JobOptions jobOptions)
+		public JobExecutorOptions(ISession session, IServerStorageFactory storageFactory, string jobId, string batchId, string agentType, JobOptions jobOptions)
 		{
 			Session = session;
-			Storage = storage;
+			StorageFactory = storageFactory;
 			JobId = jobId;
 			BatchId = batchId;
 			AgentType = agentType;
@@ -191,7 +191,7 @@ namespace Horde.Agent.Execution
 		protected bool _compileAutomationTool = true;
 
 		protected readonly ISession _session;
-		protected readonly IStorageClient _storage;
+		protected readonly IServerStorageFactory _storageFactory;
 		protected readonly JobOptions _jobOptions;
 
 		protected IRpcConnection RpcConnection => _session.RpcConnection;
@@ -202,7 +202,7 @@ namespace Horde.Agent.Execution
 		public JobExecutor(JobExecutorOptions options, ILogger logger)
 		{
 			_session = options.Session;
-			_storage = options.Storage;
+			_storageFactory = options.StorageFactory;
 
 			_jobId = options.JobId;
 			_batchId = options.BatchId;
@@ -704,7 +704,8 @@ namespace Horde.Agent.Execution
 			DirectoryReference manifestDir = DirectoryReference.Combine(workspaceDir, "Engine", "Saved", "BuildGraph");
 
 			using MemoryCache cache = new MemoryCache(new MemoryCacheOptions { });
-			TreeReader reader = new TreeReader(_storage, cache, _logger);
+			IStorageClient storage = _storageFactory.CreateStorageClient(_session, $"/api/v1/jobs/{_jobId}/batches/{_batchId}/steps/{step.StepId}/temp");
+			TreeReader reader = new TreeReader(storage, cache, _logger);
 
 			// Create the mapping of tag names to file sets
 			Dictionary<string, HashSet<FileReference>> tagNameToFileSet = new Dictionary<string, HashSet<FileReference>>();
@@ -760,7 +761,7 @@ namespace Horde.Agent.Execution
 			}
 
 			// Run UAT
-			if (await ExecuteAutomationToolAsync(step, workspaceDir, arguments, _storage, logger, cancellationToken) != 0)
+			if (await ExecuteAutomationToolAsync(step, workspaceDir, arguments, storage, logger, cancellationToken) != 0)
 			{
 				return false;
 			}
@@ -867,7 +868,7 @@ namespace Horde.Agent.Execution
 				RefName refName = TempStorage.GetRefNameForNode(RefPrefix, step.Name);
 
 				TreeOptions treeOptions = new TreeOptions();
-				using TreeWriter treeWriter = new TreeWriter(_storage, treeOptions, refName.Text);
+				using TreeWriter treeWriter = new TreeWriter(storage, treeOptions, refName.Text);
 
 				TempStorageNode outputNode = new TempStorageNode();
 
diff --git a/Engine/Source/Programs/Horde/Horde.Agent/Leases/Handlers/JobHandler.cs b/Engine/Source/Programs/Horde/Horde.Agent/Leases/Handlers/JobHandler.cs
--- a/Engine/Source/Programs/Horde/Horde.Agent/Leases/Handlers/JobHandler.cs
+++ b/Engine/Source/Programs/Horde/Horde.Agent/Leases/Handlers/JobHandler.cs
@@ -38,7 +38,7 @@ namespace Horde.Agent.Leases.Handlers
 
 		readonly IEnumerable<IJobExecutorFactory> _executorFactories;
 		readonly AgentSettings _settings;
-		readonly IServerStorageFactory _storageClientFactory;
+		readonly IServerStorageFactory _serverStorageFactory;
 		readonly IServerLoggerFactory _serverLoggerFactory;
 		readonly ILogger _defaultLogger;
 
@@ -49,7 +49,7 @@ namespace Horde.Agent.Leases.Handlers
 		{
 			_executorFactories = executorFactories;
 			_settings = settings.Value;
-			_storageClientFactory = storageClientFactory;
+			_serverStorageFactory = storageClientFactory;
 			_serverLoggerFactory = serverLoggerFactory;
 			_defaultLogger = defaultLogger;
 		}
@@ -74,8 +74,7 @@ namespace Horde.Agent.Leases.Handlers
 			// Execute the batch
 			try
 			{
-				IStorageClient storage = _storageClientFactory.CreateStorageClient(session);
-				JobExecutorOptions options = new JobExecutorOptions(session, storage, executeTask.JobId, executeTask.BatchId, batch.AgentType, jobOptions);
+				JobExecutorOptions options = new JobExecutorOptions(session, _serverStorageFactory, executeTask.JobId, executeTask.BatchId, batch.AgentType, jobOptions);
 				await ExecuteBatchAsync(session, leaseId, executeTask.Workspace, executeTask.AutoSdkWorkspace, options, logger, cancellationToken);
 			}
 			catch (Exception ex)
diff --git a/Engine/Source/Programs/Horde/Horde.Agent/Services/ServerLoggerFactory.cs b/Engine/Source/Programs/Horde/Horde.Agent/Services/ServerLoggerFactory.cs
--- a/Engine/Source/Programs/Horde/Horde.Agent/Services/ServerLoggerFactory.cs
+++ b/Engine/Source/Programs/Horde/Horde.Agent/Services/ServerLoggerFactory.cs
@@ -69,7 +69,7 @@ namespace Horde.Agent.Services
 			IJsonRpcLogSink sink = new JsonRpcLogSink(session.RpcConnection, jobId, batchId, stepId, _logger);
 			if (useNewLogger ?? false)
 			{
-				IStorageClient storageClient = _storageClientFactory.CreateStorageClient(session);
+				IStorageClient storageClient = _storageClientFactory.CreateStorageClient(session, $"/api/v1/logs/{logId}");
 				sink = new JsonRpcAndStorageLogSink(session.RpcConnection, logId, sink, storageClient, _logger);
 			}
 			return new JsonRpcLogger(sink, logId, warnings, _logger);
diff --git a/Engine/Source/Programs/Horde/Horde.Agent/Services/ServerStorageFactory.cs b/Engine/Source/Programs/Horde/Horde.Agent/Services/ServerStorageFactory.cs
--- a/Engine/Source/Programs/Horde/Horde.Agent/Services/ServerStorageFactory.cs
+++ b/Engine/Source/Programs/Horde/Horde.Agent/Services/ServerStorageFactory.cs
@@ -18,8 +18,9 @@ namespace Horde.Agent.Services
 		/// Creates a storage client which authenticates using tokens for the current session
 		/// </summary>
 		/// <param name="session">The current session</param>
+		/// <param name="baseUrl">Base URL for writing blobs on the server</param>
 		/// <returns>New logger instance</returns>
-		IStorageClient CreateStorageClient(ISession session);
+		IStorageClient CreateStorageClient(ISession session, string baseUrl);
 	}
 
 	/// <summary>
@@ -40,15 +41,15 @@ namespace Horde.Agent.Services
 		}
 
 		/// <inheritdoc/>
-		public IStorageClient CreateStorageClient(ISession session)
+		public IStorageClient CreateStorageClient(ISession session, string baseUrl)
 		{
-			return new HttpStorageClient(() => CreateHttpClient(session), CreateHttpRedirectClient, _logger);
+			return new HttpStorageClient(() => CreateHttpClient(session, baseUrl), CreateHttpRedirectClient, _logger);
 		}
 
-		HttpClient CreateHttpClient(ISession session)
+		HttpClient CreateHttpClient(ISession session, string baseUrl)
 		{
 			HttpClient httpClient = _httpClientFactory.CreateClient(HttpStorageClient.HttpClientName);
-			httpClient.BaseAddress = new Uri($"{session.ServerUrl}/api/v1/storage/default");
+			httpClient.BaseAddress = new Uri(session.ServerUrl, baseUrl);
 			httpClient.DefaultRequestHeaders.Authorization = new AuthenticationHeaderValue("Bearer", session.Token);
 			return httpClient;
 		}
