Horde: Fix malformed log upload URLs.

diff --git a/Engine/Source/Programs/Horde/Horde.Agent/Services/ServerLoggerFactory.cs b/Engine/Source/Programs/Horde/Horde.Agent/Services/ServerLoggerFactory.cs
--- a/Engine/Source/Programs/Horde/Horde.Agent/Services/ServerLoggerFactory.cs
+++ b/Engine/Source/Programs/Horde/Horde.Agent/Services/ServerLoggerFactory.cs
@@ -69,7 +69,7 @@ namespace Horde.Agent.Services
 			IJsonRpcLogSink sink = new JsonRpcLogSink(session.RpcConnection, jobId, batchId, stepId, _logger);
 			if (useNewLogger ?? false)
 			{
-				IStorageClient storageClient = _storageClientFactory.CreateStorageClient(session, $"/api/v1/logs/{logId}");
+				IStorageClient storageClient = _storageClientFactory.CreateStorageClient(session, $"/api/v1/logs/{logId}/");
 				sink = new JsonRpcAndStorageLogSink(session.RpcConnection, logId, sink, storageClient, _logger);
 			}
 			return new JsonRpcLogger(sink, logId, warnings, _logger);
