Horde: Allow log writes to be authorized by session.

diff --git a/Engine/Source/Programs/Horde/Horde.Build/Logs/LogsController.cs b/Engine/Source/Programs/Horde/Horde.Build/Logs/LogsController.cs
--- a/Engine/Source/Programs/Horde/Horde.Build/Logs/LogsController.cs
+++ b/Engine/Source/Programs/Horde/Horde.Build/Logs/LogsController.cs
@@ -394,6 +394,10 @@ namespace Horde.Build.Logs
 		async Task<bool> AuthorizeAsync(ILogFile logFile, AclAction action, ClaimsPrincipal user)
 		{
 			GlobalConfig globalConfig = _globalConfig.Value;
+			if (logFile.SessionId != null && LogFileService.AuthorizeForSession(logFile, user))
+			{
+				return true;
+			}
 			if (logFile.JobId != JobId.Empty && await _jobService.AuthorizeAsync(logFile.JobId, action, user, globalConfig))
 			{
 				return true;
