Horde: Add a flag to LogFileDocument indicating whether to use the legacy or new storage backend, and enable both code paths rather than switching through config setting.

diff --git a/Engine/Source/Programs/Horde/Horde.Agent/Utility/JsonRpcLogSink.cs b/Engine/Source/Programs/Horde/Horde.Agent/Utility/JsonRpcLogSink.cs
--- a/Engine/Source/Programs/Horde/Horde.Agent/Utility/JsonRpcLogSink.cs
+++ b/Engine/Source/Programs/Horde/Horde.Agent/Utility/JsonRpcLogSink.cs
@@ -194,11 +194,6 @@ namespace Horde.Agent.Utility
 		/// <inheritdoc/>
 		public async Task WriteOutputAsync(WriteOutputRequest request, CancellationToken cancellationToken)
 		{
-			if (_inner != null)
-			{
-				await _inner.WriteOutputAsync(request, cancellationToken);
-			}
-
 			_builder.WriteData(request.Data.Memory);
 			_bufferLength += request.Data.Length;
 
diff --git a/Engine/Source/Programs/Horde/Horde.Build.Tests/IssueServiceTests.cs b/Engine/Source/Programs/Horde/Horde.Build.Tests/IssueServiceTests.cs
--- a/Engine/Source/Programs/Horde/Horde.Build.Tests/IssueServiceTests.cs
+++ b/Engine/Source/Programs/Horde/Horde.Build.Tests/IssueServiceTests.cs
@@ -221,7 +221,7 @@ namespace Horde.Build.Tests
 				{
 					SubResourceId stepId = new SubResourceId((ushort)((groupIdx * 100) + nodeIdx));
 
-					ILogFile logFile = LogFileService.CreateLogFileAsync(jobId, null, LogType.Json).Result;
+					ILogFile logFile = LogFileService.CreateLogFileAsync(jobId, null, LogType.Json, false).Result;
 
 					Mock<IJobStep> step = new Mock<IJobStep>(MockBehavior.Strict);
 					step.SetupGet(x => x.Id).Returns(stepId);
diff --git a/Engine/Source/Programs/Horde/Horde.Build.Tests/LogFileServiceTests.cs b/Engine/Source/Programs/Horde/Horde.Build.Tests/LogFileServiceTests.cs
--- a/Engine/Source/Programs/Horde/Horde.Build.Tests/LogFileServiceTests.cs
+++ b/Engine/Source/Programs/Horde/Horde.Build.Tests/LogFileServiceTests.cs
@@ -79,7 +79,7 @@ namespace Horde.Build.Tests
         public async Task WriteLogLifecycleOldTest()
         {
 			JobId jobId = JobId.GenerateNewId();
-            ILogFile logFile = await _logFileService.CreateLogFileAsync(jobId, null, LogType.Text, null, cancellationToken: CancellationToken.None);
+            ILogFile logFile = await _logFileService.CreateLogFileAsync(jobId, null, LogType.Text, useNewStorageBackend: false, logId: null, cancellationToken: CancellationToken.None);
 
             logFile = (await ((ILogFileService)_logFileService).WriteLogDataAsync(logFile, 0, 0, Encoding.ASCII.GetBytes("hello\n"), true))!;
 			logFile = (await ((ILogFileService)_logFileService).WriteLogDataAsync(logFile, 6, 1, Encoding.ASCII.GetBytes("foo\nbar\n"), true))!;
@@ -150,7 +150,7 @@ namespace Horde.Build.Tests
         public async Task WriteLogLifecycle(ILogFileService lfs, int maxChunkLength)
         {
 			JobId jobId = JobId.GenerateNewId();
-            ILogFile logFile = await lfs.CreateLogFileAsync(jobId, null, LogType.Text);
+            ILogFile logFile = await lfs.CreateLogFileAsync(jobId, null, LogType.Text, useNewStorageBackend: false);
 
             string str1 = "hello\n";
             string str2 = "foo\nbar\n";
@@ -257,7 +257,7 @@ namespace Horde.Build.Tests
 			// Will implicitly test GetLogFileAsync(), AddCachedLogFile()
 			JobId jobId = JobId.GenerateNewId();
             ObjectId sessionId = ObjectId.GenerateNewId();
-            ILogFile a = await _logFileService.CreateLogFileAsync(jobId, new ObjectId<ISession>(sessionId), LogType.Text, null, CancellationToken.None);
+            ILogFile a = await _logFileService.CreateLogFileAsync(jobId, new ObjectId<ISession>(sessionId), LogType.Text, useNewStorageBackend: false, logId: null, CancellationToken.None);
             ILogFile b = (await _logFileService.GetCachedLogFileAsync(a.Id, CancellationToken.None))!;
             Assert.AreEqual(a.JobId, b.JobId);
             Assert.AreEqual(a.SessionId, b.SessionId);
@@ -266,7 +266,7 @@ namespace Horde.Build.Tests
             ILogFile? notFound = await _logFileService.GetCachedLogFileAsync(LogId.GenerateNewId(), CancellationToken.None);
             Assert.IsNull(notFound);
 
-            await _logFileService.CreateLogFileAsync(JobId.GenerateNewId(), new ObjectId<ISession>(ObjectId.GenerateNewId()), LogType.Text, null, CancellationToken.None);
+            await _logFileService.CreateLogFileAsync(JobId.GenerateNewId(), new ObjectId<ISession>(ObjectId.GenerateNewId()), LogType.Text, useNewStorageBackend: false, logId: null, cancellationToken: CancellationToken.None);
             Assert.AreEqual(2, (await _logFileService.GetLogFilesAsync()).Count);
         }
 
@@ -275,8 +275,8 @@ namespace Horde.Build.Tests
         {
 			JobId jobId = JobId.GenerateNewId();
             ObjectId sessionId = ObjectId.GenerateNewId();
-            ILogFile logFile = await _logFileService.CreateLogFileAsync(jobId, new ObjectId<ISession>(sessionId), LogType.Text, null, CancellationToken.None);
-            ILogFile logFileNoSession = await _logFileService.CreateLogFileAsync(jobId, null, LogType.Text, null, CancellationToken.None);
+            ILogFile logFile = await _logFileService.CreateLogFileAsync(jobId, new ObjectId<ISession>(sessionId), LogType.Text, useNewStorageBackend: false, logId: null, cancellationToken: CancellationToken.None);
+            ILogFile logFileNoSession = await _logFileService.CreateLogFileAsync(jobId, null, LogType.Text, useNewStorageBackend: false, logId: null, cancellationToken: CancellationToken.None);
 
 			ClaimsPrincipal hasClaim = new ClaimsPrincipal(new ClaimsIdentity(new List<Claim>
             {
@@ -296,7 +296,7 @@ namespace Horde.Build.Tests
 		public async Task ChunkSplitting()
 		{
 			JobId jobId = JobId.GenerateNewId();
-			ILogFile logFile = await _logFileService.CreateLogFileAsync(jobId, null, LogType.Text, null, CancellationToken.None);
+			ILogFile logFile = await _logFileService.CreateLogFileAsync(jobId, null, LogType.Text, useNewStorageBackend: false, logId: null, cancellationToken: CancellationToken.None);
 
 			long offset = 0;
 
diff --git a/Engine/Source/Programs/Horde/Horde.Build.Tests/LogIndexingTests.cs b/Engine/Source/Programs/Horde/Horde.Build.Tests/LogIndexingTests.cs
--- a/Engine/Source/Programs/Horde/Horde.Build.Tests/LogIndexingTests.cs
+++ b/Engine/Source/Programs/Horde/Horde.Build.Tests/LogIndexingTests.cs
@@ -66,7 +66,7 @@ namespace Horde.Build.Tests
 		public async Task IndexTests()
 		{
 			JobId jobId = JobId.GenerateNewId();
-			ILogFile logFile = await _logFileService.CreateLogFileAsync(jobId, null, LogType.Text);
+			ILogFile logFile = await _logFileService.CreateLogFileAsync(jobId, null, LogType.Text, useNewStorageBackend: false);
 
 			// Write the test data to the log file in blocks
 			int offset = 0;
@@ -137,7 +137,7 @@ namespace Horde.Build.Tests
 		public async Task PartialTokenTests()
 		{
 			JobId jobId = JobId.GenerateNewId();
-			ILogFile logFile = await _logFileService.CreateLogFileAsync(jobId, null, LogType.Text);
+			ILogFile logFile = await _logFileService.CreateLogFileAsync(jobId, null, LogType.Text, useNewStorageBackend: false);
 
 			string[] lines =
 			{
@@ -179,7 +179,7 @@ namespace Horde.Build.Tests
 		public async Task AppendIndexTests()
 		{
 			JobId jobId = JobId.GenerateNewId();
-			ILogFile logFile = await _logFileService.CreateLogFileAsync(jobId, null, LogType.Text);
+			ILogFile logFile = await _logFileService.CreateLogFileAsync(jobId, null, LogType.Text, useNewStorageBackend: false);
 
 			logFile = await WriteLogDataAsync(logFile, 0, 0, Encoding.UTF8.GetBytes("abc\n"), true);
 			logFile = await WriteLogDataAsync(logFile, 4, 1, Encoding.UTF8.GetBytes("def\n"), true);
diff --git a/Engine/Source/Programs/Horde/Horde.Build.Tests/TestDataTests.cs b/Engine/Source/Programs/Horde/Horde.Build.Tests/TestDataTests.cs
--- a/Engine/Source/Programs/Horde/Horde.Build.Tests/TestDataTests.cs
+++ b/Engine/Source/Programs/Horde/Horde.Build.Tests/TestDataTests.cs
@@ -122,7 +122,7 @@ namespace Horde.Build.Tests
 				{
 					SubResourceId stepId = new SubResourceId((ushort)((groupIdx * 100) + nodeIdx));
 
-					ILogFile logFile = LogFileService.CreateLogFileAsync(jobId, null, LogType.Json).Result;
+					ILogFile logFile = LogFileService.CreateLogFileAsync(jobId, null, LogType.Json, useNewStorageBackend: false).Result;
 
 					Mock<IJobStep> step = new Mock<IJobStep>(MockBehavior.Strict);
 					step.SetupGet(x => x.Id).Returns(stepId);
diff --git a/Engine/Source/Programs/Horde/Horde.Build/Jobs/JobTaskSource.cs b/Engine/Source/Programs/Horde/Horde.Build/Jobs/JobTaskSource.cs
--- a/Engine/Source/Programs/Horde/Horde.Build/Jobs/JobTaskSource.cs
+++ b/Engine/Source/Programs/Horde/Horde.Build/Jobs/JobTaskSource.cs
@@ -780,7 +780,7 @@ namespace Horde.Build.Jobs
 					if (waiter.LeaseSource.TrySetResult(lease))
 					{
 						_logger.LogDebug("Assigned lease {LeaseId} to agent {AgentId}", leaseId, agent.Id);
-						await _logFileService.CreateLogFileAsync(job.Id, agent.SessionId, LogType.Json, logId);
+						await _logFileService.CreateLogFileAsync(job.Id, agent.SessionId, LogType.Json, job.JobOptions?.UseNewLogStorage ?? false, logId);
 						return lease;
 					}
 				}
diff --git a/Engine/Source/Programs/Horde/Horde.Build/Logs/ILogFile.cs b/Engine/Source/Programs/Horde/Horde.Build/Logs/ILogFile.cs
--- a/Engine/Source/Programs/Horde/Horde.Build/Logs/ILogFile.cs
+++ b/Engine/Source/Programs/Horde/Horde.Build/Logs/ILogFile.cs
@@ -6,6 +6,7 @@ using EpicGames.Core;
 using EpicGames.Horde.Storage;
 using Horde.Build.Agents.Sessions;
 using Horde.Build.Jobs;
+using Horde.Build.Storage;
 using Horde.Build.Utilities;
 
 namespace Horde.Build.Logs
@@ -60,6 +61,11 @@ namespace Horde.Build.Logs
 		/// </summary>
 		public SessionId? SessionId { get; }
 
+		/// <summary>
+		/// Whether to use the new storage backend for log data
+		/// </summary>
+		public bool UseNewStorageBackend { get; }
+
 		/// <summary>
 		/// Maximum line index in the file
 		/// </summary>
diff --git a/Engine/Source/Programs/Horde/Horde.Build/Logs/ILogFileCollection.cs b/Engine/Source/Programs/Horde/Horde.Build/Logs/ILogFileCollection.cs
--- a/Engine/Source/Programs/Horde/Horde.Build/Logs/ILogFileCollection.cs
+++ b/Engine/Source/Programs/Horde/Horde.Build/Logs/ILogFileCollection.cs
@@ -58,10 +58,11 @@ namespace Horde.Build.Logs
 		/// <param name="jobId">Unique id of the job that owns this log file</param>
 		/// <param name="sessionId">Agent session allowed to update the log</param>
 		/// <param name="type">Type of events to be stored in the log</param>
+		/// <param name="newStorageBackend">Whether to use the new storage backend</param>
 		/// <param name="logId">ID of the log file (optional)</param>
 		/// <param name="cancellationToken">Cancellation token for the call</param>
 		/// <returns>The new log file document</returns>
-		Task<ILogFile> CreateLogFileAsync(JobId jobId, SessionId? sessionId, LogType type, LogId? logId = null, CancellationToken cancellationToken = default);
+		Task<ILogFile> CreateLogFileAsync(JobId jobId, SessionId? sessionId, LogType type, bool newStorageBackend, LogId? logId = null, CancellationToken cancellationToken = default);
 
 		/// <summary>
 		/// Updates the line count for a log file (v2 backend only)
diff --git a/Engine/Source/Programs/Horde/Horde.Build/Logs/LogFileCollection.cs b/Engine/Source/Programs/Horde/Horde.Build/Logs/LogFileCollection.cs
--- a/Engine/Source/Programs/Horde/Horde.Build/Logs/LogFileCollection.cs
+++ b/Engine/Source/Programs/Horde/Horde.Build/Logs/LogFileCollection.cs
@@ -19,7 +19,6 @@ namespace Horde.Build.Logs
 	using JobId = ObjectId<IJob>;
 	using LogId = ObjectId<ILogFile>;
 	using SessionId = ObjectId<ISession>;
-	using StreamId = StringId<IStream>;
 
 	/// <summary>
 	/// Wrapper around the jobs collection in a mongo DB
@@ -64,6 +63,7 @@ namespace Horde.Build.Logs
 
 			public SessionId? SessionId { get; set; }
 			public LogType Type { get; set; }
+			public bool UseNewStorageBackend { get; set; }
 
 			[BsonIgnoreIfNull]
 			public int? MaxLineIndex { get; set; }
@@ -86,12 +86,13 @@ namespace Horde.Build.Logs
 			{
 			}
 
-			public LogFileDocument(JobId jobId, SessionId? sessionId, LogType type, LogId? logId)
+			public LogFileDocument(JobId jobId, SessionId? sessionId, LogType type, bool newStorageBackend, LogId? logId)
 			{
 				Id = logId ?? LogId.GenerateNewId();
 				JobId = jobId;
 				SessionId = sessionId;
 				Type = type;
+				UseNewStorageBackend = newStorageBackend;
 				MaxLineIndex = 0;
 				RefName = new RefName(Id.ToString());
 			}
@@ -125,9 +126,9 @@ namespace Horde.Build.Logs
 		}
 
 		/// <inheritdoc/>
-		public async Task<ILogFile> CreateLogFileAsync(JobId jobId, SessionId? sessionId, LogType type, LogId? logId, CancellationToken cancellationToken)
+		public async Task<ILogFile> CreateLogFileAsync(JobId jobId, SessionId? sessionId, LogType type, bool newStorageBackend, LogId? logId, CancellationToken cancellationToken)
 		{
-			LogFileDocument newLogFile = new (jobId, sessionId, type, logId);
+			LogFileDocument newLogFile = new (jobId, sessionId, type, newStorageBackend, logId);
 			await _logFiles.InsertOneAsync(newLogFile, null, cancellationToken);
 			return newLogFile;
 		}
diff --git a/Engine/Source/Programs/Horde/Horde.Build/Logs/LogFileService.cs b/Engine/Source/Programs/Horde/Horde.Build/Logs/LogFileService.cs
--- a/Engine/Source/Programs/Horde/Horde.Build/Logs/LogFileService.cs
+++ b/Engine/Source/Programs/Horde/Horde.Build/Logs/LogFileService.cs
@@ -64,10 +64,11 @@ namespace Horde.Build.Logs
 		/// <param name="jobId">Unique id of the job that owns this log file</param>
 		/// <param name="sessionId">Agent session allowed to update the log</param>
 		/// <param name="type">Type of events to be stored in the log</param>
+		/// <param name="useNewStorageBackend">Whether to use the new storage backend for log data</param>
 		/// <param name="cancellationToken">Cancellation token for the call</param>
 		/// <param name="logId">ID of the log file (optional)</param>
 		/// <returns>The new log file document</returns>
-		Task<ILogFile> CreateLogFileAsync(JobId jobId, SessionId? sessionId, LogType type, LogId? logId = null, CancellationToken cancellationToken = default);
+		Task<ILogFile> CreateLogFileAsync(JobId jobId, SessionId? sessionId, LogType type, bool useNewStorageBackend, LogId? logId = null, CancellationToken cancellationToken = default);
 
 		/// <summary>
 		/// Gets a logfile by ID
@@ -600,9 +601,9 @@ namespace Horde.Build.Logs
 		}
 
 		/// <inheritdoc/>
-		public Task<ILogFile> CreateLogFileAsync(JobId jobId, SessionId? sessionId, LogType type, LogId? logId, CancellationToken cancellationToken)
+		public Task<ILogFile> CreateLogFileAsync(JobId jobId, SessionId? sessionId, LogType type, bool useNewStorageBackend, LogId? logId, CancellationToken cancellationToken)
 		{
-			return _logFiles.CreateLogFileAsync(jobId, sessionId, type, logId, cancellationToken);
+			return _logFiles.CreateLogFileAsync(jobId, sessionId, type, useNewStorageBackend, logId, cancellationToken);
 		}
 
 		/// <inheritdoc/>
@@ -649,7 +650,7 @@ namespace Horde.Build.Logs
 		{
 			List<Utf8String> lines = new List<Utf8String>();
 
-			if (_settings.Value.FeatureFlags.EnableNewLogger)
+			if (logFile.UseNewStorageBackend)
 			{
 				TreeReader reader = await GetTreeReaderAsync(cancellationToken);
 
@@ -926,7 +927,7 @@ namespace Horde.Build.Logs
 		public async Task<LogMetadata> GetMetadataAsync(ILogFile logFile, CancellationToken cancellationToken)
 		{
 			LogMetadata metadata = new LogMetadata();
-			if (_settings.Value.FeatureFlags.EnableNewLogger)
+			if (logFile.UseNewStorageBackend)
 			{
 				metadata.MaxLineIndex = logFile.LineCount;
 
@@ -1076,7 +1077,7 @@ namespace Horde.Build.Logs
 		/// <inheritdoc/>
 		public async Task<Stream> OpenRawStreamAsync(ILogFile logFile, long offset, long length, CancellationToken cancellationToken)
 		{
-			if (_settings.Value.FeatureFlags.EnableNewLogger)
+			if (logFile.UseNewStorageBackend)
 			{
 				TreeReader reader = await GetTreeReaderAsync(cancellationToken);
 
@@ -1226,7 +1227,7 @@ namespace Horde.Build.Logs
 		/// <inheritdoc/>
 		public async Task<(int, long)> GetLineOffsetAsync(ILogFile logFile, int lineIdx, CancellationToken cancellationToken)
 		{
-			if (_settings.Value.FeatureFlags.EnableNewLogger)
+			if (logFile.UseNewStorageBackend)
 			{
 				TreeReader reader = await GetTreeReaderAsync(cancellationToken);
 
@@ -1868,7 +1869,7 @@ namespace Horde.Build.Logs
 			List<int> results = new List<int>();
 			if (count > 0)
 			{
-				IAsyncEnumerable<int> enumerable =_settings.Value.FeatureFlags.EnableNewLogger ?
+				IAsyncEnumerable<int> enumerable = (logFile.UseNewStorageBackend) ?
 						SearchLogDataInternalNewAsync(logFile, text, firstLine, searchStats, cancellationToken) :
 						SearchLogDataInternalAsync(logFile, text, firstLine, searchStats);
 
diff --git a/Engine/Source/Programs/Horde/Horde.Build/Logs/LogsController.cs b/Engine/Source/Programs/Horde/Horde.Build/Logs/LogsController.cs
--- a/Engine/Source/Programs/Horde/Horde.Build/Logs/LogsController.cs
+++ b/Engine/Source/Programs/Horde/Horde.Build/Logs/LogsController.cs
@@ -55,7 +55,6 @@ namespace Horde.Build.Logs
 	{
 		private readonly ILogFileService _logFileService;
 		private readonly IIssueCollection _issueCollection;
-		private readonly AclService _aclService;
 		private readonly JobService _jobService;
 		private readonly StorageService _storageService;
 		private readonly IOptionsSnapshot<GlobalConfig> _globalConfig;
@@ -63,37 +62,15 @@ namespace Horde.Build.Logs
 		/// <summary>
 		/// Constructor
 		/// </summary>
-		public LogsController(ILogFileService logFileService, IIssueCollection issueCollection, AclService aclService, JobService jobService, StorageService storageService, IOptionsSnapshot<GlobalConfig> globalConfig)
+		public LogsController(ILogFileService logFileService, IIssueCollection issueCollection, JobService jobService, StorageService storageService, IOptionsSnapshot<GlobalConfig> globalConfig)
 		{
 			_logFileService = logFileService;
 			_issueCollection = issueCollection;
-			_aclService = aclService;
 			_jobService = jobService;
 			_storageService = storageService;
 			_globalConfig = globalConfig;
  		}
 
-		/// <summary>
-		/// Creates a new log file. This endpoint is used mainly for debugging; log documents for specific uses are usually 
-		/// created by the server and have their id passed into clients to append to.
-		/// </summary>
-		/// <param name="request">Request to create the log</param>
-		/// <param name="cancellationToken">Cancellation token for the request</param>
-		/// <returns>Information about the requested project</returns>
-		[HttpPost]
-		[Route("/api/v1/logs")]
-		[ProducesResponseType(typeof(GetLogFileResponse), 200)]
-		public async Task<ActionResult<object>> CreateLog(CreateLogFileRequest request, CancellationToken cancellationToken = default)
-		{
-			if (!_globalConfig.Value.Authorize(AclAction.CreateLog, User))
-			{
-				return Forbid();
-			}
-
-			ILogFile logFile = await _logFileService.CreateLogFileAsync(JobId.Empty, null, request.Type, cancellationToken: cancellationToken);
-			return new CreateLogFileResponse { Id = logFile.Id.ToString() };
-		}
-
 		/// <summary>
 		/// Retrieve metadata about a specific log file
 		/// </summary>
diff --git a/Engine/Source/Programs/Horde/Horde.Build/Server/RpcService.cs b/Engine/Source/Programs/Horde/Horde.Build/Server/RpcService.cs
--- a/Engine/Source/Programs/Horde/Horde.Build/Server/RpcService.cs
+++ b/Engine/Source/Programs/Horde/Horde.Build/Server/RpcService.cs
@@ -615,7 +615,7 @@ namespace Horde.Build.Server
 			}
 
 			// Create a log file if necessary
-			log.Value ??= await _logFileService.CreateLogFileAsync(job.Id, batch.SessionId, LogType.Json);
+			log.Value ??= await _logFileService.CreateLogFileAsync(job.Id, batch.SessionId, LogType.Json, job.JobOptions?.UseNewLogStorage ?? false);
 
 			// Get the node for this step
 			IGraph graph = await _jobService.GetGraphAsync(job);
diff --git a/Engine/Source/Programs/Horde/Horde.Build/ServerSettings.cs b/Engine/Source/Programs/Horde/Horde.Build/ServerSettings.cs
--- a/Engine/Source/Programs/Horde/Horde.Build/ServerSettings.cs
+++ b/Engine/Source/Programs/Horde/Horde.Build/ServerSettings.cs
@@ -212,10 +212,6 @@ namespace Horde.Build
 	/// </summary>
 	public class FeatureFlagSettings
 	{
-		/// <summary>
-		/// Whether to use the new log storage backend
-		/// </summary>
-		public bool EnableNewLogger { get; set; } = false;
 	}
 
 	/// <summary>
diff --git a/Engine/Source/Programs/Horde/Horde.Build/Tasks/ConformTaskSource.cs b/Engine/Source/Programs/Horde/Horde.Build/Tasks/ConformTaskSource.cs
--- a/Engine/Source/Programs/Horde/Horde.Build/Tasks/ConformTaskSource.cs
+++ b/Engine/Source/Programs/Horde/Horde.Build/Tasks/ConformTaskSource.cs
@@ -180,7 +180,7 @@ namespace Horde.Build.Tasks
 					LeaseId leaseId = LeaseId.GenerateNewId();
 					if (await AllocateConformLeaseAsync(agent.Id, task.Workspaces, leaseId))
 					{
-						ILogFile log = await _logService.CreateLogFileAsync(JobId.Empty, agent.SessionId, LogType.Json, cancellationToken: cancellationToken);
+						ILogFile log = await _logService.CreateLogFileAsync(JobId.Empty, agent.SessionId, LogType.Json, useNewStorageBackend: false, cancellationToken: cancellationToken);
 						task.LogId = log.Id.ToString();
 						task.RemoveUntrackedFiles = agent.RequestFullConform;
 
diff --git a/Engine/Source/Programs/Horde/Horde.Build/Tasks/RestartTaskSource.cs b/Engine/Source/Programs/Horde/Horde.Build/Tasks/RestartTaskSource.cs
--- a/Engine/Source/Programs/Horde/Horde.Build/Tasks/RestartTaskSource.cs
+++ b/Engine/Source/Programs/Horde/Horde.Build/Tasks/RestartTaskSource.cs
@@ -44,7 +44,7 @@ namespace Horde.Build.Tasks
 				return await DrainAsync(cancellationToken);
 			}
 
-			ILogFile log = await _logService.CreateLogFileAsync(JobId.Empty, agent.SessionId, LogType.Json, cancellationToken: cancellationToken);
+			ILogFile log = await _logService.CreateLogFileAsync(JobId.Empty, agent.SessionId, LogType.Json, useNewStorageBackend: false, cancellationToken: cancellationToken);
 
 			RestartTask task = new RestartTask();
 			task.LogId = log.Id.ToString();
diff --git a/Engine/Source/Programs/Horde/Horde.Build/Tasks/ShutdownTaskSource.cs b/Engine/Source/Programs/Horde/Horde.Build/Tasks/ShutdownTaskSource.cs
--- a/Engine/Source/Programs/Horde/Horde.Build/Tasks/ShutdownTaskSource.cs
+++ b/Engine/Source/Programs/Horde/Horde.Build/Tasks/ShutdownTaskSource.cs
@@ -43,7 +43,7 @@ namespace Horde.Build.Tasks
 				return await DrainAsync(cancellationToken);
 			}
 
-			ILogFile log = await _logService.CreateLogFileAsync(JobId.Empty, agent.SessionId, LogType.Json, cancellationToken: cancellationToken);
+			ILogFile log = await _logService.CreateLogFileAsync(JobId.Empty, agent.SessionId, LogType.Json, useNewStorageBackend: false, cancellationToken: cancellationToken);
 
 			ShutdownTask task = new ShutdownTask();
 			task.LogId = log.Id.ToString();
diff --git a/Engine/Source/Programs/Horde/Horde.Build/Tasks/UpgradeTaskSource.cs b/Engine/Source/Programs/Horde/Horde.Build/Tasks/UpgradeTaskSource.cs
--- a/Engine/Source/Programs/Horde/Horde.Build/Tasks/UpgradeTaskSource.cs
+++ b/Engine/Source/Programs/Horde/Horde.Build/Tasks/UpgradeTaskSource.cs
@@ -63,7 +63,7 @@ namespace Horde.Build.Tasks
 				return await DrainAsync(cancellationToken);
 			}
 
-			ILogFile logFile = await _logService.CreateLogFileAsync(JobId.Empty, agent.SessionId, LogType.Json, cancellationToken: cancellationToken);
+			ILogFile logFile = await _logService.CreateLogFileAsync(JobId.Empty, agent.SessionId, LogType.Json, useNewStorageBackend: false, cancellationToken: cancellationToken);
 
 			UpgradeTask task = new UpgradeTask();
 			task.SoftwareId = requiredVersion.ToString();
