Horde: Increase verbosity of log flush messages.

diff --git a/Engine/Source/Programs/Horde/Horde.Agent/Utility/JsonRpcLogSink.cs b/Engine/Source/Programs/Horde/Horde.Agent/Utility/JsonRpcLogSink.cs
--- a/Engine/Source/Programs/Horde/Horde.Agent/Utility/JsonRpcLogSink.cs
+++ b/Engine/Source/Programs/Horde/Horde.Agent/Utility/JsonRpcLogSink.cs
@@ -211,7 +211,7 @@ namespace Horde.Agent.Utility
 
 		protected virtual async Task UpdateLogAsync(NodeHandle target, int lineCount, CancellationToken cancellationToken)
 		{
-			_logger.LogDebug("Updating log {LogId} to line {LineCount}, target {Locator}", _logId, lineCount, target);
+			_logger.LogInformation("Updating log {LogId} to line {LineCount}, target {Locator}", _logId, lineCount, target);
 
 			UpdateLogRequest request = new UpdateLogRequest();
 			request.LogId = _logId;
