Horde: Replace references to abort
#rnx
#skipci

diff --git a/Engine/Source/Programs/Horde/HordeDashboard/attic/JobDetailHealth.tsx b/Engine/Source/Programs/Horde/HordeDashboard/attic/JobDetailHealth.tsx
deleted file mode 100644
diff --git a/Engine/Source/Programs/Horde/HordeDashboard/attic/JobDetailOverview.tsx b/Engine/Source/Programs/Horde/HordeDashboard/attic/JobDetailOverview.tsx
deleted file mode 100644
diff --git a/Engine/Source/Programs/Horde/HordeDashboard/attic/JobDetailStep.tsx b/Engine/Source/Programs/Horde/HordeDashboard/attic/JobDetailStep.tsx
deleted file mode 100644
diff --git a/Engine/Source/Programs/Horde/HordeDashboard/attic/JobDetailView.tsx b/Engine/Source/Programs/Horde/HordeDashboard/attic/JobDetailView.tsx
deleted file mode 100644
diff --git a/Engine/Source/Programs/Horde/HordeDashboard/src/backend/JobDetails.ts b/Engine/Source/Programs/Horde/HordeDashboard/src/backend/JobDetails.ts
--- a/Engine/Source/Programs/Horde/HordeDashboard/src/backend/JobDetails.ts
+++ b/Engine/Source/Programs/Horde/HordeDashboard/src/backend/JobDetails.ts
@@ -1188,16 +1188,16 @@ export const getStepSummaryMarkdown = (jobDetails: JobDetails, stepId: string):
         eta.display = eta.server = "";
         let aborted = "";
         if (step.abortedByUserInfo) {
-            aborted = "This step was aborted";
+            aborted = "This step was canceled";
             aborted += ` by ${step.abortedByUserInfo.name}.`;
         } else if (jobDetails.jobdata?.abortedByUserInfo) {
-            aborted = "The job was aborted";
+            aborted = "The job was canceled";
             aborted += ` by ${jobDetails.jobdata?.abortedByUserInfo.name}.`;
         } else {
-            aborted = "The step was aborted";
+            aborted = "The step was canceled";
             
             if (step.error === JobStepError.TimedOut) {
-                aborted = "The step was aborted due to reaching the maximum run time limit";
+                aborted = "The step was canceled due to reaching the maximum run time limit";
             }
         }
         text.push(aborted);
diff --git a/Engine/Source/Programs/Horde/HordeDashboard/src/components/ChangeButton.tsx b/Engine/Source/Programs/Horde/HordeDashboard/src/components/ChangeButton.tsx
--- a/Engine/Source/Programs/Horde/HordeDashboard/src/components/ChangeButton.tsx
+++ b/Engine/Source/Programs/Horde/HordeDashboard/src/components/ChangeButton.tsx
@@ -180,7 +180,7 @@ export const ChangeButton: React.FC<{ job?: JobData, stepRef?: GetJobStepRefResp
    return (<Stack verticalAlign="center" verticalFill={true} horizontalAlign="start"> <div style={{ paddingBottom: "1px" }}>
       <Stack tokens={{ childrenGap: 4 }}>
          <span ref={spanRef} style={{ padding: "2px 6px 2px 6px", height: "15px", cursor: "pointer" }} className={job.startedByUserInfo ? "cl-callout-button-user" : "cl-callout-button"} onClick={(ev) => { ev.stopPropagation(); ev.preventDefault(); setMenuShown(!menuShown) }} >{change}</span>
-         {!hideAborted && <span ref={spanRef} style={{ padding: "2px 6px 2px 6px", height: "16px", cursor: "pointer", userSelect: "none", fontFamily: "Horde Open Sans SemiBold", fontSize: "10px", backgroundColor: colors.get(StatusColor.Failure), color: "rgb(255, 255, 255)" }} onClick={(ev) => { ev.preventDefault(); setMenuShown(!menuShown) }}>Aborted</span>}
+         {!hideAborted && <span ref={spanRef} style={{ padding: "2px 6px 2px 6px", height: "16px", cursor: "pointer", userSelect: "none", fontFamily: "Horde Open Sans SemiBold", fontSize: "10px", backgroundColor: colors.get(StatusColor.Failure), color: "rgb(255, 255, 255)" }} onClick={(ev) => { ev.preventDefault(); setMenuShown(!menuShown) }}>Canceled</span>}
       </Stack>
       {menuShown && <ChangeContextMenu target={{ ref: spanRef }} job={job} commit={commit} stepRef={stepRef} rangeCL={rangeCL} onDismiss={() => setMenuShown(false)} />}
    </div></Stack>);
diff --git a/Engine/Source/Programs/Horde/HordeDashboard/src/components/DashboardView.tsx b/Engine/Source/Programs/Horde/HordeDashboard/src/components/DashboardView.tsx
--- a/Engine/Source/Programs/Horde/HordeDashboard/src/components/DashboardView.tsx
+++ b/Engine/Source/Programs/Horde/HordeDashboard/src/components/DashboardView.tsx
@@ -35,7 +35,7 @@ const statusText = new Map<StatusColor, string>([
    [StatusColor.Waiting, "Waiting"],
    [StatusColor.Ready, "Ready"],
    [StatusColor.Skipped, "Skipped"],
-   [StatusColor.Aborted, "Aborted"],
+   [StatusColor.Aborted, "Canceled"],
    [StatusColor.Running, "Running"]
 ]);
 
diff --git a/Engine/Source/Programs/Horde/HordeDashboard/src/components/JobDetailCommon.tsx b/Engine/Source/Programs/Horde/HordeDashboard/src/components/JobDetailCommon.tsx
--- a/Engine/Source/Programs/Horde/HordeDashboard/src/components/JobDetailCommon.tsx
+++ b/Engine/Source/Programs/Horde/HordeDashboard/src/components/JobDetailCommon.tsx
@@ -22,7 +22,7 @@ export const getStepStatusMessage = (step: StepData) => {
    let message = "";
 
    if (step.abortRequested && step.state !== JobStepState.Aborted) {
-      return "Abort Requested";
+      return "Cancel Requested";
    }
 
    switch (step.state) {
@@ -33,7 +33,7 @@ export const getStepStatusMessage = (step: StepData) => {
          message = "Skipped";
          break;
       case JobStepState.Aborted:
-         message = "Aborted";
+         message = "Canceled";
          break;
       case JobStepState.Ready:
          message = "Ready";
diff --git a/Engine/Source/Programs/Horde/HordeDashboard/src/components/JobOperationsContextMenu.tsx b/Engine/Source/Programs/Horde/HordeDashboard/src/components/JobOperationsContextMenu.tsx
--- a/Engine/Source/Programs/Horde/HordeDashboard/src/components/JobOperationsContextMenu.tsx
+++ b/Engine/Source/Programs/Horde/HordeDashboard/src/components/JobOperationsContextMenu.tsx
@@ -94,7 +94,7 @@ export const JobOperationsContextMenu: React.FC<{ job: GetJobResponse, children?
    if (!abortDisabled) {
       menuItems.push({
          key: 'abort_job',
-         text: 'Abort Job',
+         text: 'Cancel Job',
          onClick: () => setAbortShown(true),
       });
    }
diff --git a/Engine/Source/Programs/Horde/HordeDashboard/src/components/jobDetailsV2/AbortJobModal.tsx b/Engine/Source/Programs/Horde/HordeDashboard/src/components/jobDetailsV2/AbortJobModal.tsx
--- a/Engine/Source/Programs/Horde/HordeDashboard/src/components/jobDetailsV2/AbortJobModal.tsx
+++ b/Engine/Source/Programs/Horde/HordeDashboard/src/components/jobDetailsV2/AbortJobModal.tsx
@@ -11,14 +11,14 @@ export const AbortJobModal: React.FC<{ jobDetails?: JobDetailsV2; jobDataIn?: Ge
 
    const jobData = jobDetails?.jobData ?? jobDataIn;
    if (!jobData) {
-      console.error("Attempting to display abort job modal with no job data")
+      console.error("Attempting to display cancel job modal with no job data")
       return null;
    }
 
-   let headerText = `Abort ${jobData.name} Job?`;
+   let headerText = `Cancel ${jobData.name} Job?`;
 
    if (jobDetails && stepId) {
-      headerText = `Abort ${jobDetails.getStepName(stepId)} Step?`;
+      headerText = `Cancel ${jobDetails.getStepName(stepId)} Step?`;
    }
 
    const close = () => {
@@ -29,7 +29,7 @@ export const AbortJobModal: React.FC<{ jobDetails?: JobDetailsV2; jobDataIn?: Ge
 
       if (!stepId) {
          await backend.updateJob(jobData.id, { aborted: true }).then((response) => {
-            console.log("Job aborted", response);
+            console.log("Job canceled", response);
          }).catch((reason) => {
             // @todo: error ui
             console.error(reason);
@@ -41,12 +41,12 @@ export const AbortJobModal: React.FC<{ jobDetails?: JobDetailsV2; jobDataIn?: Ge
          const batch = jobDetails?.batchByStepId(stepId);
 
          if (!batch) {
-            console.error("Unable to get batch id for step abort");
+            console.error("Unable to get batch id for step cancel");
             return;
          }
 
          await backend.updateJobStep(jobData.id, batch.id, stepId, { abortRequested: true }).then((response) => {
-            console.log("Job step aborted", response);
+            console.log("Job step canceled", response);
          }).catch((reason) => {
             // @todo: error ui
             console.error(reason);
@@ -77,8 +77,8 @@ export const AbortJobModal: React.FC<{ jobDetails?: JobDetailsV2; jobDataIn?: Ge
       <Stack styles={{ root: { padding: 8 } }}>
          <Stack horizontal tokens={{ childrenGap: 16 }} styles={{ root: { paddingTop: 12, paddingLeft: 8, paddingBottom: 8 } }}>
             <Stack grow />
-            <PrimaryButton text="Abort" disabled={false} onClick={() => { onAbort(); }} />
-            <DefaultButton text="Cancel" disabled={false} onClick={() => { close(); }} />
+            <PrimaryButton text={(jobDetails && stepId) ? "Cancel Step" : "Cancel Job"} disabled={false} onClick={() => { onAbort(); }} />
+            <DefaultButton text="Close" disabled={false} onClick={() => { close(); }} />
          </Stack>
       </Stack>
    </Modal>;
diff --git a/Engine/Source/Programs/Horde/HordeDashboard/src/components/jobDetailsV2/JobDetailViewSteps.tsx b/Engine/Source/Programs/Horde/HordeDashboard/src/components/jobDetailsV2/JobDetailViewSteps.tsx
--- a/Engine/Source/Programs/Horde/HordeDashboard/src/components/jobDetailsV2/JobDetailViewSteps.tsx
+++ b/Engine/Source/Programs/Horde/HordeDashboard/src/components/jobDetailsV2/JobDetailViewSteps.tsx
@@ -727,13 +727,13 @@ export const getStepSummaryMarkdown = (jobDetails: JobDetailsV2, stepId: string)
       eta.display = eta.server = "";
       let aborted = "";
       if (step.abortedByUserInfo) {
-         aborted = "This step was aborted";
+         aborted = "This step was canceled";
          aborted += ` by ${step.abortedByUserInfo.name}.`;
       } else if (jobData.abortedByUserInfo) {
-         aborted = "The job was aborted";
+         aborted = "The job was canceled";
          aborted += ` by ${jobData.abortedByUserInfo.name}.`;
       } else {
-         aborted = "The step was aborted";
+         aborted = "The step was canceled";
       }
       text.push(aborted);
    } else if (step.state === JobStepState.Skipped) {
diff --git a/Engine/Source/Programs/Horde/HordeDashboard/src/components/jobDetailsV2/JobDetailViewSummary.tsx b/Engine/Source/Programs/Horde/HordeDashboard/src/components/jobDetailsV2/JobDetailViewSummary.tsx
--- a/Engine/Source/Programs/Horde/HordeDashboard/src/components/jobDetailsV2/JobDetailViewSummary.tsx
+++ b/Engine/Source/Programs/Horde/HordeDashboard/src/components/jobDetailsV2/JobDetailViewSummary.tsx
@@ -202,7 +202,7 @@ export const SummaryPanel: React.FC<{ jobDetails: JobDetailsV2 }> = observer(({
    let jobText = `Job created ${timeStr} by ${jobData.startedByUserInfo ? jobData.startedByUserInfo.name : "scheduler"} and `;
 
    if (jobDetails.aborted) {
-      jobText += "was aborted";
+      jobText += "was canceled";
       if (jobData.abortedByUserInfo) {
          jobText += ` by ${jobData.abortedByUserInfo.name}.`;
       }
diff --git a/Engine/Source/Programs/Horde/HordeDashboard/src/components/jobDetailsV2/JobDetailsViewCommon.tsx b/Engine/Source/Programs/Horde/HordeDashboard/src/components/jobDetailsV2/JobDetailsViewCommon.tsx
--- a/Engine/Source/Programs/Horde/HordeDashboard/src/components/jobDetailsV2/JobDetailsViewCommon.tsx
+++ b/Engine/Source/Programs/Horde/HordeDashboard/src/components/jobDetailsV2/JobDetailsViewCommon.tsx
@@ -1057,7 +1057,7 @@ export const JobFilterBar: React.FC<{ jobDetails: JobDetailsV2 }> = observer(({
    const stateItems: StateItem[] = ["All", "Failure", "Warnings", "Skipped", "Running", "Completed", "Aborted", "Waiting", "Ready"].map(state => {
       return {
          key: state,
-         text: state,
+         text: state === "Aborted" ? "Canceled" : state,
          state: state as StateFilter
       };
    });
diff --git a/Engine/Source/Programs/Horde/HordeDashboard/src/components/jobDetailsV2/JobOperationsBar.tsx b/Engine/Source/Programs/Horde/HordeDashboard/src/components/jobDetailsV2/JobOperationsBar.tsx
--- a/Engine/Source/Programs/Horde/HordeDashboard/src/components/jobDetailsV2/JobOperationsBar.tsx
+++ b/Engine/Source/Programs/Horde/HordeDashboard/src/components/jobDetailsV2/JobOperationsBar.tsx
@@ -74,7 +74,7 @@ export const JobOperations: React.FC<{ jobDetails: JobDetailsV2 }> = observer(({
 
    opsList.push({
       key: 'jobops_abort',
-      text: "Abort",
+      text: "Cancel Job",
       disabled: abortDisabled,
       iconProps: { iconName: "Delete" },
       onClick: () => { setAbortShown(true); }
@@ -206,7 +206,7 @@ const StepOperations: React.FC<{ jobDetails: JobDetailsV2, stepId: string }> = o
 
    opsList.push({
       key: 'stepops_abort',
-      text: "Abort",
+      text: "Cancel Step",
       iconProps: { iconName: "Cross" },
       disabled: !!step.finishTime,
       onClick: () => { setShown({ abortShown: true }); }
diff --git a/Engine/Source/Programs/Horde/HordeDashboard/src/components/jobDetailsV2/StepDetailSummary.tsx b/Engine/Source/Programs/Horde/HordeDashboard/src/components/jobDetailsV2/StepDetailSummary.tsx
--- a/Engine/Source/Programs/Horde/HordeDashboard/src/components/jobDetailsV2/StepDetailSummary.tsx
+++ b/Engine/Source/Programs/Horde/HordeDashboard/src/components/jobDetailsV2/StepDetailSummary.tsx
@@ -94,16 +94,16 @@ const getStepSummaryMarkdown = (jobDetails: JobDetailsV2, stepId: string): strin
       eta.display = eta.server = "";
       let aborted = "";
       if (step.abortedByUserInfo) {
-         aborted = "This step was aborted";
+         aborted = "This step was canceled";
          aborted += ` by ${step.abortedByUserInfo.name}.`;
       } else if (jobData?.abortedByUserInfo) {
-         aborted = "The job was aborted";
+         aborted = "The job was canceled";
          aborted += ` by ${jobData?.abortedByUserInfo.name}.`;
       } else {
-         aborted = "The step was aborted";
+         aborted = "The step was canceled";
 
          if (step.error === JobStepError.TimedOut) {
-            aborted = "The step was aborted due to reaching the maximum run time limit";
+            aborted = "The step was canceled due to reaching the maximum run time limit";
          }
       }
       text.push(aborted);
