Horde: Disable internal trace logger for new log sink.

diff --git a/Engine/Source/Programs/Horde/Horde.Agent/Utility/JsonRpcLogSink.cs b/Engine/Source/Programs/Horde/Horde.Agent/Utility/JsonRpcLogSink.cs
--- a/Engine/Source/Programs/Horde/Horde.Agent/Utility/JsonRpcLogSink.cs
+++ b/Engine/Source/Programs/Horde/Horde.Agent/Utility/JsonRpcLogSink.cs
@@ -100,7 +100,7 @@ namespace Horde.Agent.Utility
 			_logId = logId;
 			_builder = new LogBuilder(LogFormat.Json, logger);
 			_inner = inner;
-			_writer = new TreeWriter(store, traceLogger: logger);
+			_writer = new TreeWriter(store);
 			_logger = logger;
 
 			_tailTaskStop = new AsyncEvent();
