Horde: Remove single instance setting from intial msi protoyping
#rnx

diff --git a/Engine/Source/Programs/Horde/Horde.Build/Agents/Software/AgentSoftwareService.cs b/Engine/Source/Programs/Horde/Horde.Build/Agents/Software/AgentSoftwareService.cs
--- a/Engine/Source/Programs/Horde/Horde.Build/Agents/Software/AgentSoftwareService.cs
+++ b/Engine/Source/Programs/Horde/Horde.Build/Agents/Software/AgentSoftwareService.cs
@@ -281,7 +281,7 @@ namespace Horde.Build.Agents.Software
 
 			// Check whether we have an installed agent zip
 			FileReference agentZip = FileReference.Combine(Program.AppDir, "DefaultAgent/Agent.zip");
-			if (!_settings.CurrentValue.SingleInstance || !FileReference.Exists(agentZip))
+			if (!FileReference.Exists(agentZip))
 			{
 				return;
 			}
diff --git a/Engine/Source/Programs/Horde/Horde.Build/Program.cs b/Engine/Source/Programs/Horde/Horde.Build/Program.cs
--- a/Engine/Source/Programs/Horde/Horde.Build/Program.cs
+++ b/Engine/Source/Programs/Horde/Horde.Build/Program.cs
@@ -151,8 +151,6 @@ namespace Horde.Build
 			ServerSettings hordeSettings = new ServerSettings();
 			config.GetSection("Horde").Bind(hordeSettings);
 
-			InitializeDefaults(hordeSettings);
-
 			DirectoryReference logDir = AppDir;
 			if (RuntimeInformation.IsOSPlatform(OSPlatform.Windows))
 			{
@@ -263,54 +261,5 @@ namespace Horde.Build
 
 			return builder.AddEnvironmentVariables().Build();
 		}
-
-		/// <summary>
-		/// Handles bootstrapping of defaults for local servers, which can't be generated during build/installation process (or are better handled here where they can be updated)
-		/// This stuff will change as we get settings into database and could be considered discovery for installer/dockerfile builds 
-		/// </summary>		
-		static void InitializeDefaults(ServerSettings settings)
-		{			
-			if (settings.SingleInstance)
-			{
-				FileReference globalConfig = FileReference.Combine(Program.DataDir, "Config/globals.json");
-
-				if (!FileReference.Exists(globalConfig))
-				{
-					DirectoryReference.CreateDirectory(globalConfig.Directory);
-					FileReference.WriteAllText(globalConfig, "{}");
-				}
-
-				FileReference privateCertFile = FileReference.Combine(Program.DataDir, "Agent/ServerToAgent.pfx");
-				string privateCertFileJsonPath = privateCertFile.ToString().Replace("\\", "/", StringComparison.Ordinal);
-
-				if (!FileReference.Exists(UserConfigFile))
-				{
-					// create new user configuration
-					DirectoryReference.CreateDirectory(UserConfigFile.Directory);
-					FileReference.WriteAllText(UserConfigFile, $"{{\"Horde\": {{ \"ConfigPath\" : \"{globalConfig.ToString().Replace("\\", "/", StringComparison.Ordinal)}\", \"ServerPrivateCert\" : \"{privateCertFileJsonPath}\", \"HttpPort\": 8080}}}}");
-				}
-
-				// make sure the cert exists
-				if (!FileReference.Exists(privateCertFile))
-				{
-					string dnsName = System.Net.Dns.GetHostName();
-					Serilog.Log.Logger.Information("Creating certificate for {DnsName}", dnsName);
-
-					byte[] privateCertData = CertificateUtils.CreateSelfSignedCert(dnsName, "Horde Server");
-					
-					Serilog.Log.Logger.Information("Writing private cert: {PrivateCert}", privateCertFile.FullName);
-
-					if (!DirectoryReference.Exists(privateCertFile.Directory))
-					{
-						DirectoryReference.CreateDirectory(privateCertFile.Directory);
-					}
-
-					FileReference.WriteAllBytes(privateCertFile, privateCertData);
-				}
-
-				// note: this isn't great, though we need it early in server startup, and this is only hit on first server boot where the grpc cert isn't generated/set 
-				settings.ServerPrivateCert ??= privateCertFile.ToString();
-			}
-		}
 	}
 }
diff --git a/Engine/Source/Programs/Horde/Horde.Build/Server/ServerController.cs b/Engine/Source/Programs/Horde/Horde.Build/Server/ServerController.cs
--- a/Engine/Source/Programs/Horde/Horde.Build/Server/ServerController.cs
+++ b/Engine/Source/Programs/Horde/Horde.Build/Server/ServerController.cs
@@ -49,7 +49,7 @@ namespace Horde.Build.Server
 		[ProducesResponseType(typeof(GetServerInfoResponse), 200)]
 		public ActionResult<GetServerInfoResponse> GetServerInfo()
 		{
-			return new GetServerInfoResponse(_settings.CurrentValue.SingleInstance);
+			return new GetServerInfoResponse();
 		}
 	}
 }
diff --git a/Engine/Source/Programs/Horde/Horde.Build/Server/ServerMessages.cs b/Engine/Source/Programs/Horde/Horde.Build/Server/ServerMessages.cs
--- a/Engine/Source/Programs/Horde/Horde.Build/Server/ServerMessages.cs
+++ b/Engine/Source/Programs/Horde/Horde.Build/Server/ServerMessages.cs
@@ -22,20 +22,14 @@ namespace Horde.Build.Server
         /// </summary>
         public string OsDescription { get; set; }
 
-		/// <summary>
-		/// Whether this is an installed Horde build
-		/// </summary>
-		public bool SingleInstance { get; set; }
-
 		/// <summary>
 		/// Constructor
 		/// </summary>
-		public GetServerInfoResponse( bool singleInstance)
+		public GetServerInfoResponse()
         {
             FileVersionInfo versionInfo = FileVersionInfo.GetVersionInfo(Assembly.GetExecutingAssembly().Location);		
 			ServerVersion = versionInfo.ProductVersion ?? String.Empty;			
-			OsDescription = RuntimeInformation.OSDescription;
-			SingleInstance = singleInstance;
+			OsDescription = RuntimeInformation.OSDescription;			
 		}
 	}
 }
diff --git a/Engine/Source/Programs/Horde/Horde.Build/ServerSettings.cs b/Engine/Source/Programs/Horde/Horde.Build/ServerSettings.cs
--- a/Engine/Source/Programs/Horde/Horde.Build/ServerSettings.cs
+++ b/Engine/Source/Programs/Horde/Horde.Build/ServerSettings.cs
@@ -273,11 +273,6 @@ namespace Horde.Build
 		/// </summary>
 		public int Http2Port { get; set; }
 
-		/// <summary>
-		/// Whether the server is running as a single instance or with multiple instances, such as in k8s
-		/// </summary>
-		public bool SingleInstance { get; set; } = false;
-
 		/// <summary>
 		/// MongoDB connection string
 		/// </summary>
