Horde: Add a mechanism for naming security scopes, allowing them to be stored alongside objects without having to have each system understand the structure of things that can consume it.

For logs and artifacts, this allows creating objects that inherit permissions from other objects without having to embed custom authorization logic in each.

Scope names are slash-delimited strings, eg. "horde/p:ue5/s:ue5-main" for the "ue5-main" stream in the "ue5" project. Type signifiers (eg. "p", "s") are used to allow other named scopes to be added in future.

diff --git a/Engine/Source/Programs/Horde/Horde.Build/Acls/IAclScope.cs b/Engine/Source/Programs/Horde/Horde.Build/Acls/IAclScope.cs
new file mode 100644
--- /dev/null
+++ b/Engine/Source/Programs/Horde/Horde.Build/Acls/IAclScope.cs
@@ -0,0 +1,122 @@
+// Copyright Epic Games, Inc. All Rights Reserved.
+
+using System;
+using System.Diagnostics;
+using System.Security.Claims;
+using System.Text.Json;
+using System.Text.Json.Serialization;
+using MongoDB.Bson.Serialization;
+using MongoDB.Bson.Serialization.Attributes;
+using MongoDB.Bson.Serialization.Serializers;
+
+namespace Horde.Build.Acls
+{
+	/// <summary>
+	/// Interface for parent classes that implement authorization
+	/// </summary>
+	public interface IAclScope
+	{
+		/// <summary>
+		/// The parent scope object
+		/// </summary>
+		IAclScope? ParentScope { get; }
+
+		/// <summary>
+		/// Name of this scope
+		/// </summary>
+		AclScopeName ScopeName { get; }
+
+		/// <summary>
+		/// ACL for this scope
+		/// </summary>
+		AclConfig? Acl { get; }
+	}
+
+	/// <summary>
+	/// Name of an ACL scope
+	/// </summary>
+	[DebuggerDisplay("{Name}")]
+	[JsonConverter(typeof(AclScopeNameJsonConverter))]
+	[BsonSerializer(typeof(AclScopeNameBsonSerializer))]
+	public record struct AclScopeName(string Text)
+	{
+		/// <summary>
+		/// The root scope name
+		/// </summary>
+		public static AclScopeName Root { get; } = new AclScopeName("horde");
+
+		/// <summary>
+		/// Append another name to this scope
+		/// </summary>
+		/// <param name="name">Name to append</param>
+		/// <returns>New scope name</returns>
+		public AclScopeName Append(string name) => new AclScopeName($"{Text}/{name}");
+
+		/// <summary>
+		/// Append another name to this scope
+		/// </summary>
+		/// <param name="type">Type of the scope</param>
+		/// <param name="name">Name to append</param>
+		/// <returns>New scope name</returns>
+		public AclScopeName Append(string type, string name) => new AclScopeName($"{Text}/{type}:{name}");
+
+		/// <inheritdoc/>
+		public bool Equals(AclScopeName other) => Text.Equals(other.Text, StringComparison.Ordinal);
+
+		/// <inheritdoc/>
+		public override int GetHashCode() => Text.GetHashCode(StringComparison.Ordinal);
+	}
+
+	/// <summary>
+	/// Serializes <see cref="AclScopeName"/> objects to JSON
+	/// </summary>
+	class AclScopeNameJsonConverter : JsonConverter<AclScopeName>
+	{
+		/// <inheritdoc/>
+		public override AclScopeName Read(ref Utf8JsonReader reader, Type typeToConvert, JsonSerializerOptions options) => new AclScopeName(reader.GetString() ?? String.Empty);
+
+		/// <inheritdoc/>
+		public override void Write(Utf8JsonWriter writer, AclScopeName value, JsonSerializerOptions options) => writer.WriteStringValue(value.Text);
+	}
+
+	/// <summary>
+	/// Serializes <see cref="AclScopeName"/> objects to JSON
+	/// </summary>
+	class AclScopeNameBsonSerializer : SerializerBase<AclScopeName> 
+	{
+		/// <inheritdoc/>
+		public override AclScopeName Deserialize(BsonDeserializationContext context, BsonDeserializationArgs args) => new AclScopeName(context.Reader.ReadString());
+
+		/// <inheritdoc/>
+		public override void Serialize(BsonSerializationContext context, BsonSerializationArgs args, AclScopeName value) => context.Writer.WriteString(value.Text);
+	}
+
+	/// <summary>
+	/// Extension methods for ACL scopes
+	/// </summary>
+	public static class AclScopeExtensions
+	{
+		/// <summary>
+		/// Authorize a particular operation against a scope
+		/// </summary>
+		/// <param name="scope">Scope to query</param>
+		/// <param name="action">Action to check</param>
+		/// <param name="principal">Principal to authorize</param>
+		/// <returns>True if the given principal is allowed to perform a particular action</returns>
+		public static bool Authorize(this IAclScope scope, AclAction action, ClaimsPrincipal principal)
+		{
+			for (IAclScope? next = scope; next != null; next = next.ParentScope)
+			{
+				if (next.Acl != null)
+				{
+					bool? result = next.Acl.Authorize(action, principal);
+					if (result.HasValue)
+					{
+						return result.Value;
+					}
+				}
+			}
+			return false;
+		}
+	}
+}
diff --git a/Engine/Source/Programs/Horde/Horde.Build/Projects/ProjectConfig.cs b/Engine/Source/Programs/Horde/Horde.Build/Projects/ProjectConfig.cs
--- a/Engine/Source/Programs/Horde/Horde.Build/Projects/ProjectConfig.cs
+++ b/Engine/Source/Programs/Horde/Horde.Build/Projects/ProjectConfig.cs
@@ -28,7 +28,7 @@ namespace Horde.Build.Projects
 	[JsonSchemaCatalog("Horde Project", "Horde project configuration file", new[] { "*.project.json", "Projects/*.json" })]
 	[ConfigIncludeRoot]
 	[DebuggerDisplay("{Id}")]
-	public class ProjectConfig
+	public class ProjectConfig : IAclScope
 	{
 		/// <summary>
 		/// Accessor for the global config owning this project
@@ -36,6 +36,14 @@ namespace Horde.Build.Projects
 		[JsonIgnore]
 		public GlobalConfig GlobalConfig { get; private set; } = null!;
 
+		/// <inheritdoc/>
+		[JsonIgnore]
+		public IAclScope? ParentScope => GlobalConfig;
+
+		/// <inheritdoc/>
+		[JsonIgnore]
+		public AclScopeName ScopeName { get; private set; }
+
 		/// <summary>
 		/// The project id
 		/// </summary>
@@ -101,22 +109,13 @@ namespace Horde.Build.Projects
 		{
 			Id = id;
 			GlobalConfig = globalConfig;
+			ScopeName = ParentScope.ScopeName.Append("p", Id.ToString());
 
 			foreach (StreamConfig stream in Streams)
 			{
 				stream.PostLoad(stream.Id, this);
 			}
 		}
-
-		/// <summary>
-		/// Authorizes a user to perform a given action
-		/// </summary>
-		/// <param name="action">The action being performed</param>
-		/// <param name="user">The principal to validate</param>
-		public bool Authorize(AclAction action, ClaimsPrincipal user)
-		{
-			return Acl?.Authorize(action, user) ?? GlobalConfig.Authorize(action, user);
-		}
 	}
 
 	/// <summary>
diff --git a/Engine/Source/Programs/Horde/Horde.Build/Server/GlobalConfig.cs b/Engine/Source/Programs/Horde/Horde.Build/Server/GlobalConfig.cs
--- a/Engine/Source/Programs/Horde/Horde.Build/Server/GlobalConfig.cs
+++ b/Engine/Source/Programs/Horde/Horde.Build/Server/GlobalConfig.cs
@@ -51,7 +51,7 @@ namespace Horde.Build.Server
 	[JsonSchema("https://unrealengine.com/horde/global")]
 	[JsonSchemaCatalog("Horde Globals", "Horde global configuration file", "globals.json")]
 	[ConfigIncludeRoot]
-	public class GlobalConfig
+	public class GlobalConfig : IAclScope
 	{
 		/// <summary>
 		/// Global server settings object
@@ -59,6 +59,14 @@ namespace Horde.Build.Server
 		[JsonIgnore]
 		public ServerSettings ServerSettings { get; private set; } = null!;
 
+		/// <inheritdoc/>
+		[JsonIgnore]
+		public IAclScope ParentScope => ServerSettings;
+
+		/// <inheritdoc/>
+		[JsonIgnore]
+		public AclScopeName ScopeName => ServerSettings.ScopeName;
+
 		/// <summary>
 		/// Unique identifier for this config revision. Useful to detect changes.
 		/// </summary>
@@ -140,6 +148,7 @@ namespace Horde.Build.Server
 		private readonly Dictionary<StreamId, StreamConfig> _streamLookup = new Dictionary<StreamId, StreamConfig>();
 		private readonly Dictionary<ToolId, ToolConfig> _toolLookup = new Dictionary<ToolId, ToolConfig>();
 		private readonly Dictionary<ClusterId, ComputeClusterConfig> _computeClusterLookup = new Dictionary<ClusterId, ComputeClusterConfig>();
+		private readonly Dictionary<AclScopeName, IAclScope> _aclScopeLookup = new Dictionary<AclScopeName, IAclScope>();
 
 		/// <summary>
 		/// Called after the config file has been read
@@ -177,6 +186,21 @@ namespace Horde.Build.Server
 				computeCluster.PostLoad(this);
 			}
 
+			_aclScopeLookup.Clear();
+			_aclScopeLookup.Add(ScopeName, this);
+			foreach (ProjectConfig project in Projects)
+			{
+				_aclScopeLookup.Add(project.ScopeName, project);
+				foreach (StreamConfig stream in project.Streams)
+				{
+					_aclScopeLookup.Add(stream.ScopeName, stream);
+					foreach (TemplateRefConfig template in stream.Templates)
+					{
+						_aclScopeLookup.Add(template.ScopeName, template);
+					}
+				}
+			}
+
 			Storage.PostLoad(this);
 		}
 
@@ -215,11 +239,12 @@ namespace Horde.Build.Server
 		/// <summary>
 		/// Authorizes a user to perform a given action
 		/// </summary>
+		/// <param name="scopeName">Name of the scope to auth against</param>
 		/// <param name="action">The action being performed</param>
 		/// <param name="user">The principal to validate</param>
-		public bool Authorize(AclAction action, ClaimsPrincipal user)
+		public bool Authorize(AclScopeName scopeName, AclAction action, ClaimsPrincipal user)
 		{
-			return Acl.Authorize(action, user) ?? ServerSettings.Authorize(action, user);
+			return _aclScopeLookup.TryGetValue(scopeName, out IAclScope? scope) && scope.Authorize(action, user);
 		}
 
 		/// <summary>
@@ -237,7 +262,7 @@ namespace Horde.Build.Server
 			}
 			else
 			{
-				return Authorize(AclAction.Impersonate, user);
+				return this.Authorize(AclAction.Impersonate, user);
 			}
 		}
 
diff --git a/Engine/Source/Programs/Horde/Horde.Build/ServerSettings.cs b/Engine/Source/Programs/Horde/Horde.Build/ServerSettings.cs
--- a/Engine/Source/Programs/Horde/Horde.Build/ServerSettings.cs
+++ b/Engine/Source/Programs/Horde/Horde.Build/ServerSettings.cs
@@ -243,8 +243,16 @@ namespace Horde.Build
 	/// <summary>
 	/// Global settings for the application
 	/// </summary>
-	public class ServerSettings
+	public class ServerSettings : IAclScope
 	{
+		/// <inheritdoc/>
+		[JsonIgnore]
+		public IAclScope? ParentScope => null;
+
+		/// <inheritdoc/>
+		[JsonIgnore]
+		public AclScopeName ScopeName => AclScopeName.Root;
+
 		/// <inheritdoc cref="RunMode" />
 		public RunMode[]? RunModes { get; set; } = null;
 
@@ -657,7 +665,7 @@ namespace Horde.Build
 		/// Default pre-baked ACL for authentication of well-known roles
 		/// </summary>
 		[JsonIgnore]
-		public AclConfig DefaultAcl
+		public AclConfig? Acl
 		{
 			get
 			{
@@ -669,16 +677,6 @@ namespace Horde.Build
 		[JsonIgnore]
 		AclConfig? _defaultAcl;
 
-		/// <summary>
-		/// Authorizes a user to perform a given action
-		/// </summary>
-		/// <param name="action">The action being performed</param>
-		/// <param name="user">The principal to validate</param>
-		public bool Authorize(AclAction action, ClaimsPrincipal user)
-		{
-			return DefaultAcl.Authorize(action, user) ?? false;
-		}
-
 		/// <summary>
 		/// Create the default ACL for the server, including all predefined roles.
 		/// </summary>
diff --git a/Engine/Source/Programs/Horde/Horde.Build/Streams/StreamConfig.cs b/Engine/Source/Programs/Horde/Horde.Build/Streams/StreamConfig.cs
--- a/Engine/Source/Programs/Horde/Horde.Build/Streams/StreamConfig.cs
+++ b/Engine/Source/Programs/Horde/Horde.Build/Streams/StreamConfig.cs
@@ -77,7 +77,7 @@ namespace Horde.Build.Streams
 	[JsonSchema("https://unrealengine.com/horde/stream")]
 	[JsonSchemaCatalog("Horde Stream", "Horde stream configuration file", new[] { "*.stream.json", "Streams/*.json" })]
 	[ConfigIncludeRoot]
-	public class StreamConfig
+	public class StreamConfig : IAclScope
 	{
 		/// <summary>
 		/// Accessor for the project containing this stream
@@ -85,6 +85,14 @@ namespace Horde.Build.Streams
 		[JsonIgnore]
 		public ProjectConfig ProjectConfig { get; private set; } = null!;
 
+		/// <inheritdoc/>
+		[JsonIgnore]
+		public IAclScope? ParentScope => ProjectConfig;
+
+		/// <inheritdoc/>
+		[JsonIgnore]
+		public AclScopeName ScopeName { get; private set; }
+
 		/// <summary>
 		/// Identifier for the stream
 		/// </summary>
@@ -236,6 +244,8 @@ namespace Horde.Build.Streams
 		{
 			Id = id;
 			ProjectConfig = projectConfig;
+			ScopeName = projectConfig.ScopeName.Append("s", Id.ToString());
+
 			JobOptions.MergeDefaults(projectConfig.JobOptions);
 
 			foreach (TemplateRefConfig template in Templates)
@@ -244,16 +254,6 @@ namespace Horde.Build.Streams
 			}
 		}
 
-		/// <summary>
-		/// Authorizes a user to perform a given action
-		/// </summary>
-		/// <param name="action">The action being performed</param>
-		/// <param name="user">The principal to validate</param>
-		public bool Authorize(AclAction action, ClaimsPrincipal user)
-		{
-			return Acl?.Authorize(action, user) ?? ProjectConfig.Authorize(action, user);
-		}
-
 		/// <summary>
 		/// Enumerates all commit tags, including the default tags for code and content.
 		/// </summary>
@@ -672,7 +672,7 @@ namespace Horde.Build.Streams
 	/// <summary>
 	/// Parameters to create a template within a stream
 	/// </summary>
-	public class TemplateRefConfig : TemplateConfig
+	public class TemplateRefConfig : TemplateConfig, IAclScope
 	{
 		/// <summary>
 		/// The owning stream config
@@ -680,6 +680,14 @@ namespace Horde.Build.Streams
 		[JsonIgnore]
 		public StreamConfig StreamConfig { get; private set; } = null!;
 
+		/// <inheritdoc/>
+		[JsonIgnore]
+		public IAclScope? ParentScope => StreamConfig;
+
+		/// <inheritdoc/>
+		[JsonIgnore]
+		public AclScopeName ScopeName { get; private set; }
+
 		/// <summary>
 		/// Optional identifier for this ref. If not specified, an id will be generated from the name.
 		/// </summary>
@@ -746,6 +754,8 @@ namespace Horde.Build.Streams
 		public void PostLoad(StreamConfig streamConfig)
 		{
 			StreamConfig = streamConfig;
+			ScopeName = streamConfig.ScopeName.Append("t", Id.ToString());
+
 			JobOptions.MergeDefaults(streamConfig.JobOptions);
 
 			if (Id.IsEmpty)
@@ -753,16 +763,6 @@ namespace Horde.Build.Streams
 				Id = TemplateId.Sanitize(Name);
 			}
 		}
-
-		/// <summary>
-		/// Authorizes a user to perform a given action
-		/// </summary>
-		/// <param name="action">The action being performed</param>
-		/// <param name="user">The principal to validate</param>
-		public bool Authorize(AclAction action, ClaimsPrincipal user)
-		{
-			return Acl?.Authorize(action, user) ?? StreamConfig.Authorize(action, user);
-		}
 	}
 
 	/// <summary>
