Horde: Do not create a user object when performing anonymous login; prevents errors with read-only database.

diff --git a/Engine/Source/Programs/Horde/Horde.Build/Authentication/AnonymousHandler.cs b/Engine/Source/Programs/Horde/Horde.Build/Authentication/AnonymousHandler.cs
--- a/Engine/Source/Programs/Horde/Horde.Build/Authentication/AnonymousHandler.cs
+++ b/Engine/Source/Programs/Horde/Horde.Build/Authentication/AnonymousHandler.cs
@@ -22,16 +22,16 @@ namespace Horde.Build.Authentication
 	class AnonymousAuthenticationHandler : AuthenticationHandler<AnonymousAuthenticationOptions>
 	{
 		public const string AuthenticationScheme = "Anonymous";
-		readonly IUserCollection _userCollection;
 
-		public AnonymousAuthenticationHandler(IUserCollection userCollection, IOptionsMonitor<AnonymousAuthenticationOptions> options,
-			ILoggerFactory logger, UrlEncoder encoder, ISystemClock clock)
+		readonly IOptionsMonitor<ServerSettings> _settings;
+
+		public AnonymousAuthenticationHandler(IOptionsMonitor<AnonymousAuthenticationOptions> options, IOptionsMonitor<ServerSettings> settings, ILoggerFactory logger, UrlEncoder encoder, ISystemClock clock)
 			: base(options, logger, encoder, clock)
 		{
-			_userCollection = userCollection;
+			_settings = settings;
 		}
 
-		protected override async Task<AuthenticateResult> HandleAuthenticateAsync()
+		protected override Task<AuthenticateResult> HandleAuthenticateAsync()
 		{
 			List<Claim> claims = new List<Claim>();
 			claims.Add(new Claim(ClaimTypes.Name, AuthenticationScheme));
@@ -41,15 +41,15 @@ namespace Horde.Build.Authentication
 				claims.Add(new Claim(Options.AdminClaimType, Options.AdminClaimValue));
 			}
 
-			IUser user = await _userCollection.FindOrAddUserByLoginAsync("anonymous", "Anonymous", "[REDACTED_EMAIL]");
+			ServerSettings currentSettings = _settings.CurrentValue;
 
 			ClaimsIdentity identity = new ClaimsIdentity(claims, Scheme.Name);
-			identity.AddClaim(new Claim(HordeClaimTypes.UserId, user.Id.ToString()));
+			identity.AddClaim(new Claim(currentSettings.AdminClaimType, currentSettings.AdminClaimValue));
 
 			ClaimsPrincipal principal = new ClaimsPrincipal(identity);
 			AuthenticationTicket ticket = new AuthenticationTicket(principal, Scheme.Name);
 
-			return AuthenticateResult.Success(ticket);
+			return Task.FromResult(AuthenticateResult.Success(ticket));
 		}
 	}
 
diff --git a/Engine/Source/Programs/Horde/Horde.Build/Projects/ProjectConfig.cs b/Engine/Source/Programs/Horde/Horde.Build/Projects/ProjectConfig.cs
--- a/Engine/Source/Programs/Horde/Horde.Build/Projects/ProjectConfig.cs
+++ b/Engine/Source/Programs/Horde/Horde.Build/Projects/ProjectConfig.cs
@@ -109,7 +109,7 @@ namespace Horde.Build.Projects
 		{
 			Id = id;
 			GlobalConfig = globalConfig;
-			ScopeName = ParentScope.ScopeName.Append("p", Id.ToString());
+			ScopeName = globalConfig.ScopeName.Append("p", Id.ToString());
 
 			foreach (StreamConfig stream in Streams)
 			{
