Horde: Disable task sources when running with a read-only database; we cannot create leases or update any states.

diff --git a/Engine/Source/Programs/Horde/Horde.Build/Startup.cs b/Engine/Source/Programs/Horde/Horde.Build/Startup.cs
--- a/Engine/Source/Programs/Horde/Horde.Build/Startup.cs
+++ b/Engine/Source/Programs/Horde/Horde.Build/Startup.cs
@@ -678,26 +678,27 @@ namespace Horde.Build
 			}
 
 			services.AddHostedService(provider => provider.GetRequiredService<IExternalIssueService>());
-				
+
 			// Task sources. Order of registration is important here; it dictates the priority in which sources are served.
-			services.AddSingleton<JobTaskSource>();
-			services.AddHostedService<JobTaskSource>(provider => provider.GetRequiredService<JobTaskSource>());
-			services.AddSingleton<ConformTaskSource>();
-			services.AddHostedService<ConformTaskSource>(provider => provider.GetRequiredService<ConformTaskSource>());
-			services.AddSingleton<Compute.V1.ComputeService>();
-			services.AddSingleton<Compute.V1.IComputeService, Compute.V1.ComputeService>();
-			services.AddHostedService(provider => provider.GetRequiredService<Compute.V1.ComputeService>());
-			services.AddSingleton<Compute.V2.ComputeServiceV2>();
-
-			services.AddSingleton<ITaskSource, UpgradeTaskSource>();
-			services.AddSingleton<ITaskSource, ShutdownTaskSource>();
-			services.AddSingleton<ITaskSource, RestartTaskSource>();
-			services.AddSingleton<ITaskSource, ConformTaskSource>(provider => provider.GetRequiredService<ConformTaskSource>());
-			services.AddSingleton<ITaskSource, JobTaskSource>(provider => provider.GetRequiredService<JobTaskSource>());
-			services.AddSingleton<ITaskSource, Compute.V1.ComputeService>();
-			services.AddSingleton<ITaskSource, Compute.V2.ComputeServiceV2>(provider => provider.GetRequiredService<Compute.V2.ComputeServiceV2>());
-
-			services.AddHostedService(provider => provider.GetRequiredService<ConformTaskSource>());
+			if (!settings.DatabaseReadOnlyMode)
+			{
+				services.AddSingleton<JobTaskSource>();
+				services.AddHostedService<JobTaskSource>(provider => provider.GetRequiredService<JobTaskSource>());
+				services.AddSingleton<ConformTaskSource>();
+				services.AddHostedService<ConformTaskSource>(provider => provider.GetRequiredService<ConformTaskSource>());
+				services.AddSingleton<Compute.V1.ComputeService>();
+				services.AddSingleton<Compute.V1.IComputeService, Compute.V1.ComputeService>();
+				services.AddHostedService(provider => provider.GetRequiredService<Compute.V1.ComputeService>());
+				services.AddSingleton<Compute.V2.ComputeServiceV2>();
+
+				services.AddSingleton<ITaskSource, UpgradeTaskSource>();
+				services.AddSingleton<ITaskSource, ShutdownTaskSource>();
+				services.AddSingleton<ITaskSource, RestartTaskSource>();
+				services.AddSingleton<ITaskSource, ConformTaskSource>(provider => provider.GetRequiredService<ConformTaskSource>());
+				services.AddSingleton<ITaskSource, JobTaskSource>(provider => provider.GetRequiredService<JobTaskSource>());
+				services.AddSingleton<ITaskSource, Compute.V1.ComputeService>();
+				services.AddSingleton<ITaskSource, Compute.V2.ComputeServiceV2>(provider => provider.GetRequiredService<Compute.V2.ComputeServiceV2>());
+			}
 
 			// Allow longer to shutdown so we can debug missing cancellation tokens
 			services.Configure<HostOptions>(options =>
