Horde: Rename Artifacts to ArtifactsV1.

diff --git a/Engine/Source/Programs/Horde/Horde.Build.Tests/ControllerIntegrationTest.cs b/Engine/Source/Programs/Horde/Horde.Build.Tests/ControllerIntegrationTest.cs
--- a/Engine/Source/Programs/Horde/Horde.Build.Tests/ControllerIntegrationTest.cs
+++ b/Engine/Source/Programs/Horde/Horde.Build.Tests/ControllerIntegrationTest.cs
@@ -96,7 +96,7 @@ public class ControllerIntegrationTest : IDisposable
 		MongoService mongoService = services.GetRequiredService<MongoService>();
 		ITemplateCollection templateService = services.GetRequiredService<ITemplateCollection>();
 		JobService jobService = services.GetRequiredService<JobService>();
-		IArtifactCollection artifactCollection = services.GetRequiredService<IArtifactCollection>();
+		IArtifactCollectionV1 artifactCollection = services.GetRequiredService<IArtifactCollectionV1>();
 		AgentService agentService = services.GetRequiredService<AgentService>();
 		GraphCollection graphCollection = new(mongoService);
 		IOptions<ServerSettings> serverSettings = services.GetRequiredService<IOptions<ServerSettings>>();
diff --git a/Engine/Source/Programs/Horde/Horde.Build.Tests/Fixture.cs b/Engine/Source/Programs/Horde/Horde.Build.Tests/Fixture.cs
--- a/Engine/Source/Programs/Horde/Horde.Build.Tests/Fixture.cs
+++ b/Engine/Source/Programs/Horde/Horde.Build.Tests/Fixture.cs
@@ -36,13 +36,13 @@ namespace Horde.Build.Tests
 		public StreamConfig? StreamConfig { get; private set; }
 		public TemplateId TemplateRefId1 { get; private set; }
 		public TemplateId TemplateRefId2 { get; private set; }
-		public IArtifact Job1Artifact { get; private set; } = null!;
+		public IArtifactV1 Job1Artifact { get; private set; } = null!;
 		public string Job1ArtifactData { get; private set; } = null!;
 		public IAgent Agent1 { get; private set; } = null!;
 		public string Agent1Name { get; private set; } = null!;
 		public const string PoolName = "TestingPool";
 
-		public static async Task<Fixture> Create(ConfigService configService, IGraphCollection graphCollection, ITemplateCollection templateCollection, JobService jobService, IArtifactCollection artifactCollection, AgentService agentService, ServerSettings serverSettings)
+		public static async Task<Fixture> Create(ConfigService configService, IGraphCollection graphCollection, ITemplateCollection templateCollection, JobService jobService, IArtifactCollectionV1 artifactCollection, AgentService agentService, ServerSettings serverSettings)
 		{
 			Fixture fixture = new Fixture();
 			await fixture.Populate(configService, graphCollection, templateCollection, jobService, artifactCollection, agentService, serverSettings);
@@ -53,7 +53,7 @@ namespace Horde.Build.Tests
 			return fixture;
 		}
 
-		private async Task Populate(ConfigService configService, IGraphCollection graphCollection, ITemplateCollection templateCollection, JobService jobService, IArtifactCollection artifactCollection, AgentService agentService, ServerSettings serverSettings)
+		private async Task Populate(ConfigService configService, IGraphCollection graphCollection, ITemplateCollection templateCollection, JobService jobService, IArtifactCollectionV1 artifactCollection, AgentService agentService, ServerSettings serverSettings)
 		{
 			FixtureGraph fg = new FixtureGraph();
 			fg.Id = ContentHash.Empty;
diff --git a/Engine/Source/Programs/Horde/Horde.Build.Tests/JobsControllerTest.cs b/Engine/Source/Programs/Horde/Horde.Build.Tests/JobsControllerTest.cs
--- a/Engine/Source/Programs/Horde/Horde.Build.Tests/JobsControllerTest.cs
+++ b/Engine/Source/Programs/Horde/Horde.Build.Tests/JobsControllerTest.cs
@@ -20,7 +20,7 @@ namespace Horde.Build.Tests
         public async Task GetArtifactDataByFilenameTest()
         {
 			Fixture? fixture = await GetFixture();
-			IArtifact? art = fixture.Job1Artifact;
+			IArtifactV1? art = fixture.Job1Artifact;
 
 			// Test existing filename
 			System.Net.Http.HttpResponseMessage? res = await Client.GetAsync(GetUri(art.JobId.ToString(), art.StepId.ToString()!, art.Name));
diff --git a/Engine/Source/Programs/Horde/Horde.Build.Tests/RpcServiceTest.cs b/Engine/Source/Programs/Horde/Horde.Build.Tests/RpcServiceTest.cs
--- a/Engine/Source/Programs/Horde/Horde.Build.Tests/RpcServiceTest.cs
+++ b/Engine/Source/Programs/Horde/Horde.Build.Tests/RpcServiceTest.cs
@@ -476,7 +476,7 @@ namespace Horde.Build.Tests
 			requestStream.AddMessage(new UploadArtifactRequest { Data = ByteString.CopyFromUtf8(data[3]) });
 			UploadArtifactResponse res = await call;
 
-			IArtifact? artifact = await ArtifactCollection.GetArtifactAsync(ObjectId.Parse(res.Id));
+			IArtifactV1? artifact = await ArtifactCollection.GetArtifactAsync(ObjectId.Parse(res.Id));
 			Assert.IsNotNull(artifact);
 			Stream stream = await ArtifactCollection.OpenArtifactReadStreamAsync(artifact!);
 			using StreamReader reader = new StreamReader(stream);
diff --git a/Engine/Source/Programs/Horde/Horde.Build.Tests/TestSetup.cs b/Engine/Source/Programs/Horde/Horde.Build.Tests/TestSetup.cs
--- a/Engine/Source/Programs/Horde/Horde.Build.Tests/TestSetup.cs
+++ b/Engine/Source/Programs/Horde/Horde.Build.Tests/TestSetup.cs
@@ -109,7 +109,7 @@ namespace Horde.Build.Tests
 		public IssueService IssueService => ServiceProvider.GetRequiredService<IssueService>();
 		public JobTaskSource JobTaskSource => ServiceProvider.GetRequiredService<JobTaskSource>();
 		public JobService JobService => ServiceProvider.GetRequiredService<JobService>();
-		public IArtifactCollection ArtifactCollection => ServiceProvider.GetRequiredService<IArtifactCollection>();
+		public IArtifactCollectionV1 ArtifactCollection => ServiceProvider.GetRequiredService<IArtifactCollectionV1>();
 		public IDowntimeService DowntimeService => ServiceProvider.GetRequiredService<IDowntimeService>();
 		public RpcService RpcService => ServiceProvider.GetRequiredService<RpcService>();
 		public CredentialService CredentialService => ServiceProvider.GetRequiredService<CredentialService>();
@@ -198,7 +198,7 @@ namespace Horde.Build.Tests
 
 			services.AddSingleton<IAgentCollection, AgentCollection>();
 			services.AddSingleton<IAgentSoftwareCollection, AgentSoftwareCollection>();
-			services.AddSingleton<IArtifactCollection, ArtifactCollection>();
+			services.AddSingleton<IArtifactCollectionV1, ArtifactCollectionV1>();
 			services.AddSingleton<ICommitService, CommitService>();
 			services.AddSingleton<IGraphCollection, GraphCollection>();
 			services.AddSingleton<IIssueCollection, IssueCollection>();
@@ -264,7 +264,7 @@ namespace Horde.Build.Tests
 			services.AddSingleton<ICommitService, CommitService>();
 
 			services.AddSingleton<IStorageBackend<PersistentLogStorage>>(sp => new MemoryStorageBackend().ForType<PersistentLogStorage>());
-			services.AddSingleton<IStorageBackend<ArtifactCollection>>(sp => new MemoryStorageBackend().ForType<ArtifactCollection>());
+			services.AddSingleton<IStorageBackend<ArtifactCollectionV1>>(sp => new MemoryStorageBackend().ForType<ArtifactCollectionV1>());
 			services.AddSingleton<IStorageBackend<BasicStorageClient>>(sp => new MemoryStorageBackend().ForType<BasicStorageClient>());
 
 			services.AddSingleton<ILegacyStorageClient, BasicStorageClient>();
diff --git a/Engine/Source/Programs/Horde/Horde.Build/Jobs/Artifacts/ArtifactCollection.cs b/Engine/Source/Programs/Horde/Horde.Build/Jobs/Artifacts/ArtifactCollection.cs
deleted file mode 100644
diff --git a/Engine/Source/Programs/Horde/Horde.Build/Jobs/Artifacts/ArtifactCollectionV1.cs b/Engine/Source/Programs/Horde/Horde.Build/Jobs/Artifacts/ArtifactCollectionV1.cs
new file mode 100644
--- /dev/null
+++ b/Engine/Source/Programs/Horde/Horde.Build/Jobs/Artifacts/ArtifactCollectionV1.cs
@@ -0,0 +1,272 @@
+// Copyright Epic Games, Inc. All Rights Reserved.
+
+using System;
+using System.Collections.Generic;
+using System.Linq;
+using System.Threading.Tasks;
+using Horde.Build.Server;
+using Horde.Build.Storage;
+using Horde.Build.Utilities;
+using MongoDB.Bson;
+using MongoDB.Bson.Serialization.Attributes;
+using MongoDB.Driver;
+
+namespace Horde.Build.Jobs.Artifacts
+{
+	using JobId = ObjectId<IJob>;
+
+	/// <summary>
+	/// Wraps functionality for manipulating artifacts
+	/// </summary>
+	public class ArtifactCollectionV1 : IArtifactCollectionV1
+	{
+		class Artifact : IArtifactV1
+		{
+			[BsonRequired, BsonId]
+			public ObjectId Id { get; set; }
+
+			[BsonRequired]
+			public JobId JobId { get; set; }
+
+			[BsonRequired]
+			public string Name { get; set; }
+
+			public SubResourceId? StepId { get; set; }
+
+			[BsonRequired]
+			public long Length { get; set; }
+
+			public string MimeType { get; set; }
+
+			[BsonRequired]
+			public int UpdateIndex { get; set; }
+
+			[BsonConstructor]
+			private Artifact()
+			{
+				Name = null!;
+				MimeType = null!;
+			}
+
+			public Artifact(JobId jobId, SubResourceId? stepId, string name, long length, string mimeType)
+			{
+				Id = ObjectId.GenerateNewId();
+				JobId = jobId;
+				StepId = stepId;
+				Name = name;
+				Length = length;
+				MimeType = mimeType;
+			}
+		}
+
+		private readonly IStorageBackend _storageBackend;
+		private readonly IMongoCollection<Artifact> _artifacts;
+
+		/// <summary>
+		/// Constructor
+		/// </summary>
+		/// <param name="mongoService">The database service</param>
+		/// <param name="storageBackend">The storage backend</param>
+		public ArtifactCollectionV1(MongoService mongoService, IStorageBackend<ArtifactCollectionV1> storageBackend)
+		{
+			_storageBackend = storageBackend;
+
+			// Initialize Artifacts table
+			_artifacts = mongoService.GetCollection<Artifact>("Artifacts", keys => keys.Ascending(x => x.JobId));
+		}
+
+		/// <summary>
+		/// Creates a new artifact
+		/// </summary>
+		/// <param name="jobId">Unique id of the job that owns this artifact</param>
+		/// <param name="stepId">Optional Step id</param>
+		/// <param name="name">Name of artifact</param>
+		/// <param name="mimeType">Type of artifact</param>
+		/// <param name="data">The data to write</param>
+		/// <returns>The new log file document</returns>
+		public async Task<IArtifactV1> CreateArtifactAsync(JobId jobId, SubResourceId? stepId, string name, string mimeType, System.IO.Stream data)
+		{
+			// upload first
+			string artifactName = ValidateName(name);
+			await _storageBackend.WriteAsync(GetPath(jobId, stepId, artifactName), data);
+
+			// then create entry
+			Artifact newArtifact = new Artifact(jobId, stepId, artifactName, data.Length, mimeType);
+			await _artifacts.InsertOneAsync(newArtifact);
+			return newArtifact;
+		}
+
+		/// <summary>
+		/// Gets all the available artifacts for a job
+		/// </summary>
+		/// <param name="jobId">Unique id of the job to query</param>
+		/// <param name="stepId">Unique id of the Step to query</param>
+		/// <param name="name">Name of the artifact</param>
+		/// <returns>List of artifact documents</returns>
+		public async Task<List<IArtifactV1>> GetArtifactsAsync(JobId? jobId, SubResourceId? stepId, string? name)
+		{
+			FilterDefinitionBuilder<Artifact> builder = Builders<Artifact>.Filter;
+
+			FilterDefinition<Artifact> filter = FilterDefinition<Artifact>.Empty;
+			if (jobId != null)
+			{
+				filter &= builder.Eq(x => x.JobId, jobId.Value);
+			}
+			if (stepId != null)
+			{
+				filter &= builder.Eq(x => x.StepId, stepId.Value);
+			}
+			if (name != null)
+			{
+				filter &= builder.Eq(x => x.Name, name);
+			}
+
+			return await _artifacts.Find(filter).ToListAsync<Artifact, IArtifactV1>();
+		}
+
+		/// <summary>
+		/// Gets a specific list of artifacts based on id
+		/// </summary>
+		/// <param name="artifactIds">The list of artifact Ids</param>
+		/// <returns>List of artifact documents</returns>
+		public async Task<List<IArtifactV1>> GetArtifactsAsync(IEnumerable<ObjectId> artifactIds)
+		{
+			FilterDefinitionBuilder<Artifact> builder = Builders<Artifact>.Filter;
+
+			FilterDefinition<Artifact> filter = FilterDefinition<Artifact>.Empty;
+			filter &= builder.In(x => x.Id, artifactIds);
+			
+			return await _artifacts.Find(filter).ToListAsync<Artifact, IArtifactV1>();
+		}
+
+		/// <summary>
+		/// Gets an artifact by ID
+		/// </summary>
+		/// <param name="artifactId">Unique id of the artifact</param>
+		/// <returns>The artifact document</returns>
+		public async Task<IArtifactV1?> GetArtifactAsync(ObjectId artifactId)
+		{
+			return await _artifacts.Find<Artifact>(x => x.Id == artifactId).FirstOrDefaultAsync();
+		}
+
+		/// <summary>
+		/// Attempts to update a single artifact document, if the update counter is valid
+		/// </summary>
+		/// <param name="current">The current artifact document</param>
+		/// <param name="update">The update definition</param>
+		/// <returns>True if the document was updated, false if another writer updated the document first</returns>
+		private async Task<bool> TryUpdateArtifactAsync(Artifact current, UpdateDefinition<Artifact> update)
+		{
+			UpdateResult result = await _artifacts.UpdateOneAsync<Artifact>(x => x.Id == current.Id && x.UpdateIndex == current.UpdateIndex, update.Set(x => x.UpdateIndex, current.UpdateIndex + 1));
+			return result.ModifiedCount == 1;
+		}
+
+		/// <summary>
+		/// Updates an artifact
+		/// </summary>
+		/// <param name="artifact">The artifact</param>
+		/// <param name="newMimeType">New mime type</param>
+		/// <param name="newData">New data</param>
+		/// <returns>Async task</returns>
+		public async Task<bool> UpdateArtifactAsync(IArtifactV1? artifact, string newMimeType, System.IO.Stream newData)
+		{
+			while (artifact != null)
+			{
+				UpdateDefinitionBuilder<Artifact> updateBuilder = Builders<Artifact>.Update;
+
+				List<UpdateDefinition<Artifact>> updates = new List<UpdateDefinition<Artifact>>();
+				updates.Add(updateBuilder.Set(x => x.MimeType, newMimeType));
+				updates.Add(updateBuilder.Set(x => x.Length, newData.Length));
+
+				// re-upload the data to external
+				string artifactName = ValidateName(artifact.Name);
+				await _storageBackend.WriteAsync(GetPath(artifact.JobId, artifact.StepId, artifactName), newData);
+
+				if (await TryUpdateArtifactAsync((Artifact)artifact, updateBuilder.Combine(updates)))
+				{
+					return true;
+				}
+
+				artifact = await GetArtifactAsync(artifact.Id);
+			}
+			return false;
+		}
+
+		/// <summary>
+		/// gets artifact data
+		/// </summary>
+		/// <param name="artifact">The artifact</param>
+		/// <returns>The chunk data</returns>
+		public async Task<System.IO.Stream> OpenArtifactReadStreamAsync(IArtifactV1 artifact)
+		{
+			System.IO.Stream? stream = await _storageBackend.TryReadAsync(GetPath(artifact.JobId, artifact.StepId, artifact.Name));
+			if (stream == null)
+			{
+				throw new Exception($"Unable to get artifact {artifact.Id}");
+			}
+			return stream;
+		}
+
+		/// <summary>
+		/// Get the path for an artifact
+		/// </summary>
+		/// <param name="jobId"></param>
+		/// <param name="stepId"></param>
+		/// <param name="name"></param>
+		/// <returns></returns>
+		private static string GetPath(JobId jobId, SubResourceId? stepId, string name)
+		{
+			if (stepId == null)
+			{
+				return $"{jobId}/{name}";
+			}
+			else
+			{
+				return $"{jobId}/{stepId.Value}/{name}";
+			}
+		}
+
+		/// <summary>
+		/// Checks that the name given is valid for an artifact
+		/// </summary>
+		/// <param name="name"></param>
+		private static string ValidateName(string name)
+		{
+			string newName = name.Replace('\\', '/');
+			if (newName.Length == 0)
+			{
+				throw new ArgumentException("Artifact name is empty");
+			}
+			else if (newName.StartsWith('/'))
+			{
+				throw new ArgumentException($"Artifact has an absolute path ({newName})");
+			}
+			else if (newName.EndsWith('/'))
+			{
+				throw new ArgumentException($"Artifact does not have a file name ({newName})");
+			}
+			else if (newName.Contains("//", StringComparison.Ordinal))
+			{
+				throw new ArgumentException($"Artifact name contains an invalid directory ({newName})");
+			}
+
+			string invalidChars = ":<>|\"";
+			for (int idx = 0; idx < newName.Length; idx++)
+			{
+				int minIdx = idx;
+				for (; idx < newName.Length && newName[idx] != '/'; idx++)
+				{
+					if (invalidChars.Contains(newName[idx], StringComparison.Ordinal))
+					{
+						throw new ArgumentException($"Invalid character in artifact name ({newName})");
+					}
+				}
+				if ((idx == minIdx + 1 && newName[minIdx] == '.') || (idx == minIdx + 2 && newName[minIdx] == '.' && newName[minIdx + 1] == '.'))
+				{
+					throw new ArgumentException($"Artifact may not contain symbolic directory names ({newName})");
+				}
+			}
+			return newName;
+		}
+	}
+}
\ No newline at end of file
diff --git a/Engine/Source/Programs/Horde/Horde.Build/Jobs/Artifacts/ArtifactMessages.cs b/Engine/Source/Programs/Horde/Horde.Build/Jobs/Artifacts/ArtifactMessages.cs
deleted file mode 100644
diff --git a/Engine/Source/Programs/Horde/Horde.Build/Jobs/Artifacts/ArtifactMessagesV1.cs b/Engine/Source/Programs/Horde/Horde.Build/Jobs/Artifacts/ArtifactMessagesV1.cs
new file mode 100644
--- /dev/null
+++ b/Engine/Source/Programs/Horde/Horde.Build/Jobs/Artifacts/ArtifactMessagesV1.cs
@@ -0,0 +1,117 @@
+// Copyright Epic Games, Inc. All Rights Reserved.
+
+using System.Collections.Generic;
+
+namespace Horde.Build.Jobs.Artifacts
+{
+	/// <summary>
+	/// Response from creating an artifact
+	/// </summary>
+	public class CreateArtifactResponse
+	{
+		/// <summary>
+		/// Unique id for this artifact
+		/// </summary>
+		public string Id { get; set; }
+
+		/// <summary>
+		/// Constructor
+		/// </summary>
+		/// <param name="id">The artifact file id</param>
+		public CreateArtifactResponse(string id)
+		{
+			Id = id;
+		}
+	}
+
+	/// <summary>
+	/// Request to get a zip of artifacts
+	/// </summary>
+	public class GetArtifactZipRequest
+	{
+		/// <summary>
+		/// Job id to get all artifacts for
+		/// </summary>
+		public string? JobId { get; set; }
+
+		/// <summary>
+		/// Step id to filter by
+		/// </summary>
+		public string? StepId { get; set; }
+
+		/// <summary>
+		/// Further filter by a list of artifact ids
+		/// </summary>
+		public List<string>? ArtifactIds { get; set; }
+
+		/// <summary>
+		/// Constructor
+		/// </summary>
+		/// <param name="jobId">Job id to get all artifacts for</param>
+		/// <param name="stepId">step to filter by</param>
+		/// <param name="artifactIds">The artifact ids.  Returns all artifacts for a job </param>
+		public GetArtifactZipRequest(string? jobId, string? stepId, List<string>? artifactIds)
+		{
+			JobId = jobId;
+			StepId = stepId;
+			ArtifactIds = artifactIds;
+		}
+	}
+
+	/// <summary>
+	/// Response describing an artifact
+	/// </summary>
+	public class GetArtifactResponse
+	{
+		/// <summary>
+		/// Unique id of the artifact
+		/// </summary>
+		public string Id { get; set; }
+
+		/// <summary>
+		/// Name of the artifact
+		/// </summary>
+		public string Name { get; set; }
+
+		/// <summary>
+		/// Unique id of the job for this artifact
+		/// </summary>
+		public string JobId { get; set; }
+
+		/// <summary>
+		/// Optional unique id of the step for this artifact
+		/// </summary>
+		public string? StepId { get; set; }
+
+		/// <summary>
+		/// Type of artifact
+		/// </summary>
+		public string MimeType { get; set; }
+
+		/// <summary>
+		/// Length of artifact
+		/// </summary>
+		public long Length { get; set; }
+
+		/// <summary>
+		/// Short-lived code that can be used to download this artifact through direct download links in the browser
+		/// </summary>
+		public string? Code { get; set; }
+
+		/// <summary>
+		/// Constructor
+		/// </summary>
+		/// <param name="artifact">The artifact to construct from</param>
+		/// <param name="code">The direct download code</param>
+		public GetArtifactResponse(IArtifactV1 artifact, string? code)
+		{
+			Id = artifact.Id.ToString();
+			Name = artifact.Name;
+			JobId = artifact.JobId.ToString();
+			StepId = artifact.StepId.ToString();
+			MimeType = artifact.MimeType;
+			Length = artifact.Length;
+			Code = code;
+		}
+	}
+}
diff --git a/Engine/Source/Programs/Horde/Horde.Build/Jobs/Artifacts/ArtifactsController.cs b/Engine/Source/Programs/Horde/Horde.Build/Jobs/Artifacts/ArtifactsController.cs
deleted file mode 100644
diff --git a/Engine/Source/Programs/Horde/Horde.Build/Jobs/Artifacts/ArtifactsControllerV1.cs b/Engine/Source/Programs/Horde/Horde.Build/Jobs/Artifacts/ArtifactsControllerV1.cs
new file mode 100644
--- /dev/null
+++ b/Engine/Source/Programs/Horde/Horde.Build/Jobs/Artifacts/ArtifactsControllerV1.cs
@@ -0,0 +1,436 @@
+// Copyright Epic Games, Inc. All Rights Reserved.
+
+using System;
+using System.Collections.Generic;
+using System.IdentityModel.Tokens.Jwt;
+using System.IO;
+using System.IO.Compression;
+using System.Linq;
+using System.Net.Mime;
+using System.Security.Claims;
+using System.Threading.Tasks;
+using Horde.Build.Acls;
+using Horde.Build.Jobs.Graphs;
+using Horde.Build.Server;
+using Horde.Build.Streams;
+using Horde.Build.Utilities;
+using Microsoft.AspNetCore.Authorization;
+using Microsoft.AspNetCore.Http;
+using Microsoft.AspNetCore.Mvc;
+using Microsoft.Extensions.Options;
+using Microsoft.IdentityModel.Tokens;
+using MongoDB.Bson;
+
+namespace Horde.Build.Jobs.Artifacts
+{
+	using JobId = ObjectId<IJob>;
+
+	/// <summary>
+	/// Controller for the /api/artifacts endpoint
+	/// </summary>
+	[ApiController]
+	[Route("[controller]")]
+	public class ArtifactsControllerV1 : ControllerBase
+	{
+		private readonly GlobalsService _globalsService;
+		private readonly IArtifactCollectionV1 _artifactCollection;
+		private readonly AclService _aclService;
+		private readonly JobService _jobService;
+		private readonly IOptionsSnapshot<GlobalConfig> _globalConfig;
+
+		/// <summary>
+		/// Constructor
+		/// </summary>
+		public ArtifactsControllerV1(GlobalsService globalsService, IArtifactCollectionV1 artifactCollection, AclService aclService, JobService jobService, IOptionsSnapshot<GlobalConfig> globalConfig)
+		{
+			_globalsService = globalsService;
+			_artifactCollection = artifactCollection;
+			_aclService = aclService;
+			_jobService = jobService;
+			_globalConfig = globalConfig;
+		}
+
+		/// <summary>
+		/// Creates an artifact
+		/// </summary>
+		/// <param name="jobId">BatchId</param>
+		/// <param name="stepId">StepId</param>
+		/// <param name="file">The file contents</param>
+		/// <returns>Http result code</returns>
+		[HttpPost]
+		[Authorize]
+		[Route("/api/v1/artifacts")]
+		public async Task<ActionResult<CreateArtifactResponse>> CreateArtifact([FromQuery] JobId jobId, [FromQuery]string? stepId, IFormFile file)
+		{
+			IJob? job = await _jobService.GetJobAsync(jobId);
+			if(job == null)
+			{
+				return NotFound(jobId);
+			}
+
+			StreamConfig? streamConfig;
+			if (!_globalConfig.Value.TryGetStream(job.StreamId, out streamConfig))
+			{
+				return NotFound(job.StreamId);
+			}
+			if (!streamConfig.Authorize(AclAction.UploadArtifact, User))
+			{
+				return Forbid();
+			}
+
+			IJobStep? step = null;
+			if(stepId != null)
+			{
+				foreach(IJobStepBatch batch in job.Batches)
+				{
+					if(batch.TryGetStep(stepId.ToSubResourceId(), out step))
+					{
+						break;
+					}
+				}
+				if(step == null)
+				{
+					// if the step doesn't exist in any of the batches, not found
+					return NotFound();
+				}
+			}
+
+			IArtifactV1 newArtifact = await _artifactCollection.CreateArtifactAsync(job.Id, step?.Id, file.FileName, file.ContentType ?? "horde-mime/unknown", file.OpenReadStream());
+			return new CreateArtifactResponse(newArtifact.Id.ToString());
+		}
+
+		/// <summary>
+		/// Updates an artifact
+		/// </summary>
+		/// <param name="artifactId">JobId</param>
+		/// <param name="file">The file contents</param>
+		/// <returns>Http result code</returns>
+		[HttpPut]
+		[Authorize]
+		[Route("/api/v1/artifacts/{artifactId}")]
+		public async Task<ActionResult<CreateArtifactResponse>> UpdateArtifact(string artifactId, IFormFile file)
+		{
+			IArtifactV1? artifact = await _artifactCollection.GetArtifactAsync(artifactId.ToObjectId());
+			if (artifact == null)
+			{
+				return NotFound();
+			}
+			if (!await _jobService.AuthorizeAsync(artifact.JobId, AclAction.UploadArtifact, User, _globalConfig.Value))
+			{
+				return Forbid();
+			}
+
+			await _artifactCollection.UpdateArtifactAsync(artifact, file.ContentType ?? "horde-mime/unknown", file.OpenReadStream());
+			return Ok();
+		}
+
+		/// <summary>
+		/// Query artifacts for a job step
+		/// </summary>
+		/// <param name="jobId">Optional JobId to filter by</param>
+		/// <param name="stepId">Optional StepId to filter by</param>
+		/// <param name="code">Whether to generate a direct download code</param>
+		/// <param name="filter">Filter for the properties to return</param>
+		/// <returns>Information about all the artifacts</returns>
+		[HttpGet]
+		[Authorize]
+		[Route("/api/v1/artifacts")]
+		[ProducesResponseType(typeof(List<GetArtifactResponse>), 200)]
+		public async Task<ActionResult<List<object>>> GetArtifacts([FromQuery] JobId jobId, [FromQuery] string? stepId = null, [FromQuery] bool code = false, [FromQuery] PropertyFilter? filter = null)
+		{
+			if (!await _jobService.AuthorizeAsync(jobId, AclAction.DownloadArtifact, User, _globalConfig.Value))
+			{
+				return Forbid();
+			}
+
+			string? downloadCode = code ? (string?)await GetDirectDownloadCodeForJobAsync(jobId) : null;
+
+			List<IArtifactV1> artifacts = await _artifactCollection.GetArtifactsAsync(jobId, stepId?.ToSubResourceId(), null);
+			return artifacts.ConvertAll(x => new GetArtifactResponse(x, downloadCode).ApplyFilter(filter));
+		}
+
+		/// <summary>
+		/// Gets the claim required to download artifacts for a particular job
+		/// </summary>
+		/// <param name="jobId">The job id</param>
+		/// <returns>The required claim</returns>
+		static Claim GetDirectDownloadClaim(JobId jobId)
+		{
+			return new Claim(HordeClaimTypes.JobArtifacts, jobId.ToString());
+		}
+
+		/// <summary>
+		/// Get a download code for the artifacts of a job
+		/// </summary>
+		/// <param name="jobId">The job id</param>
+		/// <returns>The download code</returns>
+		async ValueTask<string> GetDirectDownloadCodeForJobAsync(JobId jobId)
+		{
+			Claim downloadClaim = GetDirectDownloadClaim(jobId);
+			return await _aclService.IssueBearerTokenAsync(new[] { downloadClaim }, TimeSpan.FromHours(4.0));
+		}
+
+		/// <summary>
+		/// Retrieve metadata about a specific artifact
+		/// </summary>
+		/// <param name="artifactId">Id of the artifact to get information about</param>
+		/// <param name="code">Whether to generate a direct download code</param>
+		/// <param name="filter">Filter for the properties to return</param>
+		/// <returns>Information about the requested project</returns>
+		[HttpGet]
+		[Authorize]
+		[Route("/api/v1/artifacts/{artifactId}")]
+		[ProducesResponseType(typeof(GetArtifactResponse), 200)]
+		public async Task<ActionResult<object>> GetArtifact(string artifactId, bool code = false, [FromQuery] PropertyFilter? filter = null)
+		{
+			IArtifactV1? artifact = await _artifactCollection.GetArtifactAsync(artifactId.ToObjectId());
+			if (artifact == null)
+			{
+				return NotFound();
+			}
+			if (!await _jobService.AuthorizeAsync(artifact.JobId, AclAction.DownloadArtifact, User, _globalConfig.Value))
+			{
+				return Forbid();
+			}
+
+			string? downloadCode = code? (string?)await GetDirectDownloadCodeForJobAsync(artifact.JobId) : null;
+			return new GetArtifactResponse(artifact, downloadCode).ApplyFilter(filter);
+		}
+
+		/// <summary>
+		/// Retrieve raw data for an artifact
+		/// </summary>
+		/// <param name="artifactId">Id of the artifact to get information about</param>
+		/// <returns>Raw artifact data</returns>
+		[HttpGet]
+		[Authorize]
+		[Route("/api/v1/artifacts/{artifactId}/data")]
+		public async Task<ActionResult> GetArtifactData(string artifactId)
+		{
+			IArtifactV1? artifact = await _artifactCollection.GetArtifactAsync(artifactId.ToObjectId());
+			if (artifact == null)
+			{
+				return NotFound();
+			}
+			if (!await _jobService.AuthorizeAsync(artifact.JobId, AclAction.DownloadArtifact, User, _globalConfig.Value))
+			{
+				return Forbid();
+			}
+
+			// Fun, filestream result automatically closes the stream!
+			return new FileStreamResult(await _artifactCollection.OpenArtifactReadStreamAsync(artifact), artifact.MimeType);
+		}
+		
+		/// <summary>
+		/// Retrieve raw data for an artifact by filename
+		/// </summary>
+		/// <param name="jobId">Unique id for the job</param>
+		/// <param name="stepId">Unique id for the step</param>
+		/// <param name="filename">Filename of artifact from step</param>
+		/// <returns>Raw artifact data</returns>
+		[HttpGet]
+		[Route("/api/v1/jobs/{jobId}/steps/{stepId}/artifacts/{filename}/data")]
+		public async Task<ActionResult<object>> GetArtifactDataByFilename(JobId jobId, string stepId, string filename)
+		{
+			SubResourceId stepIdValue = stepId.ToSubResourceId();
+
+			if (!await _jobService.AuthorizeAsync(jobId, AclAction.DownloadArtifact, User, _globalConfig.Value))
+			{
+				return Forbid();
+			}
+			
+			List<IArtifactV1> artifacts = await _artifactCollection.GetArtifactsAsync(jobId, stepIdValue, filename);
+			if (artifacts.Count == 0)
+			{
+				return NotFound();
+			}
+
+			IArtifactV1 artifact = artifacts[0];
+			return new FileStreamResult(await _artifactCollection.OpenArtifactReadStreamAsync(artifact), artifact.MimeType);
+		}
+
+		/// <summary>
+		/// Class to return a file stream without the "content-disposition: attachment" header
+		/// </summary>
+		class InlineFileStreamResult : FileStreamResult
+		{
+			/// <summary>
+			/// The suggested download filename
+			/// </summary>
+			readonly string _fileName;
+
+			/// <summary>
+			/// Constructor
+			/// </summary>
+			/// <param name="stream"></param>
+			/// <param name="mimeType"></param>
+			/// <param name="fileName"></param>
+			public InlineFileStreamResult(System.IO.Stream stream, string mimeType, string fileName)
+				: base(stream, mimeType)
+			{
+				_fileName = fileName;
+			}
+
+			/// <inheritdoc/>
+			public override Task ExecuteResultAsync(ActionContext context)
+			{
+				ContentDisposition contentDisposition = new ContentDisposition();
+				contentDisposition.Inline = true;
+				contentDisposition.FileName = _fileName;
+				context.HttpContext.Response.Headers.Add("Content-Disposition", contentDisposition.ToString());
+				
+				return base.ExecuteResultAsync(context);
+			}
+		}
+
+		/// <summary>
+		/// Retrieve raw data for an artifact
+		/// </summary>
+		/// <param name="artifactId">Id of the artifact to get information about</param>
+		/// <param name="code">The authorization code for this resource</param>
+		/// <returns>Raw artifact data</returns>
+		[HttpGet]
+		[AllowAnonymous]
+		[Route("/api/v1/artifacts/{artifactId}/download")]
+		public async Task<ActionResult> DownloadArtifact(string artifactId, [FromQuery] string code)
+		{
+			IGlobals globals = await _globalsService.GetAsync();
+
+			TokenValidationParameters parameters = new TokenValidationParameters();
+			parameters.ValidateAudience = false;
+			parameters.RequireExpirationTime = true;
+			parameters.ValidateLifetime = true;
+			parameters.ValidIssuer = globals.JwtIssuer;
+			parameters.ValidateIssuer = true;
+			parameters.ValidateIssuerSigningKey = true;
+			parameters.IssuerSigningKey = globals.JwtSigningKey;
+
+			JwtSecurityTokenHandler handler = new JwtSecurityTokenHandler();
+			ClaimsPrincipal principal = handler.ValidateToken(code, parameters, out _);
+
+			IArtifactV1? artifact = await _artifactCollection.GetArtifactAsync(artifactId.ToObjectId());
+			if (artifact == null)
+			{
+				return NotFound();
+			}
+
+			Claim directDownloadClaim = GetDirectDownloadClaim(artifact.JobId);
+			if (!principal.HasClaim(directDownloadClaim.Type, directDownloadClaim.Value))
+			{
+				return Forbid();
+			}
+
+			return new InlineFileStreamResult(await _artifactCollection.OpenArtifactReadStreamAsync(artifact), artifact.MimeType, Path.GetFileName(artifact.Name));
+		}
+
+		/// <summary>
+		/// Returns a zip archive of many artifacts
+		/// </summary>
+		/// <param name="artifactZipRequest">Artifact request params</param>
+		/// <returns>Zip of many artifacts</returns>
+		[HttpPost]
+		[Authorize]
+		[Route("/api/v1/artifacts/zip")]
+		public async Task<ActionResult> ZipArtifacts(GetArtifactZipRequest artifactZipRequest)
+		{
+			if (artifactZipRequest.JobId == null)
+			{
+				return BadRequest("Must specify a JobId");
+			}
+
+			IJob? job = await _jobService.GetJobAsync(new JobId(artifactZipRequest.JobId!));
+			if (job == null)
+			{
+				return NotFound();
+			}
+			if (!_globalConfig.Value.Authorize(job, AclAction.DownloadArtifact, User))
+			{
+				return Forbid();
+			}
+
+			List<IArtifactV1> artifacts = await _artifactCollection.GetArtifactsAsync(job.Id, artifactZipRequest.StepId?.ToSubResourceId(), null);
+
+			Dictionary<ObjectId, IArtifactV1> idToArtifact = artifacts.ToDictionary(x => x.Id, x => x);
+
+			List<IArtifactV1> zipArtifacts;
+			if (artifactZipRequest.ArtifactIds == null)
+			{
+				zipArtifacts = artifacts;
+			}
+			else
+			{
+				zipArtifacts = new List<IArtifactV1>();
+				foreach (string artifactId in artifactZipRequest.ArtifactIds)
+				{
+					IArtifactV1? artifact;
+					if (idToArtifact.TryGetValue(artifactId.ToObjectId(), out artifact))
+					{
+						zipArtifacts.Add(artifact);
+					}
+					else
+					{
+						return NotFound();
+					}
+				}
+			}
+
+			IGraph graph = await _jobService.GetGraphAsync(job);
+
+			return new CustomFileCallbackResult("Artifacts.zip", "application/octet-stream", false, async (outputStream, context) =>
+			{
+				// Make an unseekable MemoryStream for the ZipArchive. We have to do this because the ZipEntry stream falls back to a synchronous write to it's own stream wrappers.
+				using (CustomBufferStream zipOutputStream = new CustomBufferStream())
+				{
+					// Keep the stream open after dispose so we can write the EOF bits.
+					using (ZipArchive zipArchive = new ZipArchive(zipOutputStream, ZipArchiveMode.Create, true))
+					{
+						foreach (IArtifactV1 artifact in zipArtifacts)
+						{
+							await using (System.IO.Stream artifactStream = await _artifactCollection.OpenArtifactReadStreamAsync(artifact))
+							{
+								// tack on the step name into the directory if it exists
+								string stepName = String.Empty;
+								if (artifact.StepId.HasValue)
+								{
+									foreach (IJobStepBatch batch in job.Batches)
+									{
+										IJobStep? step;
+										if (batch.TryGetStep(artifact.StepId.Value, out step))
+										{
+											stepName = graph.Groups[batch.GroupIdx].Nodes[step.NodeIdx].Name;
+											break;
+										}
+									}
+								}
+
+								ZipArchiveEntry zipEntry = zipArchive.CreateEntry(artifact.Name);
+								using (System.IO.Stream entryStream = zipEntry.Open())
+								{
+									byte[] buffer = new byte[4096];
+									int totalBytesRead = 0;
+									while (totalBytesRead < artifact.Length)
+									{
+										int bytesRead = await artifactStream.ReadAsync(buffer, 0, buffer.Length);
+
+										// Write bytes to the entry stream.  Also advances the MemStream pos.
+										await entryStream.WriteAsync(buffer, 0, bytesRead);
+										// Dump what we have to the output stream
+										await outputStream.WriteAsync(zipOutputStream.GetBuffer(), 0, (int)zipOutputStream.Position);
+
+										// Reset everything.
+										zipOutputStream.Position = 0;
+										zipOutputStream.SetLength(0);
+										totalBytesRead += bytesRead;
+									}
+								}
+							}
+						}
+					}
+					// Write out the EOF stuff
+					zipOutputStream.Position = 0;
+					await zipOutputStream.CopyToAsync(outputStream);
+				}
+			});
+		}
+	}
+}
diff --git a/Engine/Source/Programs/Horde/Horde.Build/Jobs/Artifacts/IArtifact.cs b/Engine/Source/Programs/Horde/Horde.Build/Jobs/Artifacts/IArtifact.cs
deleted file mode 100644
diff --git a/Engine/Source/Programs/Horde/Horde.Build/Jobs/Artifacts/IArtifactCollection.cs b/Engine/Source/Programs/Horde/Horde.Build/Jobs/Artifacts/IArtifactCollection.cs
deleted file mode 100644
diff --git a/Engine/Source/Programs/Horde/Horde.Build/Jobs/Artifacts/IArtifactCollectionV1.cs b/Engine/Source/Programs/Horde/Horde.Build/Jobs/Artifacts/IArtifactCollectionV1.cs
new file mode 100644
--- /dev/null
+++ b/Engine/Source/Programs/Horde/Horde.Build/Jobs/Artifacts/IArtifactCollectionV1.cs
@@ -0,0 +1,67 @@
+// Copyright Epic Games, Inc. All Rights Reserved.
+
+using System.Collections.Generic;
+using System.Threading.Tasks;
+using Horde.Build.Utilities;
+using MongoDB.Bson;
+
+namespace Horde.Build.Jobs.Artifacts
+{
+	using JobId = ObjectId<IJob>;
+
+	/// <summary>
+	/// Interface for a collection of artifacts
+	/// </summary>
+	public interface IArtifactCollectionV1
+	{
+		/// <summary>
+		/// Creates a new artifact
+		/// </summary>
+		/// <param name="jobId">Unique id of the job that owns this artifact</param>
+		/// <param name="stepId">Optional Step id</param>
+		/// <param name="name">Name of artifact</param>
+		/// <param name="mimeType">Type of artifact</param>
+		/// <param name="data">The data to write</param>
+		/// <returns>The new log file document</returns>
+		Task<IArtifactV1> CreateArtifactAsync(JobId jobId, SubResourceId? stepId, string name, string mimeType, System.IO.Stream data);
+
+		/// <summary>
+		/// Gets all the available artifacts for a job
+		/// </summary>
+		/// <param name="jobId">Unique id of the job to query</param>
+		/// <param name="stepId">Unique id of the Step to query</param>
+		/// <param name="name">Name of the artifact</param>
+		/// <returns>List of artifact documents</returns>
+		Task<List<IArtifactV1>> GetArtifactsAsync(JobId? jobId, SubResourceId? stepId, string? name);
+
+		/// <summary>
+		/// Gets a specific list of artifacts based on id
+		/// </summary>
+		/// <param name="artifactIds">The list of artifact Ids</param>
+		/// <returns>List of artifact documents</returns>
+		Task<List<IArtifactV1>> GetArtifactsAsync(IEnumerable<ObjectId> artifactIds);
+
+		/// <summary>
+		/// Gets an artifact by ID
+		/// </summary>
+		/// <param name="artifactId">Unique id of the artifact</param>
+		/// <returns>The artifact document</returns>
+		Task<IArtifactV1?> GetArtifactAsync(ObjectId artifactId);
+
+		/// <summary>
+		/// Updates an artifact
+		/// </summary>
+		/// <param name="artifact">The artifact</param>
+		/// <param name="newMimeType">New mime type</param>
+		/// <param name="newData">New data</param>
+		/// <returns>Async task</returns>
+		Task<bool> UpdateArtifactAsync(IArtifactV1? artifact, string newMimeType, System.IO.Stream newData);
+
+		/// <summary>
+		/// gets artifact data
+		/// </summary>
+		/// <param name="artifact">The artifact</param>
+		/// <returns>The chunk data</returns>
+		Task<System.IO.Stream> OpenArtifactReadStreamAsync(IArtifactV1 artifact);
+	}
+}
diff --git a/Engine/Source/Programs/Horde/Horde.Build/Jobs/Artifacts/IArtifactV1.cs b/Engine/Source/Programs/Horde/Horde.Build/Jobs/Artifacts/IArtifactV1.cs
new file mode 100644
--- /dev/null
+++ b/Engine/Source/Programs/Horde/Horde.Build/Jobs/Artifacts/IArtifactV1.cs
@@ -0,0 +1,45 @@
+// Copyright Epic Games, Inc. All Rights Reserved.
+
+using Horde.Build.Utilities;
+using MongoDB.Bson;
+
+namespace Horde.Build.Jobs.Artifacts
+{
+	using JobId = ObjectId<IJob>;
+
+	/// <summary>
+	/// Information about an artifact
+	/// </summary>
+	public interface IArtifactV1
+	{
+		/// <summary>
+		/// Identifier for the Artifact. Randomly generated.
+		/// </summary>
+		public ObjectId Id { get; }
+
+		/// <summary>
+		/// Unique id of the job containing this artifact
+		/// </summary>
+		public JobId JobId { get; }
+
+		/// <summary>
+		/// Unique id of the job containing this artifact
+		/// </summary>
+		public string Name { get; }
+
+		/// <summary>
+		/// Unique id of the step containing this artifact
+		/// </summary>
+		public SubResourceId? StepId { get; }
+
+		/// <summary>
+		/// Total size of the file
+		/// </summary>
+		public long Length { get; }
+
+		/// <summary>
+		/// Type of artifact
+		/// </summary>
+		public string MimeType { get; }
+	}
+}
diff --git a/Engine/Source/Programs/Horde/Horde.Build/Jobs/JobsController.cs b/Engine/Source/Programs/Horde/Horde.Build/Jobs/JobsController.cs
--- a/Engine/Source/Programs/Horde/Horde.Build/Jobs/JobsController.cs
+++ b/Engine/Source/Programs/Horde/Horde.Build/Jobs/JobsController.cs
@@ -49,7 +49,7 @@ namespace Horde.Build.Jobs
 		private readonly IStreamCollection _streamCollection;
 		private readonly JobService _jobService;
 		private readonly ITemplateCollection _templateCollection;
-		private readonly IArtifactCollection _artifactCollection;
+		private readonly IArtifactCollectionV1 _artifactCollection;
 		private readonly IUserCollection _userCollection;
 		private readonly INotificationService _notificationService;
 		private readonly AgentService _agentService;
@@ -59,7 +59,7 @@ namespace Horde.Build.Jobs
 		/// <summary>
 		/// Constructor
 		/// </summary>
-		public JobsController(IGraphCollection graphs, ICommitService commitService, IPerforceService perforce, IStreamCollection streamCollection, JobService jobService, ITemplateCollection templateCollection, IArtifactCollection artifactCollection, IUserCollection userCollection, INotificationService notificationService, AgentService agentService, IOptionsSnapshot<GlobalConfig> globalConfig, ILogger<JobsController> logger)
+		public JobsController(IGraphCollection graphs, ICommitService commitService, IPerforceService perforce, IStreamCollection streamCollection, JobService jobService, ITemplateCollection templateCollection, IArtifactCollectionV1 artifactCollection, IUserCollection userCollection, INotificationService notificationService, AgentService agentService, IOptionsSnapshot<GlobalConfig> globalConfig, ILogger<JobsController> logger)
 		{
 			_graphs = graphs;
 			_commitService = commitService;
@@ -1489,13 +1489,13 @@ namespace Horde.Build.Jobs
 				return NotFound(jobId, batchId, stepId);
 			}
 
-			List<IArtifact> artifacts = await _artifactCollection.GetArtifactsAsync(jobId, stepId, name);
+			List<IArtifactV1> artifacts = await _artifactCollection.GetArtifactsAsync(jobId, stepId, name);
 			if (artifacts.Count == 0)
 			{
 				return NotFound();
 			}
 
-			IArtifact artifact = artifacts[0];
+			IArtifactV1 artifact = artifacts[0];
 			return new FileStreamResult(await _artifactCollection.OpenArtifactReadStreamAsync(artifact), artifact.MimeType);
 		}
 
@@ -1535,8 +1535,8 @@ namespace Horde.Build.Jobs
 				return NotFound(jobId, batchId, stepId);
 			}
 
-			List<IArtifact> artifacts = await _artifactCollection.GetArtifactsAsync(jobId, stepId, null);
-			foreach (IArtifact artifact in artifacts)
+			List<IArtifactV1> artifacts = await _artifactCollection.GetArtifactsAsync(jobId, stepId, null);
+			foreach (IArtifactV1 artifact in artifacts)
 			{
 				if (artifact.Name.Equals("trace.json", StringComparison.OrdinalIgnoreCase))
 				{
diff --git a/Engine/Source/Programs/Horde/Horde.Build/Server/RpcService.cs b/Engine/Source/Programs/Horde/Horde.Build/Server/RpcService.cs
--- a/Engine/Source/Programs/Horde/Horde.Build/Server/RpcService.cs
+++ b/Engine/Source/Programs/Horde/Horde.Build/Server/RpcService.cs
@@ -72,7 +72,7 @@ namespace Horde.Build.Server
 		readonly IStreamCollection _streamCollection;
 		readonly JobService _jobService;
 		readonly AgentSoftwareService _agentSoftwareService;
-		readonly IArtifactCollection _artifactCollection;
+		readonly IArtifactCollectionV1 _artifactCollection;
 		readonly ILogFileService _logFileService;
 		readonly PoolService _poolService;
 		readonly LifetimeService _lifetimeService;
@@ -89,7 +89,7 @@ namespace Horde.Build.Server
 		/// <summary>
 		/// Constructor
 		/// </summary>
-		public RpcService(AclService aclService, AgentService agentService, IStreamCollection streamCollection, JobService jobService, AgentSoftwareService agentSoftwareService, IArtifactCollection artifactCollection, ILogFileService logFileService, PoolService poolService, LifetimeService lifetimeService, IGraphCollection graphs, ITestDataCollection testData, IJobStepRefCollection jobStepRefCollection, ITemplateCollection templateCollection, ConformTaskSource conformTaskSource, HttpClient httpClient, IClock clock, IOptionsSnapshot<GlobalConfig> globalConfig, ILogger<RpcService> logger)
+		public RpcService(AclService aclService, AgentService agentService, IStreamCollection streamCollection, JobService jobService, AgentSoftwareService agentSoftwareService, IArtifactCollectionV1 artifactCollection, ILogFileService logFileService, PoolService poolService, LifetimeService lifetimeService, IGraphCollection graphs, ITestDataCollection testData, IJobStepRefCollection jobStepRefCollection, ITemplateCollection templateCollection, ConformTaskSource conformTaskSource, HttpClient httpClient, IClock clock, IOptionsSnapshot<GlobalConfig> globalConfig, ILogger<RpcService> logger)
 		{
 			_aclService = aclService;
 			_agentService = agentService;
@@ -1059,7 +1059,7 @@ namespace Horde.Build.Server
 			// Upload the stream
 			using (ArtifactChunkStream inputStream = new ArtifactChunkStream(reader, metadata.Length))
 			{
-				IArtifact artifact = await _artifactCollection.CreateArtifactAsync(job.Id, step.Id, metadata.Name, metadata.MimeType, inputStream);
+				IArtifactV1 artifact = await _artifactCollection.CreateArtifactAsync(job.Id, step.Id, metadata.Name, metadata.MimeType, inputStream);
 
 				UploadArtifactResponse response = new UploadArtifactResponse();
 				response.Id = artifact.Id.ToString();
diff --git a/Engine/Source/Programs/Horde/Horde.Build/Startup.cs b/Engine/Source/Programs/Horde/Horde.Build/Startup.cs
--- a/Engine/Source/Programs/Horde/Horde.Build/Startup.cs
+++ b/Engine/Source/Programs/Horde/Horde.Build/Startup.cs
@@ -338,7 +338,7 @@ namespace Horde.Build
 
 			services.AddSingleton<IAgentCollection, AgentCollection>();
 			services.AddSingleton<IAgentSoftwareCollection, AgentSoftwareCollection>();
-			services.AddSingleton<IArtifactCollection, ArtifactCollection>();
+			services.AddSingleton<IArtifactCollectionV1, ArtifactCollectionV1>();
 			services.AddSingleton<IGraphCollection, GraphCollection>();
 			services.AddSingleton<IIssueCollection, IssueCollection>();
 			services.AddSingleton<IJobCollection, JobCollection>();
@@ -465,7 +465,7 @@ namespace Horde.Build
 
 			// Storage providers
 			services.AddSingleton(sp => CreateStorageBackend(sp, settings.LogStorage).ForType<PersistentLogStorage>());
-			services.AddSingleton(sp => CreateStorageBackend(sp, settings.ArtifactStorage).ForType<ArtifactCollection>());
+			services.AddSingleton(sp => CreateStorageBackend(sp, settings.ArtifactStorage).ForType<ArtifactCollectionV1>());
 
 			services.AddHordeStorage(settings => configSection.GetSection("Storage").Bind(settings));
 
