Horde: Reduce logging from fleet handling by only logging state changes

Changes fleet manager and fleet service to return a result object, describing what happened. This result together with desired and current agent count is stored in pool document and later used to see if a state change happened that is meaningful to log.

diff --git a/Engine/Source/Programs/Horde/Horde.Build.Tests/Fleet/FleetServiceTest.cs b/Engine/Source/Programs/Horde/Horde.Build.Tests/Fleet/FleetServiceTest.cs
--- a/Engine/Source/Programs/Horde/Horde.Build.Tests/Fleet/FleetServiceTest.cs
+++ b/Engine/Source/Programs/Horde/Horde.Build.Tests/Fleet/FleetServiceTest.cs
@@ -32,18 +32,18 @@ namespace Horde.Build.Tests.Fleet
 		private readonly List<IAgent> _expandAgents = new ();
 		private readonly List<IAgent> _shrinkAgents = new ();
 		
-		public Task ExpandPoolAsync(IPool pool, IReadOnlyList<IAgent> agents, int count, CancellationToken cancellationToken)
+		public Task<ScaleResult> ExpandPoolAsync(IPool pool, IReadOnlyList<IAgent> agents, int count, CancellationToken cancellationToken)
 		{
 			ExpandPoolAsyncCallCount++;
 			_expandAgents.AddRange(agents);
-			return Task.CompletedTask;
+			return Task.FromResult(new ScaleResult(FleetManagerOutcome.Success, count, 0));
 		}
 
-		public Task ShrinkPoolAsync(IPool pool, IReadOnlyList<IAgent> agents, int count, CancellationToken cancellationToken)
+		public Task<ScaleResult> ShrinkPoolAsync(IPool pool, IReadOnlyList<IAgent> agents, int count, CancellationToken cancellationToken)
 		{
 			ShrinkPoolAsyncCallCount++;
 			_shrinkAgents.AddRange(agents);
-			return Task.CompletedTask;
+			return Task.FromResult(new ScaleResult(FleetManagerOutcome.Success, 0, count));
 		}
 
 		public Task<int> GetNumStoppedInstancesAsync(IPool pool, CancellationToken cancellationToken)
@@ -52,7 +52,33 @@ namespace Horde.Build.Tests.Fleet
 		}
 	}
 
-	
+	public class FakeFleetManager : IFleetManager
+	{
+		public int AgentsAddedCount { get; private set; }
+		public int AgentsRemovedCount { get; private set; }
+
+		public ScaleResult? ForceResult { get; set; } = null;
+		
+		public Task<ScaleResult> ExpandPoolAsync(IPool pool, IReadOnlyList<IAgent> agents, int count, CancellationToken cancellationToken = default)
+		{
+			AgentsAddedCount += count;
+			ScaleResult scaleResult = new (FleetManagerOutcome.Success, count, 0, "Pool scaled out");
+			return Task.FromResult(ForceResult ?? scaleResult);
+		}
+
+		public Task<ScaleResult> ShrinkPoolAsync(IPool pool, IReadOnlyList<IAgent> agents, int count, CancellationToken cancellationToken = default)
+		{
+			AgentsRemovedCount += count;
+			ScaleResult scaleResult = new (FleetManagerOutcome.Success, 0, count, "Pool scaled in");
+			return Task.FromResult(ForceResult ?? scaleResult);
+		}
+
+		public Task<int> GetNumStoppedInstancesAsync(IPool pool, CancellationToken cancellationToken = default)
+		{
+			throw new NotImplementedException();
+		}
+	}
+
 	public class StubFleetManagerFactory : IFleetManagerFactory
 	{
 		private readonly IFleetManager _fleetManager;
@@ -72,6 +98,7 @@ namespace Horde.Build.Tests.Fleet
 	public class FleetServiceTest : TestSetup
 	{
 		readonly FleetManagerSpy _fleetManagerSpy = new();
+		readonly FakeFleetManager _fakeFleetManager = new();
 		readonly IDogStatsd _dogStatsD = new NoOpDogStatsd();
 
 		protected override void Dispose(bool disposing)
@@ -107,16 +134,19 @@ namespace Horde.Build.Tests.Fleet
 			IPool pool = await PoolService.CreatePoolAsync("testPool", null, true, 0, 0, sizeStrategy: PoolSizeStrategy.NoOp);
 
 			// First scale-out will succeed
-			await service.ScalePoolAsync(pool, new List<IAgent>(), new PoolSizeResult(0, 1), CancellationToken.None);
+			ScaleResult result = await service.ScalePoolAsync(pool, new List<IAgent>(), new PoolSizeResult(0, 1), CancellationToken.None);
+			Assert.AreEqual(new ScaleResult(FleetManagerOutcome.Success, 1, 0), result);
 			Assert.AreEqual(1, _fleetManagerSpy.ExpandPoolAsyncCallCount);
 			
 			// Cannot scale-out due to cool-down
-			await service.ScalePoolAsync(pool, new List<IAgent>(), new PoolSizeResult(0, 2), CancellationToken.None);
+			result = await service.ScalePoolAsync(pool, new List<IAgent>(), new PoolSizeResult(0, 2), CancellationToken.None);
+			Assert.AreEqual(new ScaleResult(FleetManagerOutcome.NoOp, 0, 0), result);
 			Assert.AreEqual(1, _fleetManagerSpy.ExpandPoolAsyncCallCount);
 
 			// Wait some time and then try again
 			await Clock.AdvanceAsync(TimeSpan.FromHours(2));
-			await service.ScalePoolAsync(pool, new List<IAgent>(), new PoolSizeResult(0, 2), CancellationToken.None);
+			result = await service.ScalePoolAsync(pool, new List<IAgent>(), new PoolSizeResult(0, 2), CancellationToken.None);
+			Assert.AreEqual(new ScaleResult(FleetManagerOutcome.Success, 2, 0), result);
 			Assert.AreEqual(2, _fleetManagerSpy.ExpandPoolAsyncCallCount);
 		}
 		
@@ -129,16 +159,19 @@ namespace Horde.Build.Tests.Fleet
 			IAgent agent2 = await CreateAgentAsync(pool);
 
 			// First scale-in will succeed
-			await service.ScalePoolAsync(pool, new List<IAgent> { agent1, agent2 }, new PoolSizeResult(2, 1), CancellationToken.None);
+			ScaleResult result = await service.ScalePoolAsync(pool, new List<IAgent> { agent1, agent2 }, new PoolSizeResult(2, 1), CancellationToken.None);
+			Assert.AreEqual(new ScaleResult(FleetManagerOutcome.Success, 0, 1), result);
 			Assert.AreEqual(1, _fleetManagerSpy.ShrinkPoolAsyncCallCount);
 			
 			// Cannot scale-in due to cool-down
-			await service.ScalePoolAsync(pool, new List<IAgent> { agent1 }, new PoolSizeResult(1, 0), CancellationToken.None);
+			result = await service.ScalePoolAsync(pool, new List<IAgent> { agent1 }, new PoolSizeResult(1, 0), CancellationToken.None);
+			Assert.AreEqual(new ScaleResult(FleetManagerOutcome.NoOp, 0, 0), result);
 			Assert.AreEqual(1, _fleetManagerSpy.ShrinkPoolAsyncCallCount);
 
 			// Wait some time and then try again
 			await Clock.AdvanceAsync(TimeSpan.FromHours(2));
-			await service.ScalePoolAsync(pool, new List<IAgent> { agent1 }, new PoolSizeResult(1, 0), CancellationToken.None);
+			result = await service.ScalePoolAsync(pool, new List<IAgent> { agent1 }, new PoolSizeResult(1, 0), CancellationToken.None);
+			Assert.AreEqual(new ScaleResult(FleetManagerOutcome.Success, 0, 1), result);
 			Assert.AreEqual(2, _fleetManagerSpy.ShrinkPoolAsyncCallCount);
 		}
 		
@@ -148,7 +181,8 @@ namespace Horde.Build.Tests.Fleet
 			using FleetService service = GetFleetService(_fleetManagerSpy, true);
 			IPool pool = await PoolService.CreatePoolAsync("testPool", null, true, 0, 0, sizeStrategy: PoolSizeStrategy.NoOp);
 
-			await service.ScalePoolAsync(pool, new List<IAgent>(), new PoolSizeResult(0, 1), CancellationToken.None);
+			ScaleResult result = await service.ScalePoolAsync(pool, new List<IAgent>(), new PoolSizeResult(0, 1), CancellationToken.None);
+			Assert.AreEqual(new ScaleResult(FleetManagerOutcome.NoOp, 0, 0), result);
 			Assert.AreEqual(0, _fleetManagerSpy.ExpandPoolAsyncCallCount);
 		}
 		
diff --git a/Engine/Source/Programs/Horde/Horde.Build/Agents/Fleet/FleetService.cs b/Engine/Source/Programs/Horde/Horde.Build/Agents/Fleet/FleetService.cs
--- a/Engine/Source/Programs/Horde/Horde.Build/Agents/Fleet/FleetService.cs
+++ b/Engine/Source/Programs/Horde/Horde.Build/Agents/Fleet/FleetService.cs
@@ -198,26 +198,29 @@ namespace Horde.Build.Agents.Fleet
 			return pools.Select(pool => new PoolWithAgents(pool, GetAgentsInPool(pool.Id))).ToList();
 		}
 
-		internal async Task ScalePoolAsync(IPool pool, List<IAgent> agents, PoolSizeResult poolSizeResult, CancellationToken cancellationToken)
+		internal async Task<ScaleResult> ScalePoolAsync(IPool pool, List<IAgent> agents, PoolSizeResult poolSizeResult, CancellationToken cancellationToken)
 		{
 			if (!pool.EnableAutoscaling)
 			{
-				return;
+				return new ScaleResult(FleetManagerOutcome.NoOp, 0, 0, "Auto-scaling disabled");
 			}
 
+			ScaleResult result = new (FleetManagerOutcome.NoOp, 0, 0);
 			int currentAgentCount = poolSizeResult.CurrentAgentCount;
 			int desiredAgentCount = poolSizeResult.DesiredAgentCount;
 			int deltaAgentCount = desiredAgentCount - currentAgentCount;
 
+			IFleetManager fleetManager = CreateFleetManager(pool);
+			
 			using IScope scope = GlobalTracer.Instance
 				.BuildSpan("FleetService.ScalePool")
 				.WithTag(Datadog.Trace.OpenTracing.DatadogTags.ResourceName, pool.Id.ToString())
 				.WithTag("CurrentAgentCount", currentAgentCount)
 				.WithTag("DesiredAgentCount", desiredAgentCount)
 				.WithTag("DeltaAgentCount", deltaAgentCount)
+				.WithTag("FleetManager", fleetManager.GetType().Name)
 				.StartActive();
-			
-			IFleetManager fleetManager = CreateFleetManager(pool);
+
 			Dictionary<string, object> logScopeMetadata = new()
 			{
 				["FleetManager"] = fleetManager.GetType().Name,
@@ -226,62 +229,64 @@ namespace Horde.Build.Agents.Fleet
 			};
 
 			using IDisposable logScope = _logger.BeginScope(logScopeMetadata);
-			
-			_logger.LogInformation("{PoolName} Current={Current} Target={Target} Delta={Delta}",
-				pool.Name, currentAgentCount, desiredAgentCount, deltaAgentCount);
+			if (pool.LastAgentCount != currentAgentCount || pool.LastDesiredAgentCount != desiredAgentCount)
+			{
+				_logger.LogInformation("{PoolName} Current={Current} Target={Target} Delta={Delta}",
+					pool.Name, currentAgentCount, desiredAgentCount, deltaAgentCount);	
+			}
 
+			_dogStatsd.Gauge("agentpools.autoscale.target", desiredAgentCount, tags: new []{"pool:" + pool.Name});
+			_dogStatsd.Gauge("agentpools.autoscale.current", currentAgentCount, tags: new []{"pool:" + pool.Name});
+
+			DateTime? scaleOutTime = null;
+			DateTime? scaleInTime = null;
+			
+			TimeSpan scaleOutCooldown = pool.ScaleOutCooldown ?? _defaultScaleOutCooldown;
+			bool isScaleOutCoolingDown = pool.LastScaleUpTime != null && pool.LastScaleUpTime + scaleOutCooldown > _clock.UtcNow;
+			TimeSpan? scaleOutCooldownTimeLeft = pool.LastScaleUpTime + _defaultScaleOutCooldown - _clock.UtcNow;
+			scope.Span.SetTag("ScaleOutCooldownTimeLeftSecs", scaleOutCooldownTimeLeft?.TotalSeconds ?? -1);
+			
+			TimeSpan scaleInCooldown = pool.ScaleInCooldown ?? _defaultScaleInCooldown;
+			bool isScaleInCoolingDown = pool.LastScaleDownTime != null && pool.LastScaleDownTime + scaleInCooldown > _clock.UtcNow;
+			TimeSpan? scaleInCooldownTimeLeft = pool.LastScaleDownTime + _defaultScaleInCooldown - _clock.UtcNow;
+			scope.Span.SetTag("ScaleOutCooldownTimeLeftSecs", scaleInCooldownTimeLeft?.TotalSeconds ?? -1);
+			
+			scope.Span.SetTag("IsDowntimeActive", _downtimeService.IsDowntimeActive);
+			
 			try
 			{
-				if (deltaAgentCount > 0)
+				if (deltaAgentCount > 0 && !_downtimeService.IsDowntimeActive && !isScaleOutCoolingDown)
 				{
-					if (!_downtimeService.IsDowntimeActive)
-					{
-						TimeSpan scaleOutCooldown = pool.ScaleOutCooldown ?? _defaultScaleOutCooldown;
-						bool isCoolingDown = pool.LastScaleUpTime != null && pool.LastScaleUpTime + scaleOutCooldown > _clock.UtcNow;
-						scope.Span.SetTag("isCoolingDown", isCoolingDown);
-						if (!isCoolingDown)
-						{
-							await fleetManager.ExpandPoolAsync(pool, agents, deltaAgentCount, cancellationToken);
-							await _poolCollection.TryUpdateAsync(pool, lastScaleUpTime: _clock.UtcNow);
-						}
-						else
-						{
-							TimeSpan? cooldownTimeLeft = pool.LastScaleUpTime + _defaultScaleOutCooldown - _clock.UtcNow;
-							_logger.LogDebug("Cannot scale out {PoolName}, it's cooling down for another {TimeLeft} secs", pool.Name, cooldownTimeLeft?.TotalSeconds);
-						}
-					}
-					else
-					{
-						_logger.LogDebug("Cannot scale out {PoolName}, downtime is active", pool.Name);
-						scope.Span.SetTag("IsDowntimeActive", true);
-					}
+					result = await fleetManager.ExpandPoolAsync(pool, agents, deltaAgentCount, cancellationToken);
+					scaleOutTime = _clock.UtcNow;
 				}
 
-				if (deltaAgentCount < 0)
+				if (deltaAgentCount < 0 && !isScaleInCoolingDown)
 				{
-					TimeSpan scaleInCooldown = pool.ScaleInCooldown ?? _defaultScaleInCooldown;
-					bool isCoolingDown = pool.LastScaleDownTime != null && pool.LastScaleDownTime + scaleInCooldown > _clock.UtcNow;
-					scope.Span.SetTag("isCoolingDown", isCoolingDown);
-					if (!isCoolingDown)
-					{
-						await fleetManager.ShrinkPoolAsync(pool, agents, -deltaAgentCount, cancellationToken);
-						await _poolCollection.TryUpdateAsync(pool, lastScaleDownTime: _clock.UtcNow);
-					}
-					else
-					{
-						TimeSpan? cooldownTimeLeft = pool.LastScaleDownTime + _defaultScaleInCooldown - _clock.UtcNow;
-						_logger.LogDebug("Cannot scale in {PoolName}, it's cooling down for another {TimeLeft} secs", pool.Name, cooldownTimeLeft?.TotalSeconds);
-					}
+					result = await fleetManager.ShrinkPoolAsync(pool, agents, -deltaAgentCount, cancellationToken);
+					scaleInTime = _clock.UtcNow;
 				}
 			}
 			catch (Exception ex)
 			{
 				_logger.LogInformation(ex, "Failed to scale {PoolName}:\n{Exception}", pool.Name, ex);
-				return;
+				return new ScaleResult(FleetManagerOutcome.Failure, 0, 0);
 			}
 
-			_dogStatsd.Gauge("agentpools.autoscale.target", desiredAgentCount, tags: new []{"pool:" + pool.Name});
-			_dogStatsd.Gauge("agentpools.autoscale.current", currentAgentCount, tags: new []{"pool:" + pool.Name});
+			if (pool.LastScaleResult == null || !pool.LastScaleResult.Equals(result))
+			{
+				_logger.LogInformation("Scale result: {Result}", result.Message);
+			}
+
+			await _poolCollection.TryUpdateAsync(
+				pool,
+				lastScaleUpTime: scaleOutTime,
+				lastScaleDownTime: scaleInTime,
+				lastScaleResult: result,
+				lastAgentCount: currentAgentCount,
+				lastDesiredAgentCount: desiredAgentCount);
+
+			return result;
 		}
 		
 		internal IEnumerable<string> GetPropValues(string name)
diff --git a/Engine/Source/Programs/Horde/Horde.Build/Agents/Fleet/IFleetManager.cs b/Engine/Source/Programs/Horde/Horde.Build/Agents/Fleet/IFleetManager.cs
--- a/Engine/Source/Programs/Horde/Horde.Build/Agents/Fleet/IFleetManager.cs
+++ b/Engine/Source/Programs/Horde/Horde.Build/Agents/Fleet/IFleetManager.cs
@@ -1,5 +1,6 @@
 // Copyright Epic Games, Inc. All Rights Reserved.
 
+using System;
 using System.Collections.Generic;
 using System.Threading;
 using System.Threading.Tasks;
@@ -43,6 +44,101 @@ namespace Horde.Build.Agents.Fleet
 		/// </summary>
 		AwsAsg
 	}
+
+	/// <summary>
+	/// Outcome of a scaling operation
+	/// </summary>
+	public enum FleetManagerOutcome
+	{
+		/// <summary>
+		/// Scaling operation completed as intended.
+		/// </summary>
+		Success,
+		
+		/// <summary>
+		/// Scaling operation was only partly fulfilled.
+		/// </summary>
+		PartialSuccess,
+		
+		/// <summary>
+		/// Scaling operation failed
+		/// </summary>
+		Failure,
+		
+		/// <summary>
+		/// No operation took place (disabled or skipped)
+		/// </summary>
+		NoOp
+	}
+
+	/// <summary>
+	/// Result from a scaling operation
+	/// </summary>
+	public class ScaleResult
+	{
+		/// <summary>
+		/// Outcome
+		/// </summary>
+		public FleetManagerOutcome Outcome { get; }
+		
+		/// <summary>
+		/// Agents added as part of operation
+		/// </summary>
+		public int AgentsAddedCount { get; }
+		
+		/// <summary>
+		/// Agents added as part of operation
+		/// </summary>
+		public int AgentsRemovedCount { get; }
+		
+		/// <summary>
+		/// Human-readable log message
+		/// </summary>
+		public string Message { get; }
+
+		/// <summary>
+		/// Constructor
+		/// </summary>
+		/// <param name="outcome"></param>
+		/// <param name="agentsAddedCount"></param>
+		/// <param name="agentsRemovedCount"></param>
+		/// <param name="message"></param>
+		public ScaleResult(FleetManagerOutcome outcome, int agentsAddedCount, int agentsRemovedCount, string message = "")
+		{
+			Outcome = outcome;
+			AgentsAddedCount = agentsAddedCount;
+			AgentsRemovedCount = agentsRemovedCount;
+			Message = message;
+		}
+
+		/// <inheritdoc/>
+		protected bool Equals(ScaleResult other)
+		{
+			return Outcome == other.Outcome && AgentsAddedCount == other.AgentsAddedCount && AgentsRemovedCount == other.AgentsRemovedCount;
+		}
+
+		/// <inheritdoc/>
+		public override bool Equals(object? obj)
+		{
+			if (ReferenceEquals(null, obj)) { return false; }
+			if (ReferenceEquals(this, obj)) { return true; }
+			if (obj.GetType() != GetType()) { return false; }
+			return Equals((ScaleResult)obj);
+		}
+
+		/// <inheritdoc/>
+		public override int GetHashCode()
+		{
+			return HashCode.Combine((int)Outcome, AgentsAddedCount, AgentsRemovedCount);
+		}
+
+		/// <inheritdoc/>
+		public override string ToString()
+		{
+			return $"Outcome={Outcome} Added={AgentsAddedCount} Removed={AgentsRemovedCount}";
+		}
+	}
+	
 	
 	/// <summary>
 	/// Service to manage a fleet of machines
@@ -57,7 +153,7 @@ namespace Horde.Build.Agents.Fleet
 		/// <param name="count">Number of agents to add</param>
 		/// <param name="cancellationToken">Cancellation token for the call</param>
 		/// <returns>Async task</returns>
-		Task ExpandPoolAsync(IPool pool, IReadOnlyList<IAgent> agents, int count, CancellationToken cancellationToken = default);
+		Task<ScaleResult> ExpandPoolAsync(IPool pool, IReadOnlyList<IAgent> agents, int count, CancellationToken cancellationToken = default);
 
 		/// <summary>
 		/// Shrink the given pool
@@ -67,7 +163,7 @@ namespace Horde.Build.Agents.Fleet
 		/// <param name="count">Number of agents to remove</param>
 		/// <param name="cancellationToken">Cancellation token for the call</param>
 		/// <returns>Async task</returns>
-		Task ShrinkPoolAsync(IPool pool, IReadOnlyList<IAgent> agents, int count, CancellationToken cancellationToken = default);
+		Task<ScaleResult> ShrinkPoolAsync(IPool pool, IReadOnlyList<IAgent> agents, int count, CancellationToken cancellationToken = default);
 
 		/// <summary>
 		/// Returns the number of stopped instances in the given pool
diff --git a/Engine/Source/Programs/Horde/Horde.Build/Agents/Fleet/Providers/AwsAsgFleetManager.cs b/Engine/Source/Programs/Horde/Horde.Build/Agents/Fleet/Providers/AwsAsgFleetManager.cs
--- a/Engine/Source/Programs/Horde/Horde.Build/Agents/Fleet/Providers/AwsAsgFleetManager.cs
+++ b/Engine/Source/Programs/Horde/Horde.Build/Agents/Fleet/Providers/AwsAsgFleetManager.cs
@@ -59,7 +59,7 @@ namespace Horde.Build.Agents.Fleet.Providers
 		}
 
 		/// <inheritdoc/>
-		public async Task ExpandPoolAsync(IPool pool, IReadOnlyList<IAgent> agents, int count, CancellationToken cancellationToken)
+		public async Task<ScaleResult> ExpandPoolAsync(IPool pool, IReadOnlyList<IAgent> agents, int count, CancellationToken cancellationToken)
 		{
 			int desiredCapacity = agents.Count + count;
 			
@@ -70,10 +70,11 @@ namespace Horde.Build.Agents.Fleet.Providers
 			scope.Span.SetTag("desiredCapacity", desiredCapacity);
 			
 			await UpdateAsg(pool.Id, desiredCapacity, scope, cancellationToken);
+			return new ScaleResult(FleetManagerOutcome.Success, count, 0);
 		}
 
 		/// <inheritdoc/>
-		public async Task ShrinkPoolAsync(IPool pool, IReadOnlyList<IAgent> agents, int count, CancellationToken cancellationToken)
+		public async Task<ScaleResult> ShrinkPoolAsync(IPool pool, IReadOnlyList<IAgent> agents, int count, CancellationToken cancellationToken)
 		{
 			int desiredCapacity = Math.Max(0, agents.Count - count);
 			
@@ -84,6 +85,7 @@ namespace Horde.Build.Agents.Fleet.Providers
 			scope.Span.SetTag("desiredCapacity", desiredCapacity);
 			
 			await UpdateAsg(pool.Id, desiredCapacity, scope, cancellationToken);
+			return new ScaleResult(FleetManagerOutcome.Success, 0, count);
 		}
 		
 		private async Task UpdateAsg(StringId<IPool> poolId, int desiredCapacity, IScope scope, CancellationToken cancellationToken)
diff --git a/Engine/Source/Programs/Horde/Horde.Build/Agents/Fleet/Providers/AwsFleetManager.cs b/Engine/Source/Programs/Horde/Horde.Build/Agents/Fleet/Providers/AwsFleetManager.cs
--- a/Engine/Source/Programs/Horde/Horde.Build/Agents/Fleet/Providers/AwsFleetManager.cs
+++ b/Engine/Source/Programs/Horde/Horde.Build/Agents/Fleet/Providers/AwsFleetManager.cs
@@ -153,7 +153,7 @@ namespace Horde.Build.Agents.Fleet.Providers
 		}
 
 		/// <inheritdoc/>
-		public async Task ExpandPoolAsync(IPool pool, IReadOnlyList<IAgent> agents, int count, CancellationToken cancellationToken)
+		public async Task<ScaleResult> ExpandPoolAsync(IPool pool, IReadOnlyList<IAgent> agents, int count, CancellationToken cancellationToken)
 		{
 			using IScope scope = GlobalTracer.Instance.BuildSpan("ExpandPool").StartActive();
 			scope.Span.SetTag("poolName", pool.Name);
@@ -205,16 +205,24 @@ namespace Horde.Build.Agents.Fleet.Providers
 				_logger.LogInformation("Created instance {InstanceId} for pool {PoolId}", i.InstanceId, pool.Id);
 			}
 
+			if (numStartedInstances == count)
+			{
+				return new ScaleResult(FleetManagerOutcome.Success, numStartedInstances, 0);
+			}
+			
 			if (numStartedInstances != count)
 			{
 				_logger.LogWarning("Unable to create all the requested instances for pool {PoolId}. RequestedCount={RequestedCount} ActualCount={ActualCount}", pool.Id, count, numStartedInstances);
 			}
+			
+			return new ScaleResult(numStartedInstances == 0 ? FleetManagerOutcome.Failure : FleetManagerOutcome.Success, numStartedInstances, 0);
 		}
 
 		/// <inheritdoc/>
-		public Task ShrinkPoolAsync(IPool pool, IReadOnlyList<IAgent> agents, int count, CancellationToken cancellationToken)
+		public async Task<ScaleResult> ShrinkPoolAsync(IPool pool, IReadOnlyList<IAgent> agents, int count, CancellationToken cancellationToken)
 		{
-			return ShrinkPoolViaAgentShutdownRequestAsync(_agentCollection, pool, agents, count, cancellationToken);
+			await ShrinkPoolViaAgentShutdownRequestAsync(_agentCollection, pool, agents, count, cancellationToken);
+			return new ScaleResult(FleetManagerOutcome.Success, 0, count);
 		}
 		
 		/// <summary>
diff --git a/Engine/Source/Programs/Horde/Horde.Build/Agents/Fleet/Providers/AwsRecyclingFleetManager.cs b/Engine/Source/Programs/Horde/Horde.Build/Agents/Fleet/Providers/AwsRecyclingFleetManager.cs
--- a/Engine/Source/Programs/Horde/Horde.Build/Agents/Fleet/Providers/AwsRecyclingFleetManager.cs
+++ b/Engine/Source/Programs/Horde/Horde.Build/Agents/Fleet/Providers/AwsRecyclingFleetManager.cs
@@ -106,7 +106,7 @@ public sealed class AwsRecyclingFleetManager : IFleetManager
 	}
 
 	/// <inheritdoc/>
-	public async Task ExpandPoolAsync(IPool pool, IReadOnlyList<IAgent> agents, int requestedInstancesCount, CancellationToken cancellationToken)
+	public async Task<ScaleResult> ExpandPoolAsync(IPool pool, IReadOnlyList<IAgent> agents, int requestedInstancesCount, CancellationToken cancellationToken)
 	{
 		using IScope scope = GlobalTracer.Instance.BuildSpan("ExpandPool").StartActive();
 		scope.Span.SetTag("PoolId", pool.Id.ToString());
@@ -134,20 +134,36 @@ public sealed class AwsRecyclingFleetManager : IFleetManager
 			["StoppedInstancesPerAz"] = stoppedInstancesPerAz.ToDictionary(x => x.Key, x => x.Value.Select(y => y.InstanceId)),
 			["InstanceTypePriority"] = Settings.InstanceTypes,
 		};
+		
+		scope.Span.SetTag("InstancesToStartCount", instancesToStartCount);
+		scope.Span.SetTag("StoppedInstancesMissingCount", stoppedInstancesMissingCount);
 
+		ScaleResult result;
 		using (_logger.BeginScope(logScopeMetadata))
 		{
 			if (instancesToStartCount == 0)
 			{
-				_logger.LogInformation("Unable to start any instance(s). No stopped instances are available.");
+				result = new ScaleResult(
+					FleetManagerOutcome.Failure,
+					instancesToStartCount,
+					0,
+					"Unable to start any instance(s). No stopped instances are available");
 			}
 			else if (stoppedInstancesMissingCount > 0)
 			{
-				_logger.LogInformation("Starting {InstancesToStartCount} instance(s) but not enough stopped instances to accommodate the full pool scale-out", instancesToStartCount);
+				result = new ScaleResult(
+					FleetManagerOutcome.PartialSuccess,
+					instancesToStartCount,
+					0,
+					$"Starting {instancesToStartCount} instance(s) but not enough stopped instances to accommodate the full pool scale-out");
 			}
 			else
 			{
-				_logger.LogInformation("Starting {InstancesToStartCount} instance(s)", instancesToStartCount);	
+				result = new ScaleResult(
+					FleetManagerOutcome.Success,
+					instancesToStartCount,
+					0,
+					$"Starting {instancesToStartCount} instance(s)");
 			}
 		}
 
@@ -159,6 +175,8 @@ public sealed class AwsRecyclingFleetManager : IFleetManager
 				await StartInstancesWithRetriesAsync(instances, instanceTypePriority, cancellationToken);					
 			}
 		}
+
+		return result;
 	}
 
 	private static Dictionary<string, List<Instance>> GetInstancesToLaunch(
@@ -404,9 +422,10 @@ public sealed class AwsRecyclingFleetManager : IFleetManager
 	}
 
 	/// <inheritdoc/>
-	public Task ShrinkPoolAsync(IPool pool, IReadOnlyList<IAgent> agents, int count, CancellationToken cancellationToken)
+	public async Task<ScaleResult> ShrinkPoolAsync(IPool pool, IReadOnlyList<IAgent> agents, int count, CancellationToken cancellationToken)
 	{
-		return AwsFleetManager.ShrinkPoolViaAgentShutdownRequestAsync(_agentCollection, pool, agents, count, cancellationToken);
+		await AwsFleetManager.ShrinkPoolViaAgentShutdownRequestAsync(_agentCollection, pool, agents, count, cancellationToken);
+		return new ScaleResult(FleetManagerOutcome.Success, 0, count);
 	} 
 
 	/// <inheritdoc/>
diff --git a/Engine/Source/Programs/Horde/Horde.Build/Agents/Fleet/Providers/AwsReuseFleetManager.cs b/Engine/Source/Programs/Horde/Horde.Build/Agents/Fleet/Providers/AwsReuseFleetManager.cs
--- a/Engine/Source/Programs/Horde/Horde.Build/Agents/Fleet/Providers/AwsReuseFleetManager.cs
+++ b/Engine/Source/Programs/Horde/Horde.Build/Agents/Fleet/Providers/AwsReuseFleetManager.cs
@@ -63,7 +63,7 @@ namespace Horde.Build.Agents.Fleet.Providers
 		}
 
 		/// <inheritdoc/>
-		public async Task ExpandPoolAsync(IPool pool, IReadOnlyList<IAgent> agents, int requestedInstancesCount, CancellationToken cancellationToken)
+		public async Task<ScaleResult> ExpandPoolAsync(IPool pool, IReadOnlyList<IAgent> agents, int requestedInstancesCount, CancellationToken cancellationToken)
 		{
 			using IScope scope = GlobalTracer.Instance.BuildSpan("ExpandPool").StartActive();
 			scope.Span.SetTag("poolName", pool.Name);
@@ -143,13 +143,26 @@ namespace Horde.Build.Agents.Fleet.Providers
 					                       "StartedInstancesCount={StartedInstancesCount}",
 						pool.Name, reason.Trim(), requestedInstancesCount, stoppedInstancesCount, startedInstancesCount);
 				}
+
+				if (startedInstancesCount == requestedInstancesCount)
+				{
+					return new ScaleResult(FleetManagerOutcome.Success, startedInstancesCount, 0);
+				}
+				
+				if (startedInstancesCount > 0 && startedInstancesCount < requestedInstancesCount)
+				{
+					return new ScaleResult(FleetManagerOutcome.PartialSuccess, startedInstancesCount, 0);
+				}
+				
+				return new ScaleResult(FleetManagerOutcome.Failure, startedInstancesCount, 0);
 			}
 		}
 
 		/// <inheritdoc/>
-		public Task ShrinkPoolAsync(IPool pool, IReadOnlyList<IAgent> agents, int count, CancellationToken cancellationToken)
+		public async Task<ScaleResult> ShrinkPoolAsync(IPool pool, IReadOnlyList<IAgent> agents, int count, CancellationToken cancellationToken)
 		{
-			return AwsFleetManager.ShrinkPoolViaAgentShutdownRequestAsync(_agentCollection, pool, agents, count, cancellationToken);
+			await AwsFleetManager.ShrinkPoolViaAgentShutdownRequestAsync(_agentCollection, pool, agents, count, cancellationToken);
+			return new ScaleResult(FleetManagerOutcome.Success, 0, count);
 		} 
 
 		/// <inheritdoc/>
diff --git a/Engine/Source/Programs/Horde/Horde.Build/Agents/Fleet/Providers/NoOpFleetManager.cs b/Engine/Source/Programs/Horde/Horde.Build/Agents/Fleet/Providers/NoOpFleetManager.cs
--- a/Engine/Source/Programs/Horde/Horde.Build/Agents/Fleet/Providers/NoOpFleetManager.cs
+++ b/Engine/Source/Programs/Horde/Horde.Build/Agents/Fleet/Providers/NoOpFleetManager.cs
@@ -25,17 +25,17 @@ namespace Horde.Build.Agents.Fleet.Providers
 		}
 
 		/// <inheritdoc/>
-		public Task ExpandPoolAsync(IPool pool, IReadOnlyList<IAgent> agents, int count, CancellationToken cancellationToken)
+		public Task<ScaleResult> ExpandPoolAsync(IPool pool, IReadOnlyList<IAgent> agents, int count, CancellationToken cancellationToken)
 		{
 			_logger.LogInformation("Expand pool {PoolId} by {Count} agents", pool.Id, count);
-			return Task.CompletedTask;
+			return Task.FromResult(new ScaleResult(FleetManagerOutcome.Success, count, 0));
 		}
 
 		/// <inheritdoc/>
-		public Task ShrinkPoolAsync(IPool pool, IReadOnlyList<IAgent> agents, int count, CancellationToken cancellationToken)
+		public Task<ScaleResult> ShrinkPoolAsync(IPool pool, IReadOnlyList<IAgent> agents, int count, CancellationToken cancellationToken)
 		{
 			_logger.LogInformation("Shrink pool {PoolId} by {Count} agents", pool.Id, count);
-			return Task.CompletedTask;
+			return Task.FromResult(new ScaleResult(FleetManagerOutcome.Success, 0, count));
 		}
 
 		/// <inheritdoc/>
diff --git a/Engine/Source/Programs/Horde/Horde.Build/Agents/Pools/IPoolCollection.cs b/Engine/Source/Programs/Horde/Horde.Build/Agents/Pools/IPoolCollection.cs
--- a/Engine/Source/Programs/Horde/Horde.Build/Agents/Pools/IPoolCollection.cs
+++ b/Engine/Source/Programs/Horde/Horde.Build/Agents/Pools/IPoolCollection.cs
@@ -130,6 +130,9 @@ namespace Horde.Build.Agents.Pools
 		/// <param name="lastScaleDownTime">New time for last (auto) scale down</param>
 		/// <param name="scaleOutCooldown">Cooldown time between scale-out events</param>
 		/// <param name="scaleInCooldown">Cooldown time between scale-in events</param>
+		/// <param name="lastScaleResult">Result from last scaling in/out attempt</param>
+		/// <param name="lastAgentCount">Last calculated agent count</param>
+		/// <param name="lastDesiredAgentCount">Last calculated desired agent count</param>
 		/// <param name="sizeStrategy">Pool sizing strategy</param>
 		/// <param name="newSizeStrategies">List of pool sizing strategies</param>
 		/// <param name="newFleetManagers">List of fleet managers</param>
@@ -153,6 +156,9 @@ namespace Horde.Build.Agents.Pools
 			DateTime? lastScaleDownTime = null,
 			TimeSpan? scaleOutCooldown = null,
 			TimeSpan? scaleInCooldown = null,
+			ScaleResult? lastScaleResult = null,
+			int? lastAgentCount = null,
+			int? lastDesiredAgentCount = null,
 			PoolSizeStrategy? sizeStrategy = null,
 			List<PoolSizeStrategyInfo>? newSizeStrategies = null,
 			List<FleetManagerInfo>? newFleetManagers = null,
diff --git a/Engine/Source/Programs/Horde/Horde.Build/Agents/Pools/PoolCollection.cs b/Engine/Source/Programs/Horde/Horde.Build/Agents/Pools/PoolCollection.cs
--- a/Engine/Source/Programs/Horde/Horde.Build/Agents/Pools/PoolCollection.cs
+++ b/Engine/Source/Programs/Horde/Horde.Build/Agents/Pools/PoolCollection.cs
@@ -64,6 +64,15 @@ namespace Horde.Build.Agents.Pools
 			[BsonIgnoreIfNull]
 			public TimeSpan? ScaleInCooldown { get; set; }
 
+			[BsonIgnoreIfNull]
+			public ScaleResult? LastScaleResult { get; set; }
+
+			[BsonIgnoreIfNull]
+			public int? LastAgentCount { get; set; }
+			
+			[BsonIgnoreIfNull]
+			public int? LastDesiredAgentCount { get; set; }
+
 			[BsonIgnoreIfNull]
 			public PoolSizeStrategy? SizeStrategy { get; set; }
 
@@ -282,6 +291,9 @@ namespace Horde.Build.Agents.Pools
 			DateTime? lastScaleDownTime,
 			TimeSpan? scaleOutCooldown,
 			TimeSpan? scaleInCooldown, 
+			ScaleResult? lastScaleResult, 
+			int? lastAgentCount, 
+			int? lastDesiredAgentCount, 
 			PoolSizeStrategy? sizeStrategy,
 			List<PoolSizeStrategyInfo>? newSizeStrategies,
 			List<FleetManagerInfo>? newFleetManagers,
@@ -378,6 +390,18 @@ namespace Horde.Build.Agents.Pools
 			{
 				transaction.Set(x => x.ScaleInCooldown, scaleInCooldown);
 			}
+			if (lastScaleResult != null)
+			{
+				transaction.Set(x => x.LastScaleResult, lastScaleResult);
+			}
+			if (lastAgentCount != null)
+			{
+				transaction.Set(x => x.LastAgentCount, lastAgentCount);
+			}
+			if (lastDesiredAgentCount != null)
+			{
+				transaction.Set(x => x.LastDesiredAgentCount, lastDesiredAgentCount);
+			}
 			if (sizeStrategy != null)
 			{
 				transaction.Set(x => x.SizeStrategy, sizeStrategy);				
diff --git a/Engine/Source/Programs/Horde/Horde.Build/Agents/Pools/PoolConfig.cs b/Engine/Source/Programs/Horde/Horde.Build/Agents/Pools/PoolConfig.cs
--- a/Engine/Source/Programs/Horde/Horde.Build/Agents/Pools/PoolConfig.cs
+++ b/Engine/Source/Programs/Horde/Horde.Build/Agents/Pools/PoolConfig.cs
@@ -79,6 +79,21 @@ namespace Horde.Build.Agents.Pools
 		/// Cooldown time between scale-in events
 		/// </summary>
 		public TimeSpan? ScaleInCooldown { get; }
+		
+		/// <summary>
+		/// Last result from scaling the pool
+		/// </summary>
+		public ScaleResult? LastScaleResult { get; }
+		
+		/// <summary>
+		/// Last known agent count
+		/// </summary>
+		public int? LastAgentCount { get; }
+
+		/// <summary>
+		/// Last known desired agent count
+		/// </summary>
+		public int? LastDesiredAgentCount { get; }
 
 		/// <summary>
 		/// Pool sizing strategy to be used for this pool
@@ -158,7 +173,16 @@ namespace Horde.Build.Agents.Pools
 
 		/// <inheritdoc/>
 		public TimeSpan? ScaleInCooldown { get; set; }
+		
+		/// <inheritdoc/>
+		public ScaleResult? LastScaleResult { get; set; }
 
+		/// <inheritdoc/>
+		public int? LastAgentCount { get; set; }
+		
+		/// <inheritdoc/>
+		public int? LastDesiredAgentCount { get; set; }
+		
 		/// <inheritdoc/>
 		public PoolSizeStrategy? SizeStrategy { get; set; }
 
