Horde: Guard against a null JobOptions object.

diff --git a/Engine/Source/Programs/Horde/Horde.Agent/Leases/Handlers/JobHandler.cs b/Engine/Source/Programs/Horde/Horde.Agent/Leases/Handlers/JobHandler.cs
--- a/Engine/Source/Programs/Horde/Horde.Agent/Leases/Handlers/JobHandler.cs
+++ b/Engine/Source/Programs/Horde/Horde.Agent/Leases/Handlers/JobHandler.cs
@@ -57,6 +57,8 @@ namespace Horde.Agent.Leases.Handlers
 		/// <inheritdoc/>
 		public override async Task<LeaseResult> ExecuteAsync(ISession session, string leaseId, ExecuteJobTask executeTask, CancellationToken cancellationToken)
 		{
+			executeTask.JobOptions ??= new JobOptions();
+
 			// Create a storage client for this session
 			JobOptions jobOptions = executeTask.JobOptions;
 			await using IServerLogger logger = _serverLoggerFactory.CreateLogger(session, executeTask.LogId, executeTask.JobId, executeTask.BatchId, null, null, jobOptions.UseNewLogStorage);
