Horde: Change StringId types to be implemented as records rather than using aliases.

diff --git a/Engine/Source/Programs/Horde/Horde.Build.Tests/BlobStoreTests.cs b/Engine/Source/Programs/Horde/Horde.Build.Tests/BlobStoreTests.cs
--- a/Engine/Source/Programs/Horde/Horde.Build.Tests/BlobStoreTests.cs
+++ b/Engine/Source/Programs/Horde/Horde.Build.Tests/BlobStoreTests.cs
@@ -16,8 +16,6 @@ using Horde.Build.Utilities;
 
 namespace Horde.Build.Tests
 {
-	using BackendId = StringId<IStorageBackend>;
-
 	[TestClass]
 	public class BlobStoreTests : TestSetup
 	{
diff --git a/Engine/Source/Programs/Horde/Horde.Build.Tests/DeviceServiceTests.cs b/Engine/Source/Programs/Horde/Horde.Build.Tests/DeviceServiceTests.cs
--- a/Engine/Source/Programs/Horde/Horde.Build.Tests/DeviceServiceTests.cs
+++ b/Engine/Source/Programs/Horde/Horde.Build.Tests/DeviceServiceTests.cs
@@ -12,12 +12,10 @@ using System.Security.Claims;
 using Microsoft.AspNetCore.Http;
 using Horde.Build.Server;
 using System;
+using EpicGames.Horde;
 
 namespace Horde.Build.Tests
 {
-	using DeviceId = StringId<IDevice>;
-	using DevicePlatformId = StringId<IDevicePlatform>;
-
 	/// <summary>
 	/// Tests for the device service
 	/// </summary>
@@ -74,7 +72,7 @@ namespace Horde.Build.Tests
 				}
 
 				// add models
-				await DevicesController.UpdatePlatformAsync(DevicePlatformId.Sanitize(platformName).ToString(), new UpdateDevicePlatformRequest() { ModelIds = modelIds.ToArray() });
+				await DevicesController.UpdatePlatformAsync(new DevicePlatformId(StringId.Sanitize(platformName)).ToString(), new UpdateDevicePlatformRequest() { ModelIds = modelIds.ToArray() });
 
 				platformMap[platformName] = platformName;
 
@@ -82,7 +80,7 @@ namespace Horde.Build.Tests
 				if (i == 3)
 				{
 					DevicePlatformConfig config = new DevicePlatformConfig();
-					config.Id = DevicePlatformId.Sanitize(platformName).ToString();
+					config.Id = new DevicePlatformId(StringId.Sanitize(platformName)).ToString();
 					config.Names.Add("TestDevicePlatform3Alias");
 					config.LegacyPerfSpecHighModel = "TestDevicePlatform3_Model4";
 					Devices.Platforms.Add(config);
diff --git a/Engine/Source/Programs/Horde/Horde.Build.Tests/Fixture.cs b/Engine/Source/Programs/Horde/Horde.Build.Tests/Fixture.cs
--- a/Engine/Source/Programs/Horde/Horde.Build.Tests/Fixture.cs
+++ b/Engine/Source/Programs/Horde/Horde.Build.Tests/Fixture.cs
@@ -21,10 +21,6 @@ using HordeCommon;
 namespace Horde.Build.Tests
 {
 	using JobId = ObjectId<IJob>;
-	using PoolId = StringId<IPool>;
-    using ProjectId = StringId<ProjectConfig>;
-    using StreamId = StringId<IStream>;
-    using TemplateId = StringId<ITemplateRef>;
 
 	public class Fixture
 	{
diff --git a/Engine/Source/Programs/Horde/Horde.Build.Tests/Fleet/FleetServiceTest.cs b/Engine/Source/Programs/Horde/Horde.Build.Tests/Fleet/FleetServiceTest.cs
--- a/Engine/Source/Programs/Horde/Horde.Build.Tests/Fleet/FleetServiceTest.cs
+++ b/Engine/Source/Programs/Horde/Horde.Build.Tests/Fleet/FleetServiceTest.cs
@@ -20,8 +20,6 @@ using StatsdClient;
 
 namespace Horde.Build.Tests.Fleet
 {
-	using PoolId = StringId<IPool>;
-
 	public class FleetManagerSpy : IFleetManager
 	{
 		public int ExpandPoolAsyncCallCount { get; private set; }
@@ -98,7 +96,6 @@ namespace Horde.Build.Tests.Fleet
 	public class FleetServiceTest : TestSetup
 	{
 		readonly FleetManagerSpy _fleetManagerSpy = new();
-		readonly FakeFleetManager _fakeFleetManager = new();
 		readonly IDogStatsd _dogStatsD = new NoOpDogStatsd();
 
 		protected override void Dispose(bool disposing)
diff --git a/Engine/Source/Programs/Horde/Horde.Build.Tests/Fleet/JobQueueStrategyTest.cs b/Engine/Source/Programs/Horde/Horde.Build.Tests/Fleet/JobQueueStrategyTest.cs
--- a/Engine/Source/Programs/Horde/Horde.Build.Tests/Fleet/JobQueueStrategyTest.cs
+++ b/Engine/Source/Programs/Horde/Horde.Build.Tests/Fleet/JobQueueStrategyTest.cs
@@ -17,10 +17,6 @@ using Horde.Build.Agents;
 
 namespace Horde.Build.Tests.Fleet
 {
-	using PoolId = StringId<IPool>;
-	using ProjectId = StringId<ProjectConfig>;
-	using StreamId = StringId<IStream>;
-
 	[TestClass]
 	public class JobQueueStrategyTest : TestSetup
 	{
@@ -155,7 +151,7 @@ namespace Horde.Build.Tests.Fleet
 			options.Arguments.Add($"-Target={nodeNameToExecute}");
 
 			IJob job = await JobCollection.AddAsync(ObjectId<IJob>.GenerateNewId(), streamId,
-				new StringId<ITemplateRef>("bogusTemplateRefId"), ContentHash.Empty, graph, "bogusJobName",
+				new TemplateId("bogusTemplateRefId"), ContentHash.Empty, graph, "bogusJobName",
 				1000, 1000, options);
 
 			return job;
diff --git a/Engine/Source/Programs/Horde/Horde.Build.Tests/Fleet/LeaseUtilizationStrategyTest.cs b/Engine/Source/Programs/Horde/Horde.Build.Tests/Fleet/LeaseUtilizationStrategyTest.cs
--- a/Engine/Source/Programs/Horde/Horde.Build.Tests/Fleet/LeaseUtilizationStrategyTest.cs
+++ b/Engine/Source/Programs/Horde/Horde.Build.Tests/Fleet/LeaseUtilizationStrategyTest.cs
@@ -17,8 +17,6 @@ using Microsoft.VisualStudio.TestTools.UnitTesting;
 
 namespace Horde.Build.Tests.Fleet
 {
-	using StreamId = StringId<IStream>;
-
 	[TestClass]
 	public class LeaseUtilizationStrategyTest : TestSetup
 	{
diff --git a/Engine/Source/Programs/Horde/Horde.Build.Tests/GcServiceTests.cs b/Engine/Source/Programs/Horde/Horde.Build.Tests/GcServiceTests.cs
--- a/Engine/Source/Programs/Horde/Horde.Build.Tests/GcServiceTests.cs
+++ b/Engine/Source/Programs/Horde/Horde.Build.Tests/GcServiceTests.cs
@@ -14,8 +14,6 @@ using Horde.Build.Utilities;
 
 namespace Horde.Build.Tests
 {
-	using BackendId = StringId<IStorageBackend>;
-
 	[TestClass]
 	public sealed class GcServiceTests : TestSetup
 	{
diff --git a/Engine/Source/Programs/Horde/Horde.Build.Tests/IssueServiceTests.cs b/Engine/Source/Programs/Horde/Horde.Build.Tests/IssueServiceTests.cs
--- a/Engine/Source/Programs/Horde/Horde.Build.Tests/IssueServiceTests.cs
+++ b/Engine/Source/Programs/Horde/Horde.Build.Tests/IssueServiceTests.cs
@@ -35,9 +35,6 @@ namespace Horde.Build.Tests
 {
 	using JobId = ObjectId<IJob>;
 	using LogId = ObjectId<ILogFile>;
-	using ProjectId = StringId<ProjectConfig>;
-	using StreamId = StringId<IStream>;
-	using TemplateId = StringId<ITemplateRef>;
 	using UserId = ObjectId<IUser>;
 
 	[TestClass]
@@ -102,13 +99,13 @@ namespace Horde.Build.Tests
 		}
 
 		const string MainStreamName = "//UE4/Main";
-		readonly StreamId _mainStreamId = StreamId.Sanitize(MainStreamName);
+		readonly StreamId _mainStreamId = new StreamId("ue4-main");
 
 		const string ReleaseStreamName = "//UE4/Release";
-		readonly StreamId _releaseStreamId = StreamId.Sanitize(ReleaseStreamName);
+		readonly StreamId _releaseStreamId = new StreamId("ue4-release");
 
 		const string DevStreamName = "//UE4/Dev";
-		readonly StreamId _devStreamId = StreamId.Sanitize(DevStreamName);
+		readonly StreamId _devStreamId = new StreamId("ue4-dev");
 		readonly IGraph _graph;
 		readonly PerforceServiceStub _perforce;
 		readonly UserId _timId;
diff --git a/Engine/Source/Programs/Horde/Horde.Build.Tests/JobCollectionTests.cs b/Engine/Source/Programs/Horde/Horde.Build.Tests/JobCollectionTests.cs
--- a/Engine/Source/Programs/Horde/Horde.Build.Tests/JobCollectionTests.cs
+++ b/Engine/Source/Programs/Horde/Horde.Build.Tests/JobCollectionTests.cs
@@ -25,9 +25,6 @@ namespace Horde.Build.Tests
 	using JobId = ObjectId<IJob>;
 	using LeaseId = ObjectId<ILease>;
 	using LogId = ObjectId<ILogFile>;
-	using PoolId = StringId<IPool>;
-	using StreamId = StringId<IStream>;
-	using TemplateId = StringId<ITemplateRef>;
 
 	[TestClass]
 	public class JobCollectionTests : TestSetup
diff --git a/Engine/Source/Programs/Horde/Horde.Build.Tests/JobServiceTests.cs b/Engine/Source/Programs/Horde/Horde.Build.Tests/JobServiceTests.cs
--- a/Engine/Source/Programs/Horde/Horde.Build.Tests/JobServiceTests.cs
+++ b/Engine/Source/Programs/Horde/Horde.Build.Tests/JobServiceTests.cs
@@ -28,9 +28,7 @@ namespace Horde.Build.Tests
 {
 	using JobId = ObjectId<IJob>;
 	using LogId = ObjectId<ILogFile>;
-	using ProjectId = StringId<ProjectConfig>;
-	using StreamId = StringId<IStream>;
-	using TemplateId = StringId<ITemplateRef>;
+	using ProjectId = Projects.ProjectId;
 
 	[TestClass]
 	public class JobServiceTests : TestSetup
@@ -209,7 +207,7 @@ namespace Horde.Build.Tests
 
 			// ----
 
-			IAgent? agent = await AgentService.CreateAgentAsync("TestAgent", true, new List<StringId<IPool>> { new StringId<IPool>("win") });
+			IAgent? agent = await AgentService.CreateAgentAsync("TestAgent", true, new List<PoolId> { new PoolId("win") });
 			await AgentService.CreateSessionAsync(agent, AgentStatus.Ok, new List<string>(), new Dictionary<string, int>(), null);
 
 			ITemplate template = await TemplateCollection.GetOrAddAsync(new TemplateConfig { Name = "Test template" });
diff --git a/Engine/Source/Programs/Horde/Horde.Build.Tests/JobTaskSourceTests.cs b/Engine/Source/Programs/Horde/Horde.Build.Tests/JobTaskSourceTests.cs
--- a/Engine/Source/Programs/Horde/Horde.Build.Tests/JobTaskSourceTests.cs
+++ b/Engine/Source/Programs/Horde/Horde.Build.Tests/JobTaskSourceTests.cs
@@ -167,7 +167,7 @@ namespace Horde.Build.Tests
 
 			if (shouldCreateAgent)
 			{
-				IAgent? agent = await AgentService.CreateAgentAsync("TestAgent", isAgentEnabled, new List<StringId<IPool>> { pool.Id });
+				IAgent? agent = await AgentService.CreateAgentAsync("TestAgent", isAgentEnabled, new List<PoolId> { pool.Id });
 				await AgentService.CreateSessionAsync(agent, AgentStatus.Ok, new List<string>(), new Dictionary<string, int>(), null);
 			}
 			
diff --git a/Engine/Source/Programs/Horde/Horde.Build.Tests/NotificationServiceTests.cs b/Engine/Source/Programs/Horde/Horde.Build.Tests/NotificationServiceTests.cs
--- a/Engine/Source/Programs/Horde/Horde.Build.Tests/NotificationServiceTests.cs
+++ b/Engine/Source/Programs/Horde/Horde.Build.Tests/NotificationServiceTests.cs
@@ -22,8 +22,7 @@ using Microsoft.Extensions.Caching.Memory;
 namespace Horde.Build.Tests
 {
 	using JobId = ObjectId<IJob>;
-	using StreamId = StringId<IStream>;
-
+	
 	public class FakeNotificationSink : INotificationSink
 	{
 		public List<JobScheduledNotification> JobScheduledNotifications { get; } = new();
diff --git a/Engine/Source/Programs/Horde/Horde.Build.Tests/PoolServiceTests.cs b/Engine/Source/Programs/Horde/Horde.Build.Tests/PoolServiceTests.cs
--- a/Engine/Source/Programs/Horde/Horde.Build.Tests/PoolServiceTests.cs
+++ b/Engine/Source/Programs/Horde/Horde.Build.Tests/PoolServiceTests.cs
@@ -10,8 +10,6 @@ using Microsoft.VisualStudio.TestTools.UnitTesting;
 
 namespace Horde.Build.Tests
 {
-	using PoolId = StringId<IPool>;
-
 	[TestClass]
     public class PoolServiceTests : TestSetup
     {
diff --git a/Engine/Source/Programs/Horde/Horde.Build.Tests/PoolsControllerTest.cs b/Engine/Source/Programs/Horde/Horde.Build.Tests/PoolsControllerTest.cs
--- a/Engine/Source/Programs/Horde/Horde.Build.Tests/PoolsControllerTest.cs
+++ b/Engine/Source/Programs/Horde/Horde.Build.Tests/PoolsControllerTest.cs
@@ -42,7 +42,7 @@ namespace Horde.Build.Tests
 	        ActionResult<CreatePoolResponse> result = await PoolsController.CreatePoolAsync(request);
 	        CreatePoolResponse response = result.Value!;
 		        
-	        IPool pool = (await PoolService.GetPoolAsync(new StringId<IPool>(response.Id)))!;
+	        IPool pool = (await PoolService.GetPoolAsync(new PoolId(response.Id)))!;
 	        Assert.AreEqual(request.Name, pool.Name);
 	        Assert.AreEqual(request.ScaleOutCooldown, (int)pool.ScaleOutCooldown!.Value.TotalSeconds);
 	        Assert.AreEqual(request.ScaleInCooldown, (int)pool.ScaleInCooldown!.Value.TotalSeconds);
diff --git a/Engine/Source/Programs/Horde/Horde.Build.Tests/RpcServiceTest.cs b/Engine/Source/Programs/Horde/Horde.Build.Tests/RpcServiceTest.cs
--- a/Engine/Source/Programs/Horde/Horde.Build.Tests/RpcServiceTest.cs
+++ b/Engine/Source/Programs/Horde/Horde.Build.Tests/RpcServiceTest.cs
@@ -32,8 +32,7 @@ namespace Horde.Build.Tests
 {
 	using LeaseId = ObjectId<ILease>;
 	using LogId = ObjectId<ILogFile>;
-	using PoolId = StringId<IPool>;
-
+	
 	public class AppLifetimeStub : IHostApplicationLifetime
 	{
 		public CancellationToken ApplicationStarted { get; }
diff --git a/Engine/Source/Programs/Horde/Horde.Build.Tests/SchedulerTests.cs b/Engine/Source/Programs/Horde/Horde.Build.Tests/SchedulerTests.cs
--- a/Engine/Source/Programs/Horde/Horde.Build.Tests/SchedulerTests.cs
+++ b/Engine/Source/Programs/Horde/Horde.Build.Tests/SchedulerTests.cs
@@ -23,9 +23,6 @@ namespace Horde.Build.Tests
 {
 	using JobId = ObjectId<IJob>;
 	using LogId = ObjectId<ILogFile>;
-	using ProjectId = StringId<ProjectConfig>;
-	using StreamId = StringId<IStream>;
-	using TemplateId = StringId<ITemplateRef>;
 
 	[TestClass]
     public class SchedulerTests : TestSetup
diff --git a/Engine/Source/Programs/Horde/Horde.Build.Tests/Streams/StreamCollectionTests.cs b/Engine/Source/Programs/Horde/Horde.Build.Tests/Streams/StreamCollectionTests.cs
--- a/Engine/Source/Programs/Horde/Horde.Build.Tests/Streams/StreamCollectionTests.cs
+++ b/Engine/Source/Programs/Horde/Horde.Build.Tests/Streams/StreamCollectionTests.cs
@@ -7,9 +7,6 @@ using Microsoft.VisualStudio.TestTools.UnitTesting;
 
 namespace Horde.Build.Tests.Streams
 {
-	using StreamId = StringId<IStream>;
-	using TemplateId = StringId<ITemplateRef>;
-
 	[TestClass]
 	public class StreamCollectionTests : TestSetup
 	{
diff --git a/Engine/Source/Programs/Horde/Horde.Build.Tests/Stubs/Collections/JobCollectionStub.cs b/Engine/Source/Programs/Horde/Horde.Build.Tests/Stubs/Collections/JobCollectionStub.cs
--- a/Engine/Source/Programs/Horde/Horde.Build.Tests/Stubs/Collections/JobCollectionStub.cs
+++ b/Engine/Source/Programs/Horde/Horde.Build.Tests/Stubs/Collections/JobCollectionStub.cs
@@ -26,7 +26,7 @@ namespace Horde.Build.Tests.Stubs.Collections
 			throw new NotImplementedException();
 		}
 
-		public Task<List<Job>> FindAsync(StringId<Stream>? StreamId, string Name, StringId<TemplateRef>[] Templates, int? MinChange, int? MaxChange, DateTimeOffset? MinCreateTime, DateTimeOffset? MaxCreateTime, int Index, int Count)
+		public Task<List<Job>> FindAsync(StreamId? StreamId, string Name, TemplateId[] Templates, int? MinChange, int? MaxChange, DateTimeOffset? MinCreateTime, DateTimeOffset? MaxCreateTime, int Index, int Count)
 		{
 			throw new NotImplementedException();
 		}
@@ -51,7 +51,7 @@ namespace Horde.Build.Tests.Stubs.Collections
 			throw new NotImplementedException();
 		}
 
-		public Task RemoveStreamAsync(StringId<Stream> StreamId)
+		public Task RemoveStreamAsync(StreamId StreamId)
 		{
 			throw new NotImplementedException();
 		}
diff --git a/Engine/Source/Programs/Horde/Horde.Build.Tests/Stubs/Collections/JobStepRefCollectionStub.cs b/Engine/Source/Programs/Horde/Horde.Build.Tests/Stubs/Collections/JobStepRefCollectionStub.cs
--- a/Engine/Source/Programs/Horde/Horde.Build.Tests/Stubs/Collections/JobStepRefCollectionStub.cs
+++ b/Engine/Source/Programs/Horde/Horde.Build.Tests/Stubs/Collections/JobStepRefCollectionStub.cs
@@ -15,9 +15,6 @@ namespace Horde.Build.Tests.Stubs.Collections
 {
 	using JobId = ObjectId<IJob>;
 	using LogId = ObjectId<ILogFile>;
-	using PoolId = StringId<IPool>;
-	using StreamId = StringId<IStream>;
-	using TemplateRefId = StringId<ITemplateRef>;
 
 	class JobStepRefStub : IJobStepRef
 	{
@@ -25,7 +22,7 @@ namespace Horde.Build.Tests.Stubs.Collections
 		public string JobName { get; set; }
 		public string NodeName { get; set; }
 		public StreamId StreamId { get; set; }
-		public TemplateRefId TemplateId { get; set; }
+		public TemplateId TemplateId { get; set; }
 		public int Change { get; set; }
 		public LogId? LogId { get; set; }
 		public PoolId? PoolId { get; set; }
@@ -41,7 +38,7 @@ namespace Horde.Build.Tests.Stubs.Collections
 		public virtual DateTime StartTimeUtc => throw new NotImplementedException();
 		public virtual DateTime? FinishTimeUtc => throw new NotImplementedException();
 
-		public JobStepRefStub(JobId jobId, SubResourceId batchId, SubResourceId stepId, string jobName, string nodeName, StreamId streamId, TemplateRefId templateId, int change, JobStepOutcome? outcome, bool updateIssues)
+		public JobStepRefStub(JobId jobId, SubResourceId batchId, SubResourceId stepId, string jobName, string nodeName, StreamId streamId, TemplateId templateId, int change, JobStepOutcome? outcome, bool updateIssues)
 		{
 			Id = new JobStepRefId(jobId, batchId, stepId);
 			JobName = jobName;
@@ -63,7 +60,7 @@ namespace Horde.Build.Tests.Stubs.Collections
 			_refs.Add(jobStepRef);
 		}
 
-		public Task<IJobStepRef?> GetNextStepForNodeAsync(StreamId streamId, TemplateRefId templateId, string nodeName, int change, JobStepOutcome? outcome = null, bool? updateIssues = null)
+		public Task<IJobStepRef?> GetNextStepForNodeAsync(StreamId streamId, TemplateId templateId, string nodeName, int change, JobStepOutcome? outcome = null, bool? updateIssues = null)
 		{
 			IJobStepRef? nextRef = null;
 			foreach (IJobStepRef jobStepRef in _refs)
@@ -79,7 +76,7 @@ namespace Horde.Build.Tests.Stubs.Collections
 			return Task.FromResult(nextRef);
 		}
 
-		public Task<IJobStepRef?> GetPrevStepForNodeAsync(StreamId streamId, TemplateRefId templateId, string nodeName, int change, JobStepOutcome? outcome = null, bool? updateIssues = null)
+		public Task<IJobStepRef?> GetPrevStepForNodeAsync(StreamId streamId, TemplateId templateId, string nodeName, int change, JobStepOutcome? outcome = null, bool? updateIssues = null)
 		{
 			IJobStepRef? prevRef = null;
 			foreach (IJobStepRef jobStepRef in _refs)
@@ -95,12 +92,12 @@ namespace Horde.Build.Tests.Stubs.Collections
 			return Task.FromResult(prevRef);
 		}
 
-		Task<List<IJobStepRef>> IJobStepRefCollection.GetStepsForNodeAsync(StreamId streamId, TemplateRefId templateId, string nodeName, int? change, bool includeFailed, int count)
+		Task<List<IJobStepRef>> IJobStepRefCollection.GetStepsForNodeAsync(StreamId streamId, TemplateId templateId, string nodeName, int? change, bool includeFailed, int count)
 		{
 			throw new NotImplementedException();
 		}
 
-		Task<IJobStepRef> IJobStepRefCollection.InsertOrReplaceAsync(JobStepRefId id, string jobName, string nodeName, StreamId streamId, TemplateRefId templateId, int change, LogId? logId, PoolId? poolId, AgentId? agentId, JobStepOutcome? outcome, bool updateIssues, int? lastSuccess, int? lastWarning, float waitTime, float initTime, DateTime jobStartTime, DateTime startTime, DateTime? finishTime)
+		Task<IJobStepRef> IJobStepRefCollection.InsertOrReplaceAsync(JobStepRefId id, string jobName, string nodeName, StreamId streamId, TemplateId templateId, int change, LogId? logId, PoolId? poolId, AgentId? agentId, JobStepOutcome? outcome, bool updateIssues, int? lastSuccess, int? lastWarning, float waitTime, float initTime, DateTime jobStartTime, DateTime startTime, DateTime? finishTime)
 		{
 			throw new NotImplementedException();
 		}
diff --git a/Engine/Source/Programs/Horde/Horde.Build.Tests/Stubs/Collections/StreamCollectionStub.cs b/Engine/Source/Programs/Horde/Horde.Build.Tests/Stubs/Collections/StreamCollectionStub.cs
--- a/Engine/Source/Programs/Horde/Horde.Build.Tests/Stubs/Collections/StreamCollectionStub.cs
+++ b/Engine/Source/Programs/Horde/Horde.Build.Tests/Stubs/Collections/StreamCollectionStub.cs
@@ -13,7 +13,7 @@ namespace Horde.Build.Tests.Stubs.Collections
 {
 	class StreamCollectionStub : IStreamCollection
 	{
-		public Task DeleteAsync(StringId<Stream> StreamId)
+		public Task DeleteAsync(StreamId StreamId)
 		{
 			throw new NotImplementedException();
 		}
@@ -23,17 +23,17 @@ namespace Horde.Build.Tests.Stubs.Collections
 			throw new NotImplementedException();
 		}
 
-		public Task<List<Stream>> FindForProjectAsync(StringId<Project>? ProjectId)
+		public Task<List<Stream>> FindForProjectAsync(ProjectId? ProjectId)
 		{
 			throw new NotImplementedException();
 		}
 
-		public Task<Stream> GetAsync(StringId<Stream> StreamId)
+		public Task<Stream> GetAsync(StreamId StreamId)
 		{
 			throw new NotImplementedException();
 		}
 
-		public Task<StreamPermissions> GetPermissionsAsync(StringId<Stream> StreamId)
+		public Task<StreamPermissions> GetPermissionsAsync(StreamId StreamId)
 		{
 			throw new NotImplementedException();
 		}
@@ -53,12 +53,12 @@ namespace Horde.Build.Tests.Stubs.Collections
 			throw new NotImplementedException();
 		}
 
-		public Task<bool> TryUpdateScheduleJobsAsync(Stream Stream, StringId<TemplateRef> TemplateRefId, List<ObjectId> NewActiveJobs)
+		public Task<bool> TryUpdateScheduleJobsAsync(Stream Stream, TemplateId TemplateRefId, List<ObjectId> NewActiveJobs)
 		{
 			throw new NotImplementedException();
 		}
 
-		public Task<bool> TryUpdateScheduleTriggerTime(Stream Stream, StringId<TemplateRef> TemplateRefId, DateTimeOffset LastTriggerTime)
+		public Task<bool> TryUpdateScheduleTriggerTime(Stream Stream, TemplateId TemplateRefId, DateTimeOffset LastTriggerTime)
 		{
 			throw new NotImplementedException();
 		}
diff --git a/Engine/Source/Programs/Horde/Horde.Build.Tests/Stubs/Services/PerforceServiceStub.cs b/Engine/Source/Programs/Horde/Horde.Build.Tests/Stubs/Services/PerforceServiceStub.cs
--- a/Engine/Source/Programs/Horde/Horde.Build.Tests/Stubs/Services/PerforceServiceStub.cs
+++ b/Engine/Source/Programs/Horde/Horde.Build.Tests/Stubs/Services/PerforceServiceStub.cs
@@ -16,7 +16,6 @@ using Microsoft.VisualStudio.TestTools.UnitTesting;
 
 namespace Horde.Build.Tests.Stubs.Services
 {
-	using StreamId = StringId<IStream>;
 	using UserId = ObjectId<IUser>;
 
 	class PerforceServiceStub : IPerforceService
diff --git a/Engine/Source/Programs/Horde/Horde.Build.Tests/TestDataTests.cs b/Engine/Source/Programs/Horde/Horde.Build.Tests/TestDataTests.cs
--- a/Engine/Source/Programs/Horde/Horde.Build.Tests/TestDataTests.cs
+++ b/Engine/Source/Programs/Horde/Horde.Build.Tests/TestDataTests.cs
@@ -20,12 +20,10 @@ using Horde.Build.Logs;
 using System.Linq;
 using System.Threading;
 using Horde.Build.Server;
+using EpicGames.Horde;
 
 namespace Horde.Build.Tests
 {
-	using ProjectId = StringId<ProjectConfig>;
-	using StreamId = StringId<IStream>;
-	using TemplateId = StringId<ITemplateRef>;
 	using JobId = ObjectId<IJob>;
 
 	/// <summary>
@@ -35,10 +33,10 @@ namespace Horde.Build.Tests
 	public class TestDataTests : TestSetup
 	{
 		const string MainStreamName = "//UE5/Main";
-		readonly StreamId _mainStreamId = StreamId.Sanitize(MainStreamName);
+		readonly StreamId _mainStreamId = new StreamId(StringId.Sanitize(MainStreamName));
 
 		const string ReleaseStreamName = "//UE5/Release";
-		readonly StreamId _releaseStreamId = StreamId.Sanitize(ReleaseStreamName);
+		readonly StreamId _releaseStreamId = new StreamId(StringId.Sanitize(ReleaseStreamName));
 
 		readonly IGraph _graph;
 
diff --git a/Engine/Source/Programs/Horde/Horde.Build.Tests/TestSetup.cs b/Engine/Source/Programs/Horde/Horde.Build.Tests/TestSetup.cs
--- a/Engine/Source/Programs/Horde/Horde.Build.Tests/TestSetup.cs
+++ b/Engine/Source/Programs/Horde/Horde.Build.Tests/TestSetup.cs
@@ -337,7 +337,7 @@ namespace Horde.Build.Tests
 		private static int s_agentIdCounter = 1;
 		public async Task<IAgent> CreateAgentAsync(IPool pool, bool enabled = true)
 		{
-			IAgent agent = await AgentService.CreateAgentAsync("TestAgent" + s_agentIdCounter++, enabled, new List<StringId<IPool>> { pool.Id });
+			IAgent agent = await AgentService.CreateAgentAsync("TestAgent" + s_agentIdCounter++, enabled, new List<PoolId> { pool.Id });
 			agent = await AgentService.CreateSessionAsync(agent, AgentStatus.Ok, new List<string>(), new Dictionary<string, int>(), null);
 			return agent;
 		}
diff --git a/Engine/Source/Programs/Horde/Horde.Build.Tests/ToolTests.cs b/Engine/Source/Programs/Horde/Horde.Build.Tests/ToolTests.cs
--- a/Engine/Source/Programs/Horde/Horde.Build.Tests/ToolTests.cs
+++ b/Engine/Source/Programs/Horde/Horde.Build.Tests/ToolTests.cs
@@ -14,7 +14,6 @@ using Microsoft.VisualStudio.TestTools.UnitTesting;
 
 namespace Horde.Build.Tests
 {
-	using ToolId = StringId<ITool>;
 	using ToolDeploymentId = ObjectId<IToolDeployment>;
 
 	[TestClass]
diff --git a/Engine/Source/Programs/Horde/Horde.Build/Agents/AgentCollection.cs b/Engine/Source/Programs/Horde/Horde.Build/Agents/AgentCollection.cs
--- a/Engine/Source/Programs/Horde/Horde.Build/Agents/AgentCollection.cs
+++ b/Engine/Source/Programs/Horde/Horde.Build/Agents/AgentCollection.cs
@@ -20,7 +20,6 @@ using MongoDB.Driver;
 
 namespace Horde.Build.Agents
 {
-	using PoolId = StringId<IPool>;
 	using SessionId = ObjectId<ISession>;
 
 	/// <summary>
diff --git a/Engine/Source/Programs/Horde/Horde.Build/Agents/AgentService.cs b/Engine/Source/Programs/Horde/Horde.Build/Agents/AgentService.cs
--- a/Engine/Source/Programs/Horde/Horde.Build/Agents/AgentService.cs
+++ b/Engine/Source/Programs/Horde/Horde.Build/Agents/AgentService.cs
@@ -31,7 +31,6 @@ using StatsdClient;
 namespace Horde.Build.Agents
 {
 	using LeaseId = ObjectId<ILease>;
-	using PoolId = StringId<IPool>;
 	using SessionId = ObjectId<ISession>;
 
 	/// <summary>
@@ -264,24 +263,24 @@ namespace Horde.Build.Agents
 			return newDynamicPools;
 		}
 		
-		private static List<StringId<IPool>> GetRequestedPoolsFromProperties(IReadOnlyList<string> properties)
+		private static List<PoolId> GetRequestedPoolsFromProperties(IReadOnlyList<string> properties)
 		{
-			List<StringId<IPool>> poolIds = new();
+			List<PoolId> poolIds = new();
 			foreach (string property in properties)
 			{
 				const string Key = KnownPropertyNames.RequestedPools + "=";
 				if (property.StartsWith(Key, StringComparison.InvariantCulture))
 				{
-					poolIds.AddRange(property[Key.Length..].Split(",").Select(x => new StringId<IPool>(x)));
+					poolIds.AddRange(property[Key.Length..].Split(",").Select(x => new PoolId(x)));
 				}
 			}
 
 			return poolIds;
 		}
 		
-		private static List<StringId<IPool>> CombineCurrentAndRequestedPools(IReadOnlyList<StringId<IPool>> pools, IReadOnlyList<string> properties)
+		private static List<PoolId> CombineCurrentAndRequestedPools(IReadOnlyList<PoolId> pools, IReadOnlyList<string> properties)
 		{
-			HashSet<StringId<IPool>> uniquePools = new(pools);
+			HashSet<PoolId> uniquePools = new(pools);
 			uniquePools.UnionWith(GetRequestedPoolsFromProperties(properties));
 			return new List<PoolId>(uniquePools);
 		}
diff --git a/Engine/Source/Programs/Horde/Horde.Build/Agents/AgentsController.cs b/Engine/Source/Programs/Horde/Horde.Build/Agents/AgentsController.cs
--- a/Engine/Source/Programs/Horde/Horde.Build/Agents/AgentsController.cs
+++ b/Engine/Source/Programs/Horde/Horde.Build/Agents/AgentsController.cs
@@ -20,7 +20,6 @@ using Microsoft.Extensions.Options;
 namespace Horde.Build.Agents
 {
 	using LeaseId = ObjectId<ILease>;
-	using PoolId = StringId<IPool>;
 	using SessionId = ObjectId<ISession>;
 
 	/// <summary>
diff --git a/Engine/Source/Programs/Horde/Horde.Build/Agents/Fleet/FleetService.cs b/Engine/Source/Programs/Horde/Horde.Build/Agents/Fleet/FleetService.cs
--- a/Engine/Source/Programs/Horde/Horde.Build/Agents/Fleet/FleetService.cs
+++ b/Engine/Source/Programs/Horde/Horde.Build/Agents/Fleet/FleetService.cs
@@ -26,8 +26,6 @@ using StatsdClient;
 
 namespace Horde.Build.Agents.Fleet
 {
-	using PoolId = StringId<IPool>;
-	
 	/// <summary>
 	/// Parameters required for calculating pool size
 	/// </summary>
diff --git a/Engine/Source/Programs/Horde/Horde.Build/Agents/Fleet/JobQueueStrategy.cs b/Engine/Source/Programs/Horde/Horde.Build/Agents/Fleet/JobQueueStrategy.cs
--- a/Engine/Source/Programs/Horde/Horde.Build/Agents/Fleet/JobQueueStrategy.cs
+++ b/Engine/Source/Programs/Horde/Horde.Build/Agents/Fleet/JobQueueStrategy.cs
@@ -20,9 +20,6 @@ using OpenTracing.Util;
 
 namespace Horde.Build.Agents.Fleet
 {
-	using PoolId = StringId<IPool>;
-	using StreamId = StringId<IStream>;
-
 	/// <summary>
 	/// Job queue sizing settings for a pool
 	/// </summary>
diff --git a/Engine/Source/Programs/Horde/Horde.Build/Agents/Fleet/LeaseUtilizationStrategy.cs b/Engine/Source/Programs/Horde/Horde.Build/Agents/Fleet/LeaseUtilizationStrategy.cs
--- a/Engine/Source/Programs/Horde/Horde.Build/Agents/Fleet/LeaseUtilizationStrategy.cs
+++ b/Engine/Source/Programs/Horde/Horde.Build/Agents/Fleet/LeaseUtilizationStrategy.cs
@@ -17,8 +17,6 @@ using OpenTracing.Util;
 
 namespace Horde.Build.Agents.Fleet
 {
-	using PoolId = StringId<IPool>;
-
 	/// <summary>
 	/// Lease utilization sizing settings for a pool
 	/// </summary>
diff --git a/Engine/Source/Programs/Horde/Horde.Build/Agents/Fleet/Providers/AwsAsgFleetManager.cs b/Engine/Source/Programs/Horde/Horde.Build/Agents/Fleet/Providers/AwsAsgFleetManager.cs
--- a/Engine/Source/Programs/Horde/Horde.Build/Agents/Fleet/Providers/AwsAsgFleetManager.cs
+++ b/Engine/Source/Programs/Horde/Horde.Build/Agents/Fleet/Providers/AwsAsgFleetManager.cs
@@ -88,7 +88,7 @@ namespace Horde.Build.Agents.Fleet.Providers
 			return new ScaleResult(FleetManagerOutcome.Success, 0, count);
 		}
 		
-		private async Task UpdateAsg(StringId<IPool> poolId, int desiredCapacity, IScope scope, CancellationToken cancellationToken)
+		private async Task UpdateAsg(PoolId poolId, int desiredCapacity, IScope scope, CancellationToken cancellationToken)
 		{
 			UpdateAutoScalingGroupRequest request = new () { AutoScalingGroupName = Settings.Name, DesiredCapacity = desiredCapacity };
 			UpdateAutoScalingGroupResponse response = await _awsAutoScaling.UpdateAutoScalingGroupAsync(request, cancellationToken);
diff --git a/Engine/Source/Programs/Horde/Horde.Build/Agents/IAgent.cs b/Engine/Source/Programs/Horde/Horde.Build/Agents/IAgent.cs
--- a/Engine/Source/Programs/Horde/Horde.Build/Agents/IAgent.cs
+++ b/Engine/Source/Programs/Horde/Horde.Build/Agents/IAgent.cs
@@ -29,10 +29,8 @@ namespace Horde.Build.Agents
 {
 	using LeaseId = ObjectId<ILease>;
 	using LogId = ObjectId<ILogFile>;
-	using PoolId = StringId<IPool>;
 	using SessionId = ObjectId<ISession>;
-	using StreamId = StringId<IStream>;
-
+	
 	/// <summary>
 	/// Information about a workspace synced to an agent
 	/// </summary>
diff --git a/Engine/Source/Programs/Horde/Horde.Build/Agents/IAgentCollection.cs b/Engine/Source/Programs/Horde/Horde.Build/Agents/IAgentCollection.cs
--- a/Engine/Source/Programs/Horde/Horde.Build/Agents/IAgentCollection.cs
+++ b/Engine/Source/Programs/Horde/Horde.Build/Agents/IAgentCollection.cs
@@ -12,7 +12,6 @@ using HordeCommon;
 
 namespace Horde.Build.Agents
 {
-	using PoolId = StringId<IPool>;
 	using SessionId = ObjectId<ISession>;
 
 	/// <summary>
diff --git a/Engine/Source/Programs/Horde/Horde.Build/Agents/Leases/ILease.cs b/Engine/Source/Programs/Horde/Horde.Build/Agents/Leases/ILease.cs
--- a/Engine/Source/Programs/Horde/Horde.Build/Agents/Leases/ILease.cs
+++ b/Engine/Source/Programs/Horde/Horde.Build/Agents/Leases/ILease.cs
@@ -14,10 +14,8 @@ namespace Horde.Build.Agents.Leases
 {
 	using LeaseId = ObjectId<ILease>;
 	using LogId = ObjectId<ILogFile>;
-	using PoolId = StringId<IPool>;
 	using SessionId = ObjectId<ISession>;
-	using StreamId = StringId<IStream>;
-
+	
 	/// <summary>
 	/// Document describing a lease. This exists to permanently record a lease; the agent object tracks internal state of any active leases through AgentLease objects.
 	/// </summary>
diff --git a/Engine/Source/Programs/Horde/Horde.Build/Agents/Leases/ILeaseCollection.cs b/Engine/Source/Programs/Horde/Horde.Build/Agents/Leases/ILeaseCollection.cs
--- a/Engine/Source/Programs/Horde/Horde.Build/Agents/Leases/ILeaseCollection.cs
+++ b/Engine/Source/Programs/Horde/Horde.Build/Agents/Leases/ILeaseCollection.cs
@@ -14,10 +14,8 @@ namespace Horde.Build.Agents.Leases
 {
 	using LeaseId = ObjectId<ILease>;
 	using LogId = ObjectId<ILogFile>;
-	using PoolId = StringId<IPool>;
 	using SessionId = ObjectId<ISession>;
-	using StreamId = StringId<IStream>;
-
+	
 	/// <summary>
 	/// Interface for a collection of lease documents
 	/// </summary>
diff --git a/Engine/Source/Programs/Horde/Horde.Build/Agents/Leases/LeaseCollection.cs b/Engine/Source/Programs/Horde/Horde.Build/Agents/Leases/LeaseCollection.cs
--- a/Engine/Source/Programs/Horde/Horde.Build/Agents/Leases/LeaseCollection.cs
+++ b/Engine/Source/Programs/Horde/Horde.Build/Agents/Leases/LeaseCollection.cs
@@ -18,10 +18,8 @@ namespace Horde.Build.Agents.Leases
 {
 	using LeaseId = ObjectId<ILease>;
 	using LogId = ObjectId<ILogFile>;
-	using PoolId = StringId<IPool>;
 	using SessionId = ObjectId<ISession>;
-	using StreamId = StringId<IStream>;
-
+	
 	/// <summary>
 	/// Collection of lease documents
 	/// </summary>
diff --git a/Engine/Source/Programs/Horde/Horde.Build/Agents/Pools/IPool.cs b/Engine/Source/Programs/Horde/Horde.Build/Agents/Pools/IPool.cs
--- a/Engine/Source/Programs/Horde/Horde.Build/Agents/Pools/IPool.cs
+++ b/Engine/Source/Programs/Horde/Horde.Build/Agents/Pools/IPool.cs
@@ -11,8 +11,6 @@ using Horde.Build.Utilities;
 
 namespace Horde.Build.Agents.Pools
 {
-	using PoolId = StringId<IPool>;
-
 	/// <summary>
 	/// Metadata for configuring and picking a pool sizing strategy
 	/// </summary>
diff --git a/Engine/Source/Programs/Horde/Horde.Build/Agents/Pools/IPoolCollection.cs b/Engine/Source/Programs/Horde/Horde.Build/Agents/Pools/IPoolCollection.cs
--- a/Engine/Source/Programs/Horde/Horde.Build/Agents/Pools/IPoolCollection.cs
+++ b/Engine/Source/Programs/Horde/Horde.Build/Agents/Pools/IPoolCollection.cs
@@ -10,8 +10,6 @@ using Horde.Build.Utilities;
 
 namespace Horde.Build.Agents.Pools
 {
-	using PoolId = StringId<IPool>;
-
 	/// <summary>
 	/// Exception when conflicting definitions for a pool are encountered
 	/// </summary>
diff --git a/Engine/Source/Programs/Horde/Horde.Build/Agents/Pools/PoolCollection.cs b/Engine/Source/Programs/Horde/Horde.Build/Agents/Pools/PoolCollection.cs
--- a/Engine/Source/Programs/Horde/Horde.Build/Agents/Pools/PoolCollection.cs
+++ b/Engine/Source/Programs/Horde/Horde.Build/Agents/Pools/PoolCollection.cs
@@ -15,8 +15,6 @@ using MongoDB.Driver;
 
 namespace Horde.Build.Agents.Pools
 {
-	using PoolId = StringId<IPool>;
-
 	/// <summary>
 	/// Collection of pool documents
 	/// </summary>
diff --git a/Engine/Source/Programs/Horde/Horde.Build/Agents/Pools/PoolConfig.cs b/Engine/Source/Programs/Horde/Horde.Build/Agents/Pools/PoolConfig.cs
--- a/Engine/Source/Programs/Horde/Horde.Build/Agents/Pools/PoolConfig.cs
+++ b/Engine/Source/Programs/Horde/Horde.Build/Agents/Pools/PoolConfig.cs
@@ -8,8 +8,6 @@ using Horde.Build.Utilities;
 
 namespace Horde.Build.Agents.Pools
 {
-	using PoolId = StringId<IPool>;
-
 	/// <summary>
 	/// Configuration for a pool of machines
 	/// </summary>
diff --git a/Engine/Source/Programs/Horde/Horde.Build/Agents/Pools/PoolId.cs b/Engine/Source/Programs/Horde/Horde.Build/Agents/Pools/PoolId.cs
new file mode 100644
--- /dev/null
+++ b/Engine/Source/Programs/Horde/Horde.Build/Agents/Pools/PoolId.cs
@@ -0,0 +1,46 @@
+// Copyright Epic Games, Inc. All Rights Reserved.
+
+using System.ComponentModel;
+using EpicGames.Core;
+using EpicGames.Horde;
+using EpicGames.Serialization;
+using Horde.Build.Utilities;
+
+namespace Horde.Build.Agents.Pools
+{
+	/// <summary>
+	/// Identifier for a pool
+	/// </summary>
+	/// <param name="Id">Id to construct from</param>
+	[JsonSchemaString]
+	[TypeConverter(typeof(StringIdTypeConverter<PoolId, PoolIdConverter>))]
+	[StringIdConverter(typeof(PoolIdConverter))]
+	[CbConverter(typeof(StringIdCbConverter<PoolId, PoolIdConverter>))]
+	public record struct PoolId(StringId Id)
+	{
+		/// <summary>
+		/// Constructor
+		/// </summary>
+		public PoolId(string id) : this(new StringId(id))
+		{
+		}
+
+		/// <inheritdoc cref="StringId.IsEmpty"/>
+		public bool IsEmpty => Id.IsEmpty;
+
+		/// <inheritdoc/>
+		public override string ToString() => Id.ToString();
+	}
+
+	/// <summary>
+	/// Converter to and from <see cref="StringId"/> instances.
+	/// </summary>
+	class PoolIdConverter : StringIdConverter<PoolId>
+	{
+		/// <inheritdoc/>
+		public override PoolId FromStringId(StringId id) => new PoolId(id);
+
+		/// <inheritdoc/>
+		public override StringId ToStringId(PoolId value) => value.Id;
+	}
+}
diff --git a/Engine/Source/Programs/Horde/Horde.Build/Agents/Pools/PoolService.cs b/Engine/Source/Programs/Horde/Horde.Build/Agents/Pools/PoolService.cs
--- a/Engine/Source/Programs/Horde/Horde.Build/Agents/Pools/PoolService.cs
+++ b/Engine/Source/Programs/Horde/Horde.Build/Agents/Pools/PoolService.cs
@@ -5,6 +5,7 @@ using System.Collections.Generic;
 using System.Linq;
 using System.Threading;
 using System.Threading.Tasks;
+using EpicGames.Horde;
 using EpicGames.Horde.Common;
 using Horde.Build.Agents.Fleet;
 using Horde.Build.Server;
@@ -13,8 +14,6 @@ using HordeCommon;
 
 namespace Horde.Build.Agents.Pools
 {
-	using PoolId = StringId<IPool>;
-
 	/// <summary>
 	/// Wraps functionality for manipulating pools
 	/// </summary>
@@ -90,7 +89,7 @@ namespace Horde.Build.Agents.Pools
 			Dictionary<string, string>? properties = null)
 		{
 			return _pools.AddAsync(
-				PoolId.Sanitize(name),
+				new PoolId(StringId.Sanitize(name)),
 				name,
 				condition,
 				enableAutoscaling,
diff --git a/Engine/Source/Programs/Horde/Horde.Build/Agents/Pools/PoolUpdateService.cs b/Engine/Source/Programs/Horde/Horde.Build/Agents/Pools/PoolUpdateService.cs
--- a/Engine/Source/Programs/Horde/Horde.Build/Agents/Pools/PoolUpdateService.cs
+++ b/Engine/Source/Programs/Horde/Horde.Build/Agents/Pools/PoolUpdateService.cs
@@ -16,8 +16,6 @@ using Microsoft.Extensions.Options;
 
 namespace Horde.Build.Agents.Pools
 {
-	using PoolId = StringId<IPool>;
-
 	/// <summary>
 	/// Periodically updates pool documents to contain the correct workspaces
 	/// </summary>
diff --git a/Engine/Source/Programs/Horde/Horde.Build/Agents/Pools/PoolsController.cs b/Engine/Source/Programs/Horde/Horde.Build/Agents/Pools/PoolsController.cs
--- a/Engine/Source/Programs/Horde/Horde.Build/Agents/Pools/PoolsController.cs
+++ b/Engine/Source/Programs/Horde/Horde.Build/Agents/Pools/PoolsController.cs
@@ -14,8 +14,6 @@ using Microsoft.Extensions.Options;
 
 namespace Horde.Build.Agents.Pools
 {
-	using PoolId = StringId<IPool>;
-
 	/// <summary>
 	/// Controller for the /api/v1/pools endpoint
 	/// </summary>
diff --git a/Engine/Source/Programs/Horde/Horde.Build/Agents/Software/AgentSoftwareChannelName.cs b/Engine/Source/Programs/Horde/Horde.Build/Agents/Software/AgentSoftwareChannelName.cs
new file mode 100644
--- /dev/null
+++ b/Engine/Source/Programs/Horde/Horde.Build/Agents/Software/AgentSoftwareChannelName.cs
@@ -0,0 +1,46 @@
+// Copyright Epic Games, Inc. All Rights Reserved.
+
+using System.ComponentModel;
+using EpicGames.Core;
+using EpicGames.Horde;
+using EpicGames.Serialization;
+using Horde.Build.Utilities;
+
+namespace Horde.Build.Agents.Software
+{
+	/// <summary>
+	/// Identifier for a pool
+	/// </summary>
+	/// <param name="Id">Id to construct from</param>
+	[JsonSchemaString]
+	[TypeConverter(typeof(StringIdTypeConverter<AgentSoftwareChannelName, AgentSoftwareChannelNameConverter>))]
+	[StringIdConverter(typeof(AgentSoftwareChannelNameConverter))]
+	[CbConverter(typeof(StringIdCbConverter<AgentSoftwareChannelName, AgentSoftwareChannelNameConverter>))]
+	public record struct AgentSoftwareChannelName(StringId Id)
+	{
+		/// <summary>
+		/// Constructor
+		/// </summary>
+		public AgentSoftwareChannelName(string id) : this(new StringId(id))
+		{
+		}
+
+		/// <inheritdoc cref="StringId.IsEmpty"/>
+		public bool IsEmpty => Id.IsEmpty;
+
+		/// <inheritdoc/>
+		public override string ToString() => Id.ToString();
+	}
+
+	/// <summary>
+	/// Converter to and from <see cref="StringId"/> instances.
+	/// </summary>
+	class AgentSoftwareChannelNameConverter : StringIdConverter<AgentSoftwareChannelName>
+	{
+		/// <inheritdoc/>
+		public override AgentSoftwareChannelName FromStringId(StringId id) => new AgentSoftwareChannelName(id);
+
+		/// <inheritdoc/>
+		public override StringId ToStringId(AgentSoftwareChannelName value) => value.Id;
+	}
+}
diff --git a/Engine/Source/Programs/Horde/Horde.Build/Agents/Software/AgentSoftwareController.cs b/Engine/Source/Programs/Horde/Horde.Build/Agents/Software/AgentSoftwareController.cs
--- a/Engine/Source/Programs/Horde/Horde.Build/Agents/Software/AgentSoftwareController.cs
+++ b/Engine/Source/Programs/Horde/Horde.Build/Agents/Software/AgentSoftwareController.cs
@@ -14,8 +14,6 @@ using Microsoft.Net.Http.Headers;
 
 namespace Horde.Build.Agents.Software
 {
-	using AgentSoftwareChannelName = StringId<AgentSoftwareChannels>;
-
 	/// <summary>
 	/// Controller for the /api/v1/software endpoint
 	/// </summary>
diff --git a/Engine/Source/Programs/Horde/Horde.Build/Agents/Software/AgentSoftwareService.cs b/Engine/Source/Programs/Horde/Horde.Build/Agents/Software/AgentSoftwareService.cs
--- a/Engine/Source/Programs/Horde/Horde.Build/Agents/Software/AgentSoftwareService.cs
+++ b/Engine/Source/Programs/Horde/Horde.Build/Agents/Software/AgentSoftwareService.cs
@@ -17,8 +17,6 @@ using MongoDB.Driver;
 
 namespace Horde.Build.Agents.Software
 {
-	using AgentSoftwareChannelName = StringId<AgentSoftwareChannels>;
-
 	/// <summary>
 	/// Information about a channel
 	/// </summary>
diff --git a/Engine/Source/Programs/Horde/Horde.Build/Agents/Software/IAgentSoftwareChannel.cs b/Engine/Source/Programs/Horde/Horde.Build/Agents/Software/IAgentSoftwareChannel.cs
--- a/Engine/Source/Programs/Horde/Horde.Build/Agents/Software/IAgentSoftwareChannel.cs
+++ b/Engine/Source/Programs/Horde/Horde.Build/Agents/Software/IAgentSoftwareChannel.cs
@@ -5,8 +5,6 @@ using Horde.Build.Utilities;
 
 namespace Horde.Build.Agents.Software
 {
-	using AgentSoftwareChannelName = StringId<AgentSoftwareChannels>;
-
 	/// <summary>
 	/// A software channel
 	/// </summary>
diff --git a/Engine/Source/Programs/Horde/Horde.Build/Agents/Telemetry/Telemetry.cs b/Engine/Source/Programs/Horde/Horde.Build/Agents/Telemetry/Telemetry.cs
--- a/Engine/Source/Programs/Horde/Horde.Build/Agents/Telemetry/Telemetry.cs
+++ b/Engine/Source/Programs/Horde/Horde.Build/Agents/Telemetry/Telemetry.cs
@@ -8,9 +8,6 @@ using Horde.Build.Utilities;
 
 namespace Horde.Build.Agents.Telemetry
 {
-	using PoolId = StringId<IPool>;
-	using StreamId = StringId<IStream>;
-
 	/// <summary>
 	/// Information about the utilization of a pool
 	/// </summary>
diff --git a/Engine/Source/Programs/Horde/Horde.Build/Agents/Telemetry/TelemetryCollection.cs b/Engine/Source/Programs/Horde/Horde.Build/Agents/Telemetry/TelemetryCollection.cs
--- a/Engine/Source/Programs/Horde/Horde.Build/Agents/Telemetry/TelemetryCollection.cs
+++ b/Engine/Source/Programs/Horde/Horde.Build/Agents/Telemetry/TelemetryCollection.cs
@@ -14,9 +14,6 @@ using MongoDB.Driver;
 
 namespace Horde.Build.Agents.Telemetry
 {
-	using PoolId = StringId<IPool>;
-	using StreamId = StringId<IStream>;
-
 	/// <summary>
 	/// Collection of utilization data
 	/// </summary>
diff --git a/Engine/Source/Programs/Horde/Horde.Build/Agents/Telemetry/TelemetryService.cs b/Engine/Source/Programs/Horde/Horde.Build/Agents/Telemetry/TelemetryService.cs
--- a/Engine/Source/Programs/Horde/Horde.Build/Agents/Telemetry/TelemetryService.cs
+++ b/Engine/Source/Programs/Horde/Horde.Build/Agents/Telemetry/TelemetryService.cs
@@ -15,8 +15,6 @@ using Microsoft.Extensions.Logging;
 
 namespace Horde.Build.Agents.Telemetry
 {
-	using PoolId = StringId<IPool>;
-
 	/// <summary>
 	/// Service which updates telemetry periodically
 	/// </summary>
diff --git a/Engine/Source/Programs/Horde/Horde.Build/Commands/Bundles/PerforceCommand.cs b/Engine/Source/Programs/Horde/Horde.Build/Commands/Bundles/PerforceCommand.cs
--- a/Engine/Source/Programs/Horde/Horde.Build/Commands/Bundles/PerforceCommand.cs
+++ b/Engine/Source/Programs/Horde/Horde.Build/Commands/Bundles/PerforceCommand.cs
@@ -16,8 +16,6 @@ using Microsoft.Extensions.Options;
 
 namespace Horde.Build.Commands.Bundles
 {
-	using StreamId = StringId<IStream>;
-
 	[Command("bundle", "perforce", "Replicates commits for a particular change of changes from Perforce")]
 	class PerforceCommand : Command
 	{
diff --git a/Engine/Source/Programs/Horde/Horde.Build/Configuration/ConfigSource.cs b/Engine/Source/Programs/Horde/Horde.Build/Configuration/ConfigSource.cs
--- a/Engine/Source/Programs/Horde/Horde.Build/Configuration/ConfigSource.cs
+++ b/Engine/Source/Programs/Horde/Horde.Build/Configuration/ConfigSource.cs
@@ -17,8 +17,6 @@ using Microsoft.Extensions.Options;
 
 namespace Horde.Build.Configuration
 {
-	using PerforceConnectionId = StringId<PerforceConnectionSettings>;
-
 	/// <summary>
 	/// Source for reading config files
 	/// </summary>
diff --git a/Engine/Source/Programs/Horde/Horde.Build/Devices/DeviceCollection.cs b/Engine/Source/Programs/Horde/Horde.Build/Devices/DeviceCollection.cs
--- a/Engine/Source/Programs/Horde/Horde.Build/Devices/DeviceCollection.cs
+++ b/Engine/Source/Programs/Horde/Horde.Build/Devices/DeviceCollection.cs
@@ -17,13 +17,8 @@ using MongoDB.Driver;
 
 namespace Horde.Build.Devices
 {
-	using DeviceId = StringId<IDevice>;
-	using DevicePlatformId = StringId<IDevicePlatform>;
-	using DevicePoolId = StringId<IDevicePool>;
-	using ProjectId = StringId<ProjectConfig>;
 	using UserId = ObjectId<IUser>;
-	using StreamId = StringId<IStream>;
-
+	
 	/// <summary>
 	/// Collection of device documents
 	/// </summary>
diff --git a/Engine/Source/Programs/Horde/Horde.Build/Devices/DeviceId.cs b/Engine/Source/Programs/Horde/Horde.Build/Devices/DeviceId.cs
new file mode 100644
--- /dev/null
+++ b/Engine/Source/Programs/Horde/Horde.Build/Devices/DeviceId.cs
@@ -0,0 +1,49 @@
+// Copyright Epic Games, Inc. All Rights Reserved.
+
+using System.ComponentModel;
+using EpicGames.Core;
+using EpicGames.Horde;
+using EpicGames.Serialization;
+using Horde.Build.Utilities;
+
+namespace Horde.Build.Devices
+{
+	/// <summary>
+	/// Identifier for a pool
+	/// </summary>
+	/// <param name="Id">Id to construct from</param>
+	[JsonSchemaString]
+	[TypeConverter(typeof(StringIdTypeConverter<DeviceId, DeviceIdConverter>))]
+	[StringIdConverter(typeof(DeviceIdConverter))]
+	[CbConverter(typeof(StringIdCbConverter<DeviceId, DeviceIdConverter>))]
+	public record struct DeviceId(StringId Id)
+	{
+		/// <summary>
+		/// Constructor
+		/// </summary>
+		public DeviceId(string id) : this(new StringId(id))
+		{
+		}
+
+		/// <inheritdoc cref="StringId.IsEmpty"/>
+		public bool IsEmpty => Id.IsEmpty;
+
+		/// <inheritdoc cref="StringId.Sanitize(System.String)"/>
+		public static DeviceId Sanitize(string text) => new DeviceId(StringId.Sanitize(text));
+
+		/// <inheritdoc/>
+		public override string ToString() => Id.ToString();
+	}
+
+	/// <summary>
+	/// Converter to and from <see cref="StringId"/> instances.
+	/// </summary>
+	class DeviceIdConverter : StringIdConverter<DeviceId>
+	{
+		/// <inheritdoc/>
+		public override DeviceId FromStringId(StringId id) => new DeviceId(id);
+
+		/// <inheritdoc/>
+		public override StringId ToStringId(DeviceId value) => value.Id;
+	}
+}
diff --git a/Engine/Source/Programs/Horde/Horde.Build/Devices/DevicePlatformId.cs b/Engine/Source/Programs/Horde/Horde.Build/Devices/DevicePlatformId.cs
new file mode 100644
--- /dev/null
+++ b/Engine/Source/Programs/Horde/Horde.Build/Devices/DevicePlatformId.cs
@@ -0,0 +1,49 @@
+// Copyright Epic Games, Inc. All Rights Reserved.
+
+using System.ComponentModel;
+using EpicGames.Core;
+using EpicGames.Horde;
+using EpicGames.Serialization;
+using Horde.Build.Utilities;
+
+namespace Horde.Build.Devices
+{
+	/// <summary>
+	/// Identifier for a device platform
+	/// </summary>
+	/// <param name="Id">Id to construct from</param>
+	[JsonSchemaString]
+	[TypeConverter(typeof(StringIdTypeConverter<DevicePlatformId, DevicePlatformIdConverter>))]
+	[StringIdConverter(typeof(DevicePlatformIdConverter))]
+	[CbConverter(typeof(StringIdCbConverter<DevicePlatformId, DevicePlatformIdConverter>))]
+	public record struct DevicePlatformId(StringId Id)
+	{
+		/// <summary>
+		/// Constructor
+		/// </summary>
+		public DevicePlatformId(string id) : this(new StringId(id))
+		{
+		}
+
+		/// <inheritdoc cref="StringId.IsEmpty"/>
+		public bool IsEmpty => Id.IsEmpty;
+
+		/// <inheritdoc cref="StringId.Sanitize(System.String)"/>
+		public static DevicePlatformId Sanitize(string text) => new DevicePlatformId(StringId.Sanitize(text));
+
+		/// <inheritdoc/>
+		public override string ToString() => Id.ToString();
+	}
+
+	/// <summary>
+	/// Converter to and from <see cref="StringId"/> instances.
+	/// </summary>
+	class DevicePlatformIdConverter : StringIdConverter<DevicePlatformId>
+	{
+		/// <inheritdoc/>
+		public override DevicePlatformId FromStringId(StringId id) => new DevicePlatformId(id);
+
+		/// <inheritdoc/>
+		public override StringId ToStringId(DevicePlatformId value) => value.Id;
+	}
+}
diff --git a/Engine/Source/Programs/Horde/Horde.Build/Devices/DevicePoolId.cs b/Engine/Source/Programs/Horde/Horde.Build/Devices/DevicePoolId.cs
new file mode 100644
--- /dev/null
+++ b/Engine/Source/Programs/Horde/Horde.Build/Devices/DevicePoolId.cs
@@ -0,0 +1,49 @@
+// Copyright Epic Games, Inc. All Rights Reserved.
+
+using System.ComponentModel;
+using EpicGames.Core;
+using EpicGames.Horde;
+using EpicGames.Serialization;
+using Horde.Build.Utilities;
+
+namespace Horde.Build.Devices
+{
+	/// <summary>
+	/// Identifier for a device pool
+	/// </summary>
+	/// <param name="Id">Id to construct from</param>
+	[JsonSchemaString]
+	[TypeConverter(typeof(StringIdTypeConverter<DevicePoolId, DevicePoolIdConverter>))]
+	[StringIdConverter(typeof(DevicePoolIdConverter))]
+	[CbConverter(typeof(StringIdCbConverter<DevicePoolId, DevicePoolIdConverter>))]
+	public record struct DevicePoolId(StringId Id)
+	{
+		/// <summary>
+		/// Constructor
+		/// </summary>
+		public DevicePoolId(string id) : this(new StringId(id))
+		{
+		}
+
+		/// <inheritdoc cref="StringId.IsEmpty"/>
+		public bool IsEmpty => Id.IsEmpty;
+
+		/// <inheritdoc cref="StringId.Sanitize(System.String)"/>
+		public static DevicePoolId Sanitize(string text) => new DevicePoolId(StringId.Sanitize(text));
+
+		/// <inheritdoc/>
+		public override string ToString() => Id.ToString();
+	}
+
+	/// <summary>
+	/// Converter to and from <see cref="StringId"/> instances.
+	/// </summary>
+	class DevicePoolIdConverter : StringIdConverter<DevicePoolId>
+	{
+		/// <inheritdoc/>
+		public override DevicePoolId FromStringId(StringId id) => new DevicePoolId(id);
+
+		/// <inheritdoc/>
+		public override StringId ToStringId(DevicePoolId value) => value.Id;
+	}
+}
diff --git a/Engine/Source/Programs/Horde/Horde.Build/Devices/DeviceService.cs b/Engine/Source/Programs/Horde/Horde.Build/Devices/DeviceService.cs
--- a/Engine/Source/Programs/Horde/Horde.Build/Devices/DeviceService.cs
+++ b/Engine/Source/Programs/Horde/Horde.Build/Devices/DeviceService.cs
@@ -25,11 +25,7 @@ using OpenTracing.Util;
 
 namespace Horde.Build.Devices
 {
-	using DeviceId = StringId<IDevice>;
-	using DevicePlatformId = StringId<IDevicePlatform>;
-	using DevicePoolId = StringId<IDevicePool>;
 	using JobId = ObjectId<IJob>;
-	using ProjectId = StringId<ProjectConfig>;
 	using UserId = ObjectId<IUser>;
 
 	/// <summary>
diff --git a/Engine/Source/Programs/Horde/Horde.Build/Devices/DevicesController.cs b/Engine/Source/Programs/Horde/Horde.Build/Devices/DevicesController.cs
--- a/Engine/Source/Programs/Horde/Horde.Build/Devices/DevicesController.cs
+++ b/Engine/Source/Programs/Horde/Horde.Build/Devices/DevicesController.cs
@@ -19,11 +19,6 @@ using MongoDB.Driver;
 
 namespace Horde.Build.Devices
 {
-	using DeviceId = StringId<IDevice>;
-	using DevicePlatformId = StringId<IDevicePlatform>;
-	using DevicePoolId = StringId<IDevicePool>;
-	using ProjectId = StringId<ProjectConfig>;
-
 	/// <summary>
 	/// Controller for device service
 	/// </summary>
diff --git a/Engine/Source/Programs/Horde/Horde.Build/Devices/IDevice.cs b/Engine/Source/Programs/Horde/Horde.Build/Devices/IDevice.cs
--- a/Engine/Source/Programs/Horde/Horde.Build/Devices/IDevice.cs
+++ b/Engine/Source/Programs/Horde/Horde.Build/Devices/IDevice.cs
@@ -11,13 +11,8 @@ using MongoDB.Bson.Serialization.Attributes;
 
 namespace Horde.Build.Devices
 {
-	using DeviceId = StringId<IDevice>;
-	using DevicePlatformId = StringId<IDevicePlatform>;
-	using DevicePoolId = StringId<IDevicePool>;
-	using ProjectId = StringId<ProjectConfig>;
 	using UserId = ObjectId<IUser>;
-	using StreamId = StringId<IStream>;
-
+	
 	/// <summary>
 	/// A reservation containing one or more devices
 	/// </summary>
diff --git a/Engine/Source/Programs/Horde/Horde.Build/Devices/IDeviceCollection.cs b/Engine/Source/Programs/Horde/Horde.Build/Devices/IDeviceCollection.cs
--- a/Engine/Source/Programs/Horde/Horde.Build/Devices/IDeviceCollection.cs
+++ b/Engine/Source/Programs/Horde/Horde.Build/Devices/IDeviceCollection.cs
@@ -10,10 +10,6 @@ using MongoDB.Bson;
 
 namespace Horde.Build.Devices
 {
-	using DeviceId = StringId<IDevice>;
-	using DevicePlatformId = StringId<IDevicePlatform>;
-	using DevicePoolId = StringId<IDevicePool>;
-	using ProjectId = StringId<ProjectConfig>;
 	using UserId = ObjectId<IUser>;
 
 	/// <summary>
diff --git a/Engine/Source/Programs/Horde/Horde.Build/Issues/IIssue.cs b/Engine/Source/Programs/Horde/Horde.Build/Issues/IIssue.cs
--- a/Engine/Source/Programs/Horde/Horde.Build/Issues/IIssue.cs
+++ b/Engine/Source/Programs/Horde/Horde.Build/Issues/IIssue.cs
@@ -10,7 +10,6 @@ using MongoDB.Bson;
 
 namespace Horde.Build.Issues
 {
-	using StreamId = StringId<IStream>;
 	using UserId = ObjectId<IUser>;
 
 	/// <summary>
diff --git a/Engine/Source/Programs/Horde/Horde.Build/Issues/IIssueCollection.cs b/Engine/Source/Programs/Horde/Horde.Build/Issues/IIssueCollection.cs
--- a/Engine/Source/Programs/Horde/Horde.Build/Issues/IIssueCollection.cs
+++ b/Engine/Source/Programs/Horde/Horde.Build/Issues/IIssueCollection.cs
@@ -18,8 +18,6 @@ namespace Horde.Build.Issues
 {
 	using JobId = ObjectId<IJob>;
 	using LogId = ObjectId<ILogFile>;
-	using StreamId = StringId<IStream>;
-	using TemplateId = StringId<ITemplateRef>;
 	using UserId = ObjectId<IUser>;
 
 	/// <summary>
diff --git a/Engine/Source/Programs/Horde/Horde.Build/Issues/IIssueSpan.cs b/Engine/Source/Programs/Horde/Horde.Build/Issues/IIssueSpan.cs
--- a/Engine/Source/Programs/Horde/Horde.Build/Issues/IIssueSpan.cs
+++ b/Engine/Source/Programs/Horde/Horde.Build/Issues/IIssueSpan.cs
@@ -8,8 +8,6 @@ using MongoDB.Bson;
 
 namespace Horde.Build.Issues
 {
-	using StreamId = StringId<IStream>;
-	using TemplateId = StringId<ITemplateRef>;
 	using UserId = ObjectId<IUser>;
 
 	/// <summary>
diff --git a/Engine/Source/Programs/Horde/Horde.Build/Issues/IssueCollection.cs b/Engine/Source/Programs/Horde/Horde.Build/Issues/IssueCollection.cs
--- a/Engine/Source/Programs/Horde/Horde.Build/Issues/IssueCollection.cs
+++ b/Engine/Source/Programs/Horde/Horde.Build/Issues/IssueCollection.cs
@@ -25,8 +25,6 @@ namespace Horde.Build.Issues
 {
 	using JobId = ObjectId<IJob>;
 	using LogId = ObjectId<ILogFile>;
-	using StreamId = StringId<IStream>;
-	using TemplateId = StringId<ITemplateRef>;
 	using UserId = ObjectId<IUser>;
 
 	class IssueCollection : IIssueCollection
diff --git a/Engine/Source/Programs/Horde/Horde.Build/Issues/IssueMessages.cs b/Engine/Source/Programs/Horde/Horde.Build/Issues/IssueMessages.cs
--- a/Engine/Source/Programs/Horde/Horde.Build/Issues/IssueMessages.cs
+++ b/Engine/Source/Programs/Horde/Horde.Build/Issues/IssueMessages.cs
@@ -12,10 +12,6 @@ using MongoDB.Bson;
 
 namespace Horde.Build.Issues
 {
-	using StreamId = StringId<IStream>;
-	using TemplateId = StringId<ITemplateRef>;
-	using WorkflowId = StringId<WorkflowConfig>;
-
 	/// <summary>
 	/// Identifies a particular changelist and job
 	/// </summary>
diff --git a/Engine/Source/Programs/Horde/Horde.Build/Issues/IssueReportService.cs b/Engine/Source/Programs/Horde/Horde.Build/Issues/IssueReportService.cs
--- a/Engine/Source/Programs/Horde/Horde.Build/Issues/IssueReportService.cs
+++ b/Engine/Source/Programs/Horde/Horde.Build/Issues/IssueReportService.cs
@@ -21,10 +21,6 @@ using MongoDB.Bson.Serialization.Options;
 
 namespace Horde.Build.Issues
 {
-	using StreamId = StringId<IStream>;
-	using TemplateId = StringId<ITemplateRef>;
-	using WorkflowId = StringId<WorkflowConfig>;
-
 #pragma warning disable CS1591 // Missing XML comment for publicly visible type or member
 	public class WorkflowStats
 	{
diff --git a/Engine/Source/Programs/Horde/Horde.Build/Issues/IssueService.cs b/Engine/Source/Programs/Horde/Horde.Build/Issues/IssueService.cs
--- a/Engine/Source/Programs/Horde/Horde.Build/Issues/IssueService.cs
+++ b/Engine/Source/Programs/Horde/Horde.Build/Issues/IssueService.cs
@@ -31,10 +31,7 @@ using OpenTracing.Util;
 
 namespace Horde.Build.Issues
 {
-	using StreamId = StringId<IStream>;
-	using TemplateId = StringId<ITemplateRef>;
 	using UserId = ObjectId<IUser>;
-	using WorkflowId = StringId<WorkflowConfig>;
 
 	/// <summary>
 	/// Detailed issue information
diff --git a/Engine/Source/Programs/Horde/Horde.Build/Issues/IssuesController.cs b/Engine/Source/Programs/Horde/Horde.Build/Issues/IssuesController.cs
--- a/Engine/Source/Programs/Horde/Horde.Build/Issues/IssuesController.cs
+++ b/Engine/Source/Programs/Horde/Horde.Build/Issues/IssuesController.cs
@@ -26,9 +26,7 @@ namespace Horde.Build.Issues
 {
 	using JobId = ObjectId<IJob>;
 	using LogId = ObjectId<ILogFile>;
-	using StreamId = StringId<IStream>;
 	using UserId = ObjectId<IUser>;
-	using WorkflowId = StringId<WorkflowConfig>;
 
 	/// <summary>
 	/// Controller for the /api/v1/issues endpoint
diff --git a/Engine/Source/Programs/Horde/Horde.Build/Issues/WorkflowConfig.cs b/Engine/Source/Programs/Horde/Horde.Build/Issues/WorkflowConfig.cs
--- a/Engine/Source/Programs/Horde/Horde.Build/Issues/WorkflowConfig.cs
+++ b/Engine/Source/Programs/Horde/Horde.Build/Issues/WorkflowConfig.cs
@@ -8,8 +8,6 @@ using Horde.Build.Utilities;
 
 namespace Horde.Build.Issues
 {
-	using WorkflowId = StringId<WorkflowConfig>;
-
 	/// <summary>
 	/// Configuration for an issue workflow
 	/// </summary>
diff --git a/Engine/Source/Programs/Horde/Horde.Build/Issues/WorkflowId.cs b/Engine/Source/Programs/Horde/Horde.Build/Issues/WorkflowId.cs
new file mode 100644
--- /dev/null
+++ b/Engine/Source/Programs/Horde/Horde.Build/Issues/WorkflowId.cs
@@ -0,0 +1,51 @@
+// Copyright Epic Games, Inc. All Rights Reserved.
+
+using System.ComponentModel;
+using EpicGames.Core;
+using EpicGames.Horde;
+using EpicGames.Serialization;
+using Horde.Build.Utilities;
+
+namespace Horde.Build.Issues
+{
+	/// <summary>
+	/// Identifier for a workflow
+	/// </summary>
+	/// <param name="Id">Id to construct from</param>
+	[JsonSchemaString]
+	[TypeConverter(typeof(StringIdTypeConverter<WorkflowId, WorkflowIdConverter>))]
+	[StringIdConverter(typeof(WorkflowIdConverter))]
+	[CbConverter(typeof(StringIdCbConverter<WorkflowId, WorkflowIdConverter>))]
+	public record struct WorkflowId(StringId Id)
+	{
+		/// <summary>
+		/// Empty workflow id constant
+		/// </summary>
+		public static WorkflowId Empty { get; } = default;
+
+		/// <summary>
+		/// Constructor
+		/// </summary>
+		public WorkflowId(string id) : this(new StringId(id))
+		{
+		}
+
+		/// <inheritdoc cref="StringId.IsEmpty"/>
+		public bool IsEmpty => Id.IsEmpty;
+
+		/// <inheritdoc/>
+		public override string ToString() => Id.ToString();
+	}
+
+	/// <summary>
+	/// Converter to and from <see cref="StringId"/> instances.
+	/// </summary>
+	class WorkflowIdConverter : StringIdConverter<WorkflowId>
+	{
+		/// <inheritdoc/>
+		public override WorkflowId FromStringId(StringId id) => new WorkflowId(id);
+
+		/// <inheritdoc/>
+		public override StringId ToStringId(WorkflowId value) => value.Id;
+	}
+}
diff --git a/Engine/Source/Programs/Horde/Horde.Build/Jobs/Graphs/NodeAnnotations.cs b/Engine/Source/Programs/Horde/Horde.Build/Jobs/Graphs/NodeAnnotations.cs
--- a/Engine/Source/Programs/Horde/Horde.Build/Jobs/Graphs/NodeAnnotations.cs
+++ b/Engine/Source/Programs/Horde/Horde.Build/Jobs/Graphs/NodeAnnotations.cs
@@ -7,8 +7,6 @@ using Horde.Build.Utilities;
 
 namespace Horde.Build.Jobs.Graphs
 {
-	using WorkflowId = StringId<WorkflowConfig>;
-
 	/// <summary>
 	/// Interface which wraps a generic key/value dictionary to provide specific node annotations
 	/// </summary>
diff --git a/Engine/Source/Programs/Horde/Horde.Build/Jobs/IJob.cs b/Engine/Source/Programs/Horde/Horde.Build/Jobs/IJob.cs
--- a/Engine/Source/Programs/Horde/Horde.Build/Jobs/IJob.cs
+++ b/Engine/Source/Programs/Horde/Horde.Build/Jobs/IJob.cs
@@ -30,11 +30,8 @@ namespace Horde.Build.Jobs
 	using JobId = ObjectId<IJob>;
 	using LeaseId = ObjectId<ILease>;
 	using LogId = ObjectId<ILogFile>;
-	using PoolId = StringId<IPool>;
 	using ReportPlacement = HordeCommon.Rpc.ReportPlacement;
 	using SessionId = ObjectId<ISession>;
-	using StreamId = StringId<IStream>;
-	using TemplateId = StringId<ITemplateRef>;
 	using UserId = ObjectId<IUser>;
 
 	/// <summary>
diff --git a/Engine/Source/Programs/Horde/Horde.Build/Jobs/IJobCollection.cs b/Engine/Source/Programs/Horde/Horde.Build/Jobs/IJobCollection.cs
--- a/Engine/Source/Programs/Horde/Horde.Build/Jobs/IJobCollection.cs
+++ b/Engine/Source/Programs/Horde/Horde.Build/Jobs/IJobCollection.cs
@@ -23,10 +23,7 @@ namespace Horde.Build.Jobs
 	using JobId = ObjectId<IJob>;
 	using LeaseId = ObjectId<ILease>;
 	using LogId = ObjectId<ILogFile>;
-	using PoolId = StringId<IPool>;
 	using SessionId = ObjectId<ISession>;
-	using StreamId = StringId<IStream>;
-	using TemplateId = StringId<ITemplateRef>;
 	using UserId = ObjectId<IUser>;
 
 	/// <summary>
diff --git a/Engine/Source/Programs/Horde/Horde.Build/Jobs/IJobStepRef.cs b/Engine/Source/Programs/Horde/Horde.Build/Jobs/IJobStepRef.cs
--- a/Engine/Source/Programs/Horde/Horde.Build/Jobs/IJobStepRef.cs
+++ b/Engine/Source/Programs/Horde/Horde.Build/Jobs/IJobStepRef.cs
@@ -14,10 +14,7 @@ namespace Horde.Build.Jobs
 {
 	using JobId = ObjectId<IJob>;
 	using LogId = ObjectId<ILogFile>;
-	using PoolId = StringId<IPool>;
-	using StreamId = StringId<IStream>;
-	using TemplateId = StringId<ITemplateRef>;
-
+	
 	/// <summary>
 	/// Unique id struct for JobStepRef objects. Includes a job id, batch id, and step id to uniquely identify the step.
 	/// </summary>
diff --git a/Engine/Source/Programs/Horde/Horde.Build/Jobs/IJobStepRefCollection.cs b/Engine/Source/Programs/Horde/Horde.Build/Jobs/IJobStepRefCollection.cs
--- a/Engine/Source/Programs/Horde/Horde.Build/Jobs/IJobStepRefCollection.cs
+++ b/Engine/Source/Programs/Horde/Horde.Build/Jobs/IJobStepRefCollection.cs
@@ -15,10 +15,7 @@ using Microsoft.Extensions.Logging;
 namespace Horde.Build.Jobs
 {
 	using LogId = ObjectId<ILogFile>;
-	using PoolId = StringId<IPool>;
-	using StreamId = StringId<IStream>;
-	using TemplateId = StringId<ITemplateRef>;
-
+	
 	/// <summary>
 	/// Interface for a collection of JobStepRef documents
 	/// </summary>
diff --git a/Engine/Source/Programs/Horde/Horde.Build/Jobs/JobCollection.cs b/Engine/Source/Programs/Horde/Horde.Build/Jobs/JobCollection.cs
--- a/Engine/Source/Programs/Horde/Horde.Build/Jobs/JobCollection.cs
+++ b/Engine/Source/Programs/Horde/Horde.Build/Jobs/JobCollection.cs
@@ -37,10 +37,7 @@ namespace Horde.Build.Jobs
 	using JobId = ObjectId<IJob>;
 	using LeaseId = ObjectId<ILease>;
 	using LogId = ObjectId<ILogFile>;
-	using PoolId = StringId<IPool>;
 	using SessionId = ObjectId<ISession>;
-	using StreamId = StringId<IStream>;
-	using TemplateId = StringId<ITemplateRef>;
 	using UserId = ObjectId<IUser>;
 
 	/// <summary>
diff --git a/Engine/Source/Programs/Horde/Horde.Build/Jobs/JobMessages.cs b/Engine/Source/Programs/Horde/Horde.Build/Jobs/JobMessages.cs
--- a/Engine/Source/Programs/Horde/Horde.Build/Jobs/JobMessages.cs
+++ b/Engine/Source/Programs/Horde/Horde.Build/Jobs/JobMessages.cs
@@ -15,8 +15,6 @@ using HordeCommon;
 
 namespace Horde.Build.Jobs
 {
-	using TemplateId = StringId<ITemplateRef>;
-
 	/// <summary>
 	/// State of the job
 	/// </summary>
diff --git a/Engine/Source/Programs/Horde/Horde.Build/Jobs/JobService.cs b/Engine/Source/Programs/Horde/Horde.Build/Jobs/JobService.cs
--- a/Engine/Source/Programs/Horde/Horde.Build/Jobs/JobService.cs
+++ b/Engine/Source/Programs/Horde/Horde.Build/Jobs/JobService.cs
@@ -39,8 +39,6 @@ namespace Horde.Build.Jobs
 	using LeaseId = ObjectId<ILease>;
 	using LogId = ObjectId<ILogFile>;
 	using SessionId = ObjectId<ISession>;
-	using StreamId = StringId<IStream>;
-	using TemplateId = StringId<ITemplateRef>;
 	using UserId = ObjectId<IUser>;
 
 	/// <summary>
diff --git a/Engine/Source/Programs/Horde/Horde.Build/Jobs/JobStepRefCollection.cs b/Engine/Source/Programs/Horde/Horde.Build/Jobs/JobStepRefCollection.cs
--- a/Engine/Source/Programs/Horde/Horde.Build/Jobs/JobStepRefCollection.cs
+++ b/Engine/Source/Programs/Horde/Horde.Build/Jobs/JobStepRefCollection.cs
@@ -18,10 +18,7 @@ using MongoDB.Driver;
 namespace Horde.Build.Jobs
 {
 	using LogId = ObjectId<ILogFile>;
-	using PoolId = StringId<IPool>;
-	using StreamId = StringId<IStream>;
-	using TemplateId = StringId<ITemplateRef>;
-
+	
 	/// <summary>
 	/// Collection of JobStepRef documents
 	/// </summary>
diff --git a/Engine/Source/Programs/Horde/Horde.Build/Jobs/JobTaskSource.cs b/Engine/Source/Programs/Horde/Horde.Build/Jobs/JobTaskSource.cs
--- a/Engine/Source/Programs/Horde/Horde.Build/Jobs/JobTaskSource.cs
+++ b/Engine/Source/Programs/Horde/Horde.Build/Jobs/JobTaskSource.cs
@@ -36,9 +36,7 @@ namespace Horde.Build.Jobs
 	using JobId = ObjectId<IJob>;
 	using LeaseId = ObjectId<ILease>;
 	using LogId = ObjectId<ILogFile>;
-	using PoolId = StringId<IPool>;
-	using StreamId = StringId<IStream>;
-
+	
 	/// <summary>
 	/// Background service to dispatch pending work to agents in priority order.
 	/// </summary>
@@ -323,7 +321,7 @@ namespace Horde.Build.Jobs
 		/// <returns></returns>
 		internal static Dictionary<PoolId, PoolStatus> GetPoolStatus(DateTime utcNow, List<IPool> pools, List<IAgent> agents)
 		{
-			Dictionary<StringId<IPool>,PoolStatus> poolStatus = new ();
+			Dictionary<PoolId, PoolStatus> poolStatus = new ();
 
 			foreach (IPool pool in pools)
 			{
diff --git a/Engine/Source/Programs/Horde/Horde.Build/Jobs/JobsController.cs b/Engine/Source/Programs/Horde/Horde.Build/Jobs/JobsController.cs
--- a/Engine/Source/Programs/Horde/Horde.Build/Jobs/JobsController.cs
+++ b/Engine/Source/Programs/Horde/Horde.Build/Jobs/JobsController.cs
@@ -31,8 +31,6 @@ using OpenTracing.Util;
 namespace Horde.Build.Jobs
 {
 	using JobId = ObjectId<IJob>;
-	using StreamId = StringId<IStream>;
-	using TemplateId = StringId<ITemplateRef>;
 	using UserId = ObjectId<IUser>;
 
 	/// <summary>
diff --git a/Engine/Source/Programs/Horde/Horde.Build/Jobs/Schedules/ScheduleService.cs b/Engine/Source/Programs/Horde/Horde.Build/Jobs/Schedules/ScheduleService.cs
--- a/Engine/Source/Programs/Horde/Horde.Build/Jobs/Schedules/ScheduleService.cs
+++ b/Engine/Source/Programs/Horde/Horde.Build/Jobs/Schedules/ScheduleService.cs
@@ -29,8 +29,6 @@ using StackExchange.Redis;
 namespace Horde.Build.Jobs.Schedules
 {
 	using JobId = ObjectId<IJob>;
-	using StreamId = StringId<IStream>;
-	using TemplateId = StringId<ITemplateRef>;
 
 	/// <summary>
 	/// Manipulates schedule instances
diff --git a/Engine/Source/Programs/Horde/Horde.Build/Jobs/TestData/ITestData.cs b/Engine/Source/Programs/Horde/Horde.Build/Jobs/TestData/ITestData.cs
--- a/Engine/Source/Programs/Horde/Horde.Build/Jobs/TestData/ITestData.cs
+++ b/Engine/Source/Programs/Horde/Horde.Build/Jobs/TestData/ITestData.cs
@@ -10,9 +10,7 @@ using MongoDB.Bson.Serialization.Attributes;
 namespace Horde.Build.Jobs.TestData
 {
 	using JobId = ObjectId<IJob>;
-	using StreamId = StringId<IStream>;
-	using TemplateId = StringId<ITemplateRef>;
-
+	
 	using TestId = ObjectId<ITest>;
 	using TestSuiteId = ObjectId<ITestSuite>;
 	using TestMetaId = ObjectId<ITestMeta>;
diff --git a/Engine/Source/Programs/Horde/Horde.Build/Jobs/TestData/ITestDataCollection.cs b/Engine/Source/Programs/Horde/Horde.Build/Jobs/TestData/ITestDataCollection.cs
--- a/Engine/Source/Programs/Horde/Horde.Build/Jobs/TestData/ITestDataCollection.cs
+++ b/Engine/Source/Programs/Horde/Horde.Build/Jobs/TestData/ITestDataCollection.cs
@@ -10,7 +10,6 @@ using MongoDB.Bson;
 namespace Horde.Build.Jobs.TestData
 {	
 	using JobId = ObjectId<IJob>;
-	using StreamId = StringId<IStream>;
 	using TestId = ObjectId<ITest>;
 	using TestSuiteId = ObjectId<ITestSuite>;
 	using TestMetaId = ObjectId<ITestMeta>;
diff --git a/Engine/Source/Programs/Horde/Horde.Build/Jobs/TestData/TestDataCollection.cs b/Engine/Source/Programs/Horde/Horde.Build/Jobs/TestData/TestDataCollection.cs
--- a/Engine/Source/Programs/Horde/Horde.Build/Jobs/TestData/TestDataCollection.cs
+++ b/Engine/Source/Programs/Horde/Horde.Build/Jobs/TestData/TestDataCollection.cs
@@ -21,9 +21,7 @@ using Microsoft.Extensions.Options;
 namespace Horde.Build.Jobs.TestData
 {
 	using JobId = ObjectId<IJob>;
-	using StreamId = StringId<IStream>;
-	using TemplateId = StringId<ITemplateRef>;
-
+	
 	using TestId = ObjectId<ITest>;
 	using TestSuiteId = ObjectId<ITestSuite>;
 	using TestMetaId = ObjectId<ITestMeta>;
diff --git a/Engine/Source/Programs/Horde/Horde.Build/Jobs/TestData/TestDataController.cs b/Engine/Source/Programs/Horde/Horde.Build/Jobs/TestData/TestDataController.cs
--- a/Engine/Source/Programs/Horde/Horde.Build/Jobs/TestData/TestDataController.cs
+++ b/Engine/Source/Programs/Horde/Horde.Build/Jobs/TestData/TestDataController.cs
@@ -20,9 +20,7 @@ using Horde.Build.Server;
 
 namespace Horde.Build.Jobs.TestData
 {
-	
 	using JobId = ObjectId<IJob>;
-	using StreamId = StringId<IStream>;
 	using TestId = ObjectId<ITest>;
 	using TestSuiteId = ObjectId<ITestSuite>;
 	using TestMetaId = ObjectId<ITestMeta>;
diff --git a/Engine/Source/Programs/Horde/Horde.Build/Jobs/TestData/TestDataMessages.cs b/Engine/Source/Programs/Horde/Horde.Build/Jobs/TestData/TestDataMessages.cs
--- a/Engine/Source/Programs/Horde/Horde.Build/Jobs/TestData/TestDataMessages.cs
+++ b/Engine/Source/Programs/Horde/Horde.Build/Jobs/TestData/TestDataMessages.cs
@@ -11,8 +11,6 @@ using MongoDB.Bson.Serialization;
 
 namespace Horde.Build.Jobs.TestData
 {
-	using StreamId = StringId<IStream>;	
-
 	/// <summary>
 	/// Response object describing test data to store
 	/// </summary>
diff --git a/Engine/Source/Programs/Horde/Horde.Build/Jobs/TestData/TestDataService.cs b/Engine/Source/Programs/Horde/Horde.Build/Jobs/TestData/TestDataService.cs
--- a/Engine/Source/Programs/Horde/Horde.Build/Jobs/TestData/TestDataService.cs
+++ b/Engine/Source/Programs/Horde/Horde.Build/Jobs/TestData/TestDataService.cs
@@ -14,8 +14,6 @@ using Microsoft.Extensions.Options;
 
 namespace Horde.Build.Jobs.TestData
 {
-	
-	using StreamId = StringId<IStream>;
 	using TestId = ObjectId<ITest>;
 	using TestSuiteId = ObjectId<ITestSuite>;
 	using TestMetaId = ObjectId<ITestMeta>;
diff --git a/Engine/Source/Programs/Horde/Horde.Build/Notifications/IEvent.cs b/Engine/Source/Programs/Horde/Horde.Build/Notifications/IEvent.cs
--- a/Engine/Source/Programs/Horde/Horde.Build/Notifications/IEvent.cs
+++ b/Engine/Source/Programs/Horde/Horde.Build/Notifications/IEvent.cs
@@ -5,9 +5,6 @@ using Horde.Build.Utilities;
 
 namespace Horde.Build.Notifications
 {
-	using StreamId = StringId<IStream>;
-	using TemplateId = StringId<ITemplateRef>;
-
 	/// <summary>
 	/// Base interface for events that users can subscribe for notifications to
 	/// </summary>
diff --git a/Engine/Source/Programs/Horde/Horde.Build/Notifications/Sinks/SlackNotificationSink.cs b/Engine/Source/Programs/Horde/Horde.Build/Notifications/Sinks/SlackNotificationSink.cs
--- a/Engine/Source/Programs/Horde/Horde.Build/Notifications/Sinks/SlackNotificationSink.cs
+++ b/Engine/Source/Programs/Horde/Horde.Build/Notifications/Sinks/SlackNotificationSink.cs
@@ -47,10 +47,7 @@ namespace Horde.Build.Notifications.Sinks
 {
 	using JobId = ObjectId<IJob>;
 	using LogId = ObjectId<ILogFile>;
-	using StreamId = StringId<IStream>;
-	using TemplateId = StringId<ITemplateRef>;
 	using UserId = ObjectId<IUser>;
-	using WorkflowId = StringId<WorkflowConfig>;
 
 	/// <summary>
 	/// Maintains a connection to Slack, in order to receive socket-mode notifications of user interactions
diff --git a/Engine/Source/Programs/Horde/Horde.Build/Notifications/SubscriptionCollection.cs b/Engine/Source/Programs/Horde/Horde.Build/Notifications/SubscriptionCollection.cs
--- a/Engine/Source/Programs/Horde/Horde.Build/Notifications/SubscriptionCollection.cs
+++ b/Engine/Source/Programs/Horde/Horde.Build/Notifications/SubscriptionCollection.cs
@@ -18,8 +18,6 @@ using MongoDB.Driver;
 
 namespace Horde.Build.Notifications
 {
-	using StreamId = StringId<IStream>;
-	using TemplateId = StringId<ITemplateRef>;
 	using UserId = ObjectId<IUser>;
 
 	/// <summary>
diff --git a/Engine/Source/Programs/Horde/Horde.Build/Notifications/SubscriptionMessages.cs b/Engine/Source/Programs/Horde/Horde.Build/Notifications/SubscriptionMessages.cs
--- a/Engine/Source/Programs/Horde/Horde.Build/Notifications/SubscriptionMessages.cs
+++ b/Engine/Source/Programs/Horde/Horde.Build/Notifications/SubscriptionMessages.cs
@@ -10,9 +10,6 @@ using HordeCommon;
 
 namespace Horde.Build.Notifications
 {
-	using StreamId = StringId<IStream>;
-	using TemplateId = StringId<ITemplateRef>;
-
 	/// <summary>
 	/// Request object for creating a new subscription
 	/// </summary>
diff --git a/Engine/Source/Programs/Horde/Horde.Build/Perforce/ICommit.cs b/Engine/Source/Programs/Horde/Horde.Build/Perforce/ICommit.cs
--- a/Engine/Source/Programs/Horde/Horde.Build/Perforce/ICommit.cs
+++ b/Engine/Source/Programs/Horde/Horde.Build/Perforce/ICommit.cs
@@ -13,7 +13,6 @@ using Horde.Build.Utilities;
 
 namespace Horde.Build.Perforce
 {
-	using StreamId = StringId<IStream>;
 	using UserId = ObjectId<IUser>;
 
 	/// <summary>
diff --git a/Engine/Source/Programs/Horde/Horde.Build/Perforce/PerforceReplicator.cs b/Engine/Source/Programs/Horde/Horde.Build/Perforce/PerforceReplicator.cs
--- a/Engine/Source/Programs/Horde/Horde.Build/Perforce/PerforceReplicator.cs
+++ b/Engine/Source/Programs/Horde/Horde.Build/Perforce/PerforceReplicator.cs
@@ -21,8 +21,6 @@ using Microsoft.Extensions.Logging;
 
 namespace Horde.Build.Perforce
 {
-	using StreamId = StringId<IStream>;
-
 	/// <summary>
 	/// Options for replicating commits
 	/// </summary>
diff --git a/Engine/Source/Programs/Horde/Horde.Build/Perforce/PerforceService.cs b/Engine/Source/Programs/Horde/Horde.Build/Perforce/PerforceService.cs
--- a/Engine/Source/Programs/Horde/Horde.Build/Perforce/PerforceService.cs
+++ b/Engine/Source/Programs/Horde/Horde.Build/Perforce/PerforceService.cs
@@ -27,7 +27,6 @@ using OpenTracing.Util;
 
 namespace Horde.Build.Perforce
 {
-	using StreamId = StringId<IStream>;
 	using UserId = ObjectId<IUser>;
 
 	/// <summary>
diff --git a/Engine/Source/Programs/Horde/Horde.Build/Perforce/PerforceServiceCache.cs b/Engine/Source/Programs/Horde/Horde.Build/Perforce/PerforceServiceCache.cs
--- a/Engine/Source/Programs/Horde/Horde.Build/Perforce/PerforceServiceCache.cs
+++ b/Engine/Source/Programs/Horde/Horde.Build/Perforce/PerforceServiceCache.cs
@@ -27,7 +27,6 @@ using StackExchange.Redis;
 
 namespace Horde.Build.Perforce
 {
-	using StreamId = StringId<IStream>;
 	using UserId = ObjectId<IUser>;
 
 	/// <summary>
diff --git a/Engine/Source/Programs/Horde/Horde.Build/Perforce/ReplicationService.cs b/Engine/Source/Programs/Horde/Horde.Build/Perforce/ReplicationService.cs
--- a/Engine/Source/Programs/Horde/Horde.Build/Perforce/ReplicationService.cs
+++ b/Engine/Source/Programs/Horde/Horde.Build/Perforce/ReplicationService.cs
@@ -21,8 +21,6 @@ using Microsoft.Extensions.Options;
 
 namespace Horde.Build.Perforce
 {
-	using StreamId = StringId<IStream>;
-
 	/// <summary>
 	/// Exception triggered during content replication
 	/// </summary>
diff --git a/Engine/Source/Programs/Horde/Horde.Build/Projects/ProjectConfig.cs b/Engine/Source/Programs/Horde/Horde.Build/Projects/ProjectConfig.cs
--- a/Engine/Source/Programs/Horde/Horde.Build/Projects/ProjectConfig.cs
+++ b/Engine/Source/Programs/Horde/Horde.Build/Projects/ProjectConfig.cs
@@ -19,8 +19,6 @@ using HordeCommon.Rpc.Tasks;
 
 namespace Horde.Build.Projects
 {
-	using ProjectId = StringId<ProjectConfig>;
-
 	/// <summary>
 	/// Stores configuration for a project
 	/// </summary>
diff --git a/Engine/Source/Programs/Horde/Horde.Build/Projects/ProjectId.cs b/Engine/Source/Programs/Horde/Horde.Build/Projects/ProjectId.cs
new file mode 100644
--- /dev/null
+++ b/Engine/Source/Programs/Horde/Horde.Build/Projects/ProjectId.cs
@@ -0,0 +1,46 @@
+﻿// Copyright Epic Games, Inc. All Rights Reserved.
+
+using System.ComponentModel;
+using EpicGames.Core;
+using EpicGames.Horde;
+using EpicGames.Serialization;
+using Horde.Build.Utilities;
+
+namespace Horde.Build.Projects
+{
+	/// <summary>
+	/// Identifier for a pool
+	/// </summary>
+	/// <param name="Id">Id to construct from</param>
+	[JsonSchemaString]
+	[TypeConverter(typeof(StringIdTypeConverter<ProjectId, ProjectIdConverter>))]
+	[StringIdConverter(typeof(ProjectIdConverter))]
+	[CbConverter(typeof(StringIdCbConverter<ProjectId, ProjectIdConverter>))]
+	public record struct ProjectId(StringId Id)
+	{
+		/// <summary>
+		/// Constructor
+		/// </summary>
+		public ProjectId(string id) : this(new StringId(id))
+		{
+		}
+
+		/// <inheritdoc cref="StringId.IsEmpty"/>
+		public bool IsEmpty => Id.IsEmpty;
+
+		/// <inheritdoc/>
+		public override string ToString() => Id.ToString();
+	}
+
+	/// <summary>
+	/// Converter to and from <see cref="StringId"/> instances.
+	/// </summary>
+	class ProjectIdConverter : StringIdConverter<ProjectId>
+	{
+		/// <inheritdoc/>
+		public override ProjectId FromStringId(StringId id) => new ProjectId(id);
+
+		/// <inheritdoc/>
+		public override StringId ToStringId(ProjectId value) => value.Id;
+	}
+}
diff --git a/Engine/Source/Programs/Horde/Horde.Build/Projects/ProjectsController.cs b/Engine/Source/Programs/Horde/Horde.Build/Projects/ProjectsController.cs
--- a/Engine/Source/Programs/Horde/Horde.Build/Projects/ProjectsController.cs
+++ b/Engine/Source/Programs/Horde/Horde.Build/Projects/ProjectsController.cs
@@ -15,8 +15,6 @@ using Microsoft.Extensions.Options;
 
 namespace Horde.Build.Projects
 {
-	using ProjectId = StringId<ProjectConfig>;
-
 	/// <summary>
 	/// Controller for the /api/v1/projects endpoint
 	/// </summary>
diff --git a/Engine/Source/Programs/Horde/Horde.Build/Server/GlobalConfig.cs b/Engine/Source/Programs/Horde/Horde.Build/Server/GlobalConfig.cs
--- a/Engine/Source/Programs/Horde/Horde.Build/Server/GlobalConfig.cs
+++ b/Engine/Source/Programs/Horde/Horde.Build/Server/GlobalConfig.cs
@@ -27,10 +27,6 @@ using Horde.Build.Utilities;
 namespace Horde.Build.Server
 {
 	using UserId = ObjectId<IUser>;
-	using ProjectId = StringId<ProjectConfig>;
-	using StreamId = StringId<IStream>;
-	using AgentSoftwareChannelName = StringId<AgentSoftwareChannels>;
-	using ToolId = StringId<ITool>;
 
 	/// <summary>
 	/// Directive to merge config data from another source
diff --git a/Engine/Source/Programs/Horde/Horde.Build/Server/MongoService.cs b/Engine/Source/Programs/Horde/Horde.Build/Server/MongoService.cs
--- a/Engine/Source/Programs/Horde/Horde.Build/Server/MongoService.cs
+++ b/Engine/Source/Programs/Horde/Horde.Build/Server/MongoService.cs
@@ -35,8 +35,6 @@ using StackExchange.Redis;
 
 namespace Horde.Build.Server
 {
-	using SingletonId = StringId<SingletonBase>;
-
 	/// <summary>
 	/// Constrained set of parameters for building a mongo index
 	/// </summary>
diff --git a/Engine/Source/Programs/Horde/Horde.Build/Server/RpcService.cs b/Engine/Source/Programs/Horde/Horde.Build/Server/RpcService.cs
--- a/Engine/Source/Programs/Horde/Horde.Build/Server/RpcService.cs
+++ b/Engine/Source/Programs/Horde/Horde.Build/Server/RpcService.cs
@@ -43,8 +43,6 @@ using Microsoft.Extensions.Options;
 
 namespace Horde.Build.Server
 {
-	using AgentSoftwareChannelName = StringId<AgentSoftwareChannels>;
-	using IStream = Horde.Build.Streams.IStream;
 	using JobId = ObjectId<IJob>;
 	using LogId = ObjectId<ILogFile>;
 	using RpcAgentCapabilities = HordeCommon.Rpc.Messages.AgentCapabilities;
@@ -55,8 +53,7 @@ namespace Horde.Build.Server
 	using RpcUpdateJobRequest = HordeCommon.Rpc.UpdateJobRequest;
 	using RpcUpdateStepRequest = HordeCommon.Rpc.UpdateStepRequest;
 	using SessionId = ObjectId<ISession>;
-	using StreamId = StringId<IStream>;
-
+	
 	/// <summary>
 	/// Implements the Horde gRPC service for bots updating their status and dequeing work
 	/// </summary>
diff --git a/Engine/Source/Programs/Horde/Horde.Build/ServerSettings.cs b/Engine/Source/Programs/Horde/Horde.Build/ServerSettings.cs
--- a/Engine/Source/Programs/Horde/Horde.Build/ServerSettings.cs
+++ b/Engine/Source/Programs/Horde/Horde.Build/ServerSettings.cs
@@ -5,6 +5,8 @@ using System.Collections.Generic;
 using System.Linq;
 using System.Security.Claims;
 using System.Text.Json.Serialization;
+using EpicGames.Core;
+using EpicGames.Horde;
 using EpicGames.Horde.Storage;
 using EpicGames.Horde.Storage.Nodes;
 using EpicGames.Perforce;
@@ -18,8 +20,6 @@ using Serilog.Events;
 
 namespace Horde.Build
 {
-	using PerforceConnectionId = StringId<PerforceConnectionSettings>;
-
 	/// <summary>
 	/// Types of storage backend to use
 	/// </summary>
@@ -728,6 +728,37 @@ namespace Horde.Build
 		}
 	}
 
+	/// <summary>
+	/// Identifier for a pool
+	/// </summary>
+	/// <param name="Id">Id to construct from</param>
+	[JsonSchemaString]
+	[StringIdConverter(typeof(PerforceConnectionIdConverter))]
+	public record struct PerforceConnectionId(StringId Id)
+	{
+		/// <summary>
+		/// Constructor
+		/// </summary>
+		public PerforceConnectionId(string id) : this(new StringId(id))
+		{
+		}
+
+		/// <inheritdoc/>
+		public override string ToString() => Id.ToString();
+	}
+
+	/// <summary>
+	/// Converter to and from <see cref="StringId"/> instances.
+	/// </summary>
+	class PerforceConnectionIdConverter : StringIdConverter<PerforceConnectionId>
+	{
+		/// <inheritdoc/>
+		public override PerforceConnectionId FromStringId(StringId id) => new PerforceConnectionId(id);
+
+		/// <inheritdoc/>
+		public override StringId ToStringId(PerforceConnectionId value) => value.Id;
+	}
+
 	/// <summary>
 	/// Perforce connection information for use by the Horde server (for reading config files, etc...)
 	/// </summary>
diff --git a/Engine/Source/Programs/Horde/Horde.Build/Startup.cs b/Engine/Source/Programs/Horde/Horde.Build/Startup.cs
--- a/Engine/Source/Programs/Horde/Horde.Build/Startup.cs
+++ b/Engine/Source/Programs/Horde/Horde.Build/Startup.cs
@@ -106,8 +106,8 @@ namespace Horde.Build
 	{
 		static Startup()
 		{
-			ProtoBuf.Meta.RuntimeTypeModel.Default[typeof(StringId<ProjectConfig>)].SetSurrogate(typeof(StringIdProto<ProjectConfig>));
-			ProtoBuf.Meta.RuntimeTypeModel.Default[typeof(StringId<IStream>)].SetSurrogate(typeof(StringIdProto<IStream>));
+			ProtoBuf.Meta.RuntimeTypeModel.Default[typeof(ProjectId)].SetSurrogate(typeof(StringIdProto<ProjectId, ProjectIdConverter>));
+			ProtoBuf.Meta.RuntimeTypeModel.Default[typeof(StreamId)].SetSurrogate(typeof(StringIdProto<StreamId, StreamIdConverter>));
 		}
 
 		class GrpcExceptionInterceptor : Interceptor
@@ -782,7 +782,7 @@ namespace Horde.Build
 			options.Converters.Add(new JsonObjectIdConverter());
 			options.Converters.Add(new JsonStringEnumConverter());
 			options.Converters.Add(new JsonKnownTypesConverterFactory());
-			options.Converters.Add(new JsonStringIdConverterFactory());
+			options.Converters.Add(new StringIdJsonConverterFactory());
 			options.Converters.Add(new JsonDateTimeConverter());
 			options.Converters.Add(new JsonTimeSpanConverter());
 		}
@@ -959,7 +959,7 @@ namespace Horde.Build
 				BsonSerializer.RegisterSerializer(new NamespaceIdBsonSerializer());
 				BsonSerializer.RegisterSerializer(new ConditionSerializer());
 				BsonSerializer.RegisterSerializationProvider(new BsonSerializationProvider());
-				BsonSerializer.RegisterSerializationProvider(new StringIdSerializationProvider());
+				BsonSerializer.RegisterSerializationProvider(new StringIdBsonSerializationProvider());
 				BsonSerializer.RegisterSerializationProvider(new ObjectIdSerializationProvider());
 			}
 		}
diff --git a/Engine/Source/Programs/Horde/Horde.Build/Storage/BackendId.cs b/Engine/Source/Programs/Horde/Horde.Build/Storage/BackendId.cs
new file mode 100644
--- /dev/null
+++ b/Engine/Source/Programs/Horde/Horde.Build/Storage/BackendId.cs
@@ -0,0 +1,51 @@
+// Copyright Epic Games, Inc. All Rights Reserved.
+
+using System.ComponentModel;
+using EpicGames.Core;
+using EpicGames.Horde;
+using EpicGames.Serialization;
+using Horde.Build.Utilities;
+
+namespace Horde.Build.Storage
+{
+	/// <summary>
+	/// Identifier for a pool
+	/// </summary>
+	/// <param name="Id">Id to construct from</param>
+	[JsonSchemaString]
+	[TypeConverter(typeof(StringIdTypeConverter<BackendId, BackendIdConverter>))]
+	[StringIdConverter(typeof(BackendIdConverter))]
+	[CbConverter(typeof(StringIdCbConverter<BackendId, BackendIdConverter>))]
+	public record struct BackendId(StringId Id)
+	{
+		/// <summary>
+		/// Empty backend id
+		/// </summary>
+		public static BackendId Empty { get; } = default;
+
+		/// <summary>
+		/// Constructor
+		/// </summary>
+		public BackendId(string id) : this(new StringId(id))
+		{
+		}
+
+		/// <inheritdoc cref="StringId.IsEmpty"/>
+		public bool IsEmpty => Id.IsEmpty;
+
+		/// <inheritdoc/>
+		public override string ToString() => Id.ToString();
+	}
+
+	/// <summary>
+	/// Converter to and from <see cref="StringId"/> instances.
+	/// </summary>
+	class BackendIdConverter : StringIdConverter<BackendId>
+	{
+		/// <inheritdoc/>
+		public override BackendId FromStringId(StringId id) => new BackendId(id);
+
+		/// <inheritdoc/>
+		public override StringId ToStringId(BackendId value) => value.Id;
+	}
+}
diff --git a/Engine/Source/Programs/Horde/Horde.Build/Storage/StorageConfig.cs b/Engine/Source/Programs/Horde/Horde.Build/Storage/StorageConfig.cs
--- a/Engine/Source/Programs/Horde/Horde.Build/Storage/StorageConfig.cs
+++ b/Engine/Source/Programs/Horde/Horde.Build/Storage/StorageConfig.cs
@@ -16,8 +16,6 @@ using Horde.Build.Utilities;
 
 namespace Horde.Build.Storage
 {
-	using BackendId = StringId<IStorageBackend>;
-
 	/// <summary>
 	/// Well-known namespace identifiers
 	/// </summary>
diff --git a/Engine/Source/Programs/Horde/Horde.Build/Storage/StorageService.cs b/Engine/Source/Programs/Horde/Horde.Build/Storage/StorageService.cs
--- a/Engine/Source/Programs/Horde/Horde.Build/Storage/StorageService.cs
+++ b/Engine/Source/Programs/Horde/Horde.Build/Storage/StorageService.cs
@@ -36,8 +36,6 @@ using StackExchange.Redis;
 
 namespace Horde.Build.Storage
 {
-	using BackendId = StringId<IStorageBackend>;
-
 	/// <summary>
 	/// Exception thrown by the <see cref="StorageService"/>
 	/// </summary>
diff --git a/Engine/Source/Programs/Horde/Horde.Build/Streams/IStream.cs b/Engine/Source/Programs/Horde/Horde.Build/Streams/IStream.cs
--- a/Engine/Source/Programs/Horde/Horde.Build/Streams/IStream.cs
+++ b/Engine/Source/Programs/Horde/Horde.Build/Streams/IStream.cs
@@ -17,9 +17,6 @@ using Horde.Build.Utilities;
 namespace Horde.Build.Streams
 {
 	using JobId = ObjectId<IJob>;
-	using ProjectId = StringId<ProjectConfig>;
-	using StreamId = StringId<IStream>;
-	using TemplateId = StringId<ITemplateRef>;
 	using UserId = ObjectId<IUser>;
 
 	/// <summary>
diff --git a/Engine/Source/Programs/Horde/Horde.Build/Streams/IStreamCollection.cs b/Engine/Source/Programs/Horde/Horde.Build/Streams/IStreamCollection.cs
--- a/Engine/Source/Programs/Horde/Horde.Build/Streams/IStreamCollection.cs
+++ b/Engine/Source/Programs/Horde/Horde.Build/Streams/IStreamCollection.cs
@@ -10,7 +10,6 @@ using Horde.Build.Utilities;
 namespace Horde.Build.Streams
 {
 	using JobId = ObjectId<IJob>;
-	using TemplateId = StringId<ITemplateRef>;
 
 	/// <summary>
 	/// Collection of stream documents
diff --git a/Engine/Source/Programs/Horde/Horde.Build/Streams/StreamCollection.cs b/Engine/Source/Programs/Horde/Horde.Build/Streams/StreamCollection.cs
--- a/Engine/Source/Programs/Horde/Horde.Build/Streams/StreamCollection.cs
+++ b/Engine/Source/Programs/Horde/Horde.Build/Streams/StreamCollection.cs
@@ -18,9 +18,6 @@ using MongoDB.Driver;
 namespace Horde.Build.Streams
 {
 	using JobId = ObjectId<IJob>;
-	using ProjectId = StringId<ProjectConfig>;
-	using StreamId = StringId<IStream>;
-	using TemplateId = StringId<ITemplateRef>;
 	using UserId = ObjectId<IUser>;
 
 	/// <summary>
diff --git a/Engine/Source/Programs/Horde/Horde.Build/Streams/StreamConfig.cs b/Engine/Source/Programs/Horde/Horde.Build/Streams/StreamConfig.cs
--- a/Engine/Source/Programs/Horde/Horde.Build/Streams/StreamConfig.cs
+++ b/Engine/Source/Programs/Horde/Horde.Build/Streams/StreamConfig.cs
@@ -10,6 +10,7 @@ using System.Text.Json;
 using System.Text.Json.Serialization;
 using System.Text.RegularExpressions;
 using EpicGames.Core;
+using EpicGames.Horde;
 using EpicGames.Horde.Common;
 using Horde.Build.Acls;
 using Horde.Build.Agents.Pools;
@@ -27,11 +28,6 @@ using HordeCommon.Rpc.Tasks;
 
 namespace Horde.Build.Streams
 {
-	using StreamId = StringId<IStream>;
-	using PoolId = StringId<IPool>;
-	using TemplateId = StringId<ITemplateRef>;
-	using WorkflowId = StringId<WorkflowConfig>;
-
 	/// <summary>
 	/// Flags identifying content of a changelist
 	/// </summary>
@@ -760,7 +756,7 @@ namespace Horde.Build.Streams
 
 			if (Id.IsEmpty)
 			{
-				Id = TemplateId.Sanitize(Name);
+				Id = new TemplateId(StringId.Sanitize(Name));
 			}
 		}
 	}
diff --git a/Engine/Source/Programs/Horde/Horde.Build/Streams/StreamId.cs b/Engine/Source/Programs/Horde/Horde.Build/Streams/StreamId.cs
new file mode 100644
--- /dev/null
+++ b/Engine/Source/Programs/Horde/Horde.Build/Streams/StreamId.cs
@@ -0,0 +1,59 @@
+// Copyright Epic Games, Inc. All Rights Reserved.
+
+using System.ComponentModel;
+using EpicGames.Core;
+using EpicGames.Horde;
+using EpicGames.Serialization;
+using Horde.Build.Utilities;
+using OpenTracing;
+
+namespace Horde.Build.Streams
+{
+	/// <summary>
+	/// Identifier for a stream
+	/// </summary>
+	/// <param name="Id">Id to construct from</param>
+	[JsonSchemaString]
+	[TypeConverter(typeof(StringIdTypeConverter<StreamId, StreamIdConverter>))]
+	[StringIdConverter(typeof(StreamIdConverter))]
+	[CbConverter(typeof(StringIdCbConverter<StreamId, StreamIdConverter>))]
+	public record struct StreamId(StringId Id)
+	{
+		/// <summary>
+		/// Constructor
+		/// </summary>
+		public StreamId(string id) : this(new StringId(id))
+		{
+		}
+
+		/// <inheritdoc cref="StringId.IsEmpty"/>
+		public bool IsEmpty => Id.IsEmpty;
+
+		/// <inheritdoc/>
+		public override string ToString() => Id.ToString();
+	}
+
+	/// <summary>
+	/// Converter to and from <see cref="StringId"/> instances.
+	/// </summary>
+	class StreamIdConverter : StringIdConverter<StreamId>
+	{
+		/// <inheritdoc/>
+		public override StreamId FromStringId(StringId id) => new StreamId(id);
+
+		/// <inheritdoc/>
+		public override StringId ToStringId(StreamId value) => value.Id;
+	}
+
+	/// <summary>
+	/// Extension methods for stream id values
+	/// </summary>
+	public static class StreamIdExtensions
+	{
+		/// <inheritdoc cref="ISpan.SetTag(System.String, System.String)"/>
+		public static ISpan SetTag(this ISpan span, string key, StreamId value) => span.SetTag(key, value.Id.ToString());
+
+		/// <inheritdoc cref="ISpan.SetTag(System.String, System.String)"/>
+		public static ISpan SetTag(this ISpan span, string key, StreamId? value) => span.SetTag(key, value.HasValue ? value.Value.ToString() : null);
+	}
+}
diff --git a/Engine/Source/Programs/Horde/Horde.Build/Streams/StreamMessages.cs b/Engine/Source/Programs/Horde/Horde.Build/Streams/StreamMessages.cs
--- a/Engine/Source/Programs/Horde/Horde.Build/Streams/StreamMessages.cs
+++ b/Engine/Source/Programs/Horde/Horde.Build/Streams/StreamMessages.cs
@@ -19,9 +19,6 @@ using System.Linq;
 namespace Horde.Build.Streams
 {
 	using JobId = ObjectId<IJob>;
-	using StreamId = StringId<IStream>;
-	using ProjectId = StringId<ProjectConfig>;
-	using TemplateId = StringId<ITemplateRef>;
 
 	/// <summary>
 	/// Step state update request
diff --git a/Engine/Source/Programs/Horde/Horde.Build/Streams/StreamsController.cs b/Engine/Source/Programs/Horde/Horde.Build/Streams/StreamsController.cs
--- a/Engine/Source/Programs/Horde/Horde.Build/Streams/StreamsController.cs
+++ b/Engine/Source/Programs/Horde/Horde.Build/Streams/StreamsController.cs
@@ -22,10 +22,6 @@ using TimeZoneConverter;
 
 namespace Horde.Build.Streams
 {
-	using ProjectId = StringId<ProjectConfig>;
-	using StreamId = StringId<IStream>;
-	using TemplateId = StringId<ITemplateRef>;
-
 	/// <summary>
 	/// Controller for the /api/v1/streams endpoint
 	/// </summary>
diff --git a/Engine/Source/Programs/Horde/Horde.Build/Streams/TemplateId.cs b/Engine/Source/Programs/Horde/Horde.Build/Streams/TemplateId.cs
new file mode 100644
--- /dev/null
+++ b/Engine/Source/Programs/Horde/Horde.Build/Streams/TemplateId.cs
@@ -0,0 +1,61 @@
+// Copyright Epic Games, Inc. All Rights Reserved.
+
+using System;
+using System.ComponentModel;
+using System.Linq;
+using EpicGames.Core;
+using EpicGames.Horde;
+using EpicGames.Serialization;
+using Horde.Build.Utilities;
+using OpenTracing;
+
+namespace Horde.Build.Streams
+{
+	/// <summary>
+	/// Identifier for a job template
+	/// </summary>
+	/// <param name="Id">Id to construct from</param>
+	[JsonSchemaString]
+	[TypeConverter(typeof(StringIdTypeConverter<TemplateId, TemplateIdConverter>))]
+	[StringIdConverter(typeof(TemplateIdConverter))]
+	[CbConverter(typeof(StringIdCbConverter<TemplateId, TemplateIdConverter>))]
+	public record struct TemplateId(StringId Id)
+	{
+		/// <summary>
+		/// Constructor
+		/// </summary>
+		public TemplateId(string id) : this(new StringId(id))
+		{
+		}
+
+		/// <inheritdoc cref="StringId.IsEmpty"/>
+		public bool IsEmpty => Id.IsEmpty;
+
+		/// <inheritdoc/>
+		public override string ToString() => Id.ToString();
+	}
+
+	/// <summary>
+	/// Converter to and from <see cref="StringId"/> instances.
+	/// </summary>
+	class TemplateIdConverter : StringIdConverter<TemplateId>
+	{
+		/// <inheritdoc/>
+		public override TemplateId FromStringId(StringId id) => new TemplateId(id);
+
+		/// <inheritdoc/>
+		public override StringId ToStringId(TemplateId value) => value.Id;
+	}
+
+	/// <summary>
+	/// Extension methods for stream id values
+	/// </summary>
+	public static class TemplateIdExtensions
+	{
+		/// <inheritdoc cref="ISpan.SetTag(System.String, System.String)"/>
+		public static ISpan SetTag(this ISpan span, string key, TemplateId value) => span.SetTag(key, value.Id.ToString());
+
+		/// <inheritdoc cref="ISpan.SetTag(System.String, System.String)"/>
+		public static ISpan SetTag(this ISpan span, string key, TemplateId[]? values) => span.SetTag(key, (values != null)? String.Join(',', values.Select(x => x.Id.ToString())) : null);
+	}
+}
diff --git a/Engine/Source/Programs/Horde/Horde.Build/Tasks/UpgradeTaskSource.cs b/Engine/Source/Programs/Horde/Horde.Build/Tasks/UpgradeTaskSource.cs
--- a/Engine/Source/Programs/Horde/Horde.Build/Tasks/UpgradeTaskSource.cs
+++ b/Engine/Source/Programs/Horde/Horde.Build/Tasks/UpgradeTaskSource.cs
@@ -18,7 +18,6 @@ using Microsoft.Extensions.Options;
 
 namespace Horde.Build.Tasks
 {
-	using AgentSoftwareChannelName = StringId<AgentSoftwareChannels>;
 	using JobId = ObjectId<IJob>;
 	using LeaseId = ObjectId<ILease>;
 	using LogId = ObjectId<ILogFile>;
diff --git a/Engine/Source/Programs/Horde/Horde.Build/Tools/ITool.cs b/Engine/Source/Programs/Horde/Horde.Build/Tools/ITool.cs
--- a/Engine/Source/Programs/Horde/Horde.Build/Tools/ITool.cs
+++ b/Engine/Source/Programs/Horde/Horde.Build/Tools/ITool.cs
@@ -10,7 +10,6 @@ using System.Threading.Tasks;
 
 namespace Horde.Build.Tools
 {
-	using ToolId = StringId<ITool>;
 	using ToolDeploymentId = ObjectId<IToolDeployment>;
 
 	/// <summary>
diff --git a/Engine/Source/Programs/Horde/Horde.Build/Tools/ToolCollection.cs b/Engine/Source/Programs/Horde/Horde.Build/Tools/ToolCollection.cs
--- a/Engine/Source/Programs/Horde/Horde.Build/Tools/ToolCollection.cs
+++ b/Engine/Source/Programs/Horde/Horde.Build/Tools/ToolCollection.cs
@@ -22,7 +22,6 @@ using System.Threading.Tasks;
 
 namespace Horde.Build.Tools
 {
-	using ToolId = StringId<ITool>;
 	using ToolDeploymentId = ObjectId<IToolDeployment>;
 
 	/// <summary>
diff --git a/Engine/Source/Programs/Horde/Horde.Build/Tools/ToolConfig.cs b/Engine/Source/Programs/Horde/Horde.Build/Tools/ToolConfig.cs
--- a/Engine/Source/Programs/Horde/Horde.Build/Tools/ToolConfig.cs
+++ b/Engine/Source/Programs/Horde/Horde.Build/Tools/ToolConfig.cs
@@ -11,8 +11,6 @@ using Horde.Build.Utilities;
 
 namespace Horde.Build.Tools
 {
-	using ToolId = StringId<ITool>;
-
 	/// <summary>
 	/// Options for configuring a tool
 	/// </summary>
diff --git a/Engine/Source/Programs/Horde/Horde.Build/Tools/ToolId.cs b/Engine/Source/Programs/Horde/Horde.Build/Tools/ToolId.cs
new file mode 100644
--- /dev/null
+++ b/Engine/Source/Programs/Horde/Horde.Build/Tools/ToolId.cs
@@ -0,0 +1,46 @@
+// Copyright Epic Games, Inc. All Rights Reserved.
+
+using System.ComponentModel;
+using EpicGames.Core;
+using EpicGames.Horde;
+using EpicGames.Serialization;
+using Horde.Build.Utilities;
+
+namespace Horde.Build.Tools
+{
+	/// <summary>
+	/// Identifier for a tool
+	/// </summary>
+	/// <param name="Id">Id to construct from</param>
+	[JsonSchemaString]
+	[TypeConverter(typeof(StringIdTypeConverter<ToolId, ToolIdConverter>))]
+	[StringIdConverter(typeof(ToolIdConverter))]
+	[CbConverter(typeof(StringIdCbConverter<ToolId, ToolIdConverter>))]
+	public record struct ToolId(StringId Id)
+	{
+		/// <summary>
+		/// Constructor
+		/// </summary>
+		public ToolId(string id) : this(new StringId(id))
+		{
+		}
+
+		/// <inheritdoc cref="StringId.IsEmpty"/>
+		public bool IsEmpty => Id.IsEmpty;
+
+		/// <inheritdoc/>
+		public override string ToString() => Id.ToString();
+	}
+
+	/// <summary>
+	/// Converter to and from <see cref="StringId"/> instances.
+	/// </summary>
+	class ToolIdConverter : StringIdConverter<ToolId>
+	{
+		/// <inheritdoc/>
+		public override ToolId FromStringId(StringId id) => new ToolId(id);
+
+		/// <inheritdoc/>
+		public override StringId ToStringId(ToolId value) => value.Id;
+	}
+}
diff --git a/Engine/Source/Programs/Horde/Horde.Build/Tools/ToolsController.cs b/Engine/Source/Programs/Horde/Horde.Build/Tools/ToolsController.cs
--- a/Engine/Source/Programs/Horde/Horde.Build/Tools/ToolsController.cs
+++ b/Engine/Source/Programs/Horde/Horde.Build/Tools/ToolsController.cs
@@ -19,7 +19,6 @@ using System.Threading.Tasks;
 
 namespace Horde.Build.Tools
 {
-	using ToolId = StringId<ITool>;
 	using ToolDeploymentId = ObjectId<IToolDeployment>;
 
 	/// <summary>
diff --git a/Engine/Source/Programs/Horde/Horde.Build/Utilities/HordeControllerBase.cs b/Engine/Source/Programs/Horde/Horde.Build/Utilities/HordeControllerBase.cs
--- a/Engine/Source/Programs/Horde/Horde.Build/Utilities/HordeControllerBase.cs
+++ b/Engine/Source/Programs/Horde/Horde.Build/Utilities/HordeControllerBase.cs
@@ -18,12 +18,8 @@ namespace Horde.Build.Utilities
 {
 	using JobId = ObjectId<IJob>;
 	using LeaseId = ObjectId<ILease>;
-	using ProjectId = StringId<ProjectConfig>;
-	using StreamId = StringId<IStream>;
-	using ToolId = StringId<ITool>;
 	using ToolDeploymentId = ObjectId<IToolDeployment>;
-	using TemplateId = StringId<ITemplateRef>;
-
+	
 	/// <summary>
 	/// Base class for Horde controllers
 	/// </summary>
diff --git a/Engine/Source/Programs/Horde/Horde.Build/Utilities/SingletonDocument.cs b/Engine/Source/Programs/Horde/Horde.Build/Utilities/SingletonDocument.cs
--- a/Engine/Source/Programs/Horde/Horde.Build/Utilities/SingletonDocument.cs
+++ b/Engine/Source/Programs/Horde/Horde.Build/Utilities/SingletonDocument.cs
@@ -1,15 +1,44 @@
 // Copyright Epic Games, Inc. All Rights Reserved.
 
 using System;
-using System.Reflection;
+using System.ComponentModel;
 using System.Threading.Tasks;
 using Horde.Build.Server;
 using MongoDB.Bson;
 using MongoDB.Bson.Serialization.Attributes;
+using EpicGames.Core;
+using EpicGames.Horde;
+using EpicGames.Serialization;
 
 namespace Horde.Build.Utilities
 {
-	using SingletonId = StringId<SingletonBase>;
+	/// <summary>
+	/// Identifier for a pool
+	/// </summary>
+	/// <param name="Id">Id to construct from</param>
+	[StringIdConverter(typeof(SingletonIdConverter))]
+	public record struct SingletonId(StringId Id)
+	{
+		/// <summary>
+		/// Constructor
+		/// </summary>
+		public SingletonId(string id) : this(new StringId(id)) { }
+
+		/// <inheritdoc/>
+		public override string ToString() => Id.ToString();
+	}
+
+	/// <summary>
+	/// Converter to and from <see cref="StringId"/> instances.
+	/// </summary>
+	class SingletonIdConverter : StringIdConverter<SingletonId>
+	{
+		/// <inheritdoc/>
+		public override SingletonId FromStringId(StringId id) => new SingletonId(id);
+
+		/// <inheritdoc/>
+		public override StringId ToStringId(SingletonId value) => value.Id;
+	}
 
 	/// <summary>
 	/// Base class for singleton documents
diff --git a/Engine/Source/Programs/Horde/Horde.Build/Utilities/StringId.cs b/Engine/Source/Programs/Horde/Horde.Build/Utilities/StringId.cs
deleted file mode 100644
diff --git a/Engine/Source/Programs/Horde/Horde.Build/Utilities/StringIdConverter.cs b/Engine/Source/Programs/Horde/Horde.Build/Utilities/StringIdConverter.cs
new file mode 100644
--- /dev/null
+++ b/Engine/Source/Programs/Horde/Horde.Build/Utilities/StringIdConverter.cs
@@ -0,0 +1,202 @@
+// Copyright Epic Games, Inc. All Rights Reserved.
+
+using System;
+using System.ComponentModel;
+using System.Globalization;
+using System.Reflection;
+using System.Text.Json;
+using System.Text.Json.Serialization;
+using EpicGames.Core;
+using EpicGames.Horde;
+using EpicGames.Serialization;
+using MongoDB.Bson.Serialization;
+using MongoDB.Bson.Serialization.Serializers;
+using ProtoBuf;
+
+namespace Horde.Build.Utilities
+{
+	/// <summary>
+	/// Base class for converting to and from types containing an object id. Useful pattern for reducing boilerplate with strongly typed records.
+	/// </summary>
+	/// <typeparam name="T"></typeparam>
+	abstract class StringIdConverter<T> where T : struct
+	{
+		/// <summary>
+		/// Converts a type to a <see cref="StringId"/>
+		/// </summary>
+		public abstract StringId ToStringId(T value);
+
+		/// <summary>
+		/// Constructs a type from an object id
+		/// </summary>
+		public abstract T FromStringId(StringId id);
+	}
+
+	/// <summary>
+	/// Attribute declaring an object id converter for a particular type
+	/// </summary>
+	[AttributeUsage(AttributeTargets.Struct)]
+	sealed class StringIdConverterAttribute : Attribute
+	{
+		/// <summary>
+		/// The converter type
+		/// </summary>
+		public Type ConverterType { get; }
+
+		/// <summary>
+		/// Constructor
+		/// </summary>
+		public StringIdConverterAttribute(Type converterType) => ConverterType = converterType;
+	}
+
+	/// <summary>
+	/// Converter to compact binary objects
+	/// </summary>
+	sealed class StringIdCbConverter<TValue, TConverter> : CbConverterBase<TValue> where TValue : struct where TConverter : StringIdConverter<TValue>, new()
+	{
+		readonly TConverter _converter = new TConverter();
+
+		/// <inheritdoc/>
+		public override TValue Read(CbField field)
+		{
+			return _converter.FromStringId(new StringId(field.AsString()));
+		}
+
+		/// <inheritdoc/>
+		public override void Write(CbWriter writer, TValue value)
+		{
+			writer.WriteStringValue(_converter.ToStringId(value).ToString());
+		}
+
+		/// <inheritdoc/>
+		public override void WriteNamed(CbWriter writer, Utf8String name, TValue value)
+		{
+			writer.WriteString(name, _converter.ToStringId(value).ToString());
+		}
+	}
+
+	/// <summary>
+	/// Class which serializes types with a <see cref="StringIdConverter{T}"/> to Json
+	/// </summary>
+	sealed class StringIdTypeConverter<TValue, TConverter> : TypeConverter where TValue : struct where TConverter : StringIdConverter<TValue>, new()
+	{
+		readonly TConverter _converter = new TConverter();
+
+		/// <inheritdoc/>
+		public override bool CanConvertFrom(ITypeDescriptorContext? context, Type sourceType)
+		{
+			return sourceType == typeof(string) || sourceType == typeof(StringId);
+		}
+
+		/// <inheritdoc/>
+		public override object? ConvertFrom(ITypeDescriptorContext? context, CultureInfo? culture, object value)
+		{
+			if (value is string str)
+			{
+				return _converter.FromStringId(new StringId(str));
+			}
+			if (value is StringId stringId)
+			{
+				return _converter.FromStringId(stringId);
+			}
+			return null;
+		}
+
+		/// <inheritdoc/>
+		public override bool CanConvertTo(ITypeDescriptorContext? context, Type? destinationType) => destinationType == typeof(string) || destinationType == typeof(StringId);
+
+		/// <inheritdoc/>
+		public override object? ConvertTo(ITypeDescriptorContext? context, CultureInfo? culture, object? value, Type destinationType)
+		{
+			if (destinationType == typeof(string))
+			{
+				return _converter.ToStringId((TValue)value!).ToString();
+			}
+			if (destinationType == typeof(StringId))
+			{
+				return _converter.ToStringId((TValue)value!);
+			}
+			return null;
+		}
+	}
+
+	/// <summary>
+	/// Class which serializes types with a <see cref="StringIdConverter{T}"/> to Json
+	/// </summary>
+	sealed class StringIdJsonConverter<TValue, TConverter> : JsonConverter<TValue> where TValue : struct where TConverter : StringIdConverter<TValue>, new()
+	{
+		readonly TConverter _converter = new TConverter();
+
+		/// <inheritdoc/>
+		public override TValue Read(ref Utf8JsonReader reader, Type typeToConvert, JsonSerializerOptions options) => _converter.FromStringId(new StringId(new Utf8String(reader.GetUtf8String().ToArray())));
+		
+		/// <inheritdoc/>
+		public override void Write(Utf8JsonWriter writer, TValue value, JsonSerializerOptions options) => writer.WriteStringValue(_converter.ToStringId(value).Span);
+	}
+
+	/// <summary>
+	/// Creates constructors for types with a <see cref="StringIdConverter{T}"/> to Json
+	/// </summary>
+	sealed class StringIdJsonConverterFactory : JsonConverterFactory
+	{
+		/// <inheritdoc/>
+		public override bool CanConvert(Type typeToConvert) => typeToConvert.GetCustomAttribute<StringIdConverterAttribute>() != null;
+
+		/// <inheritdoc/>
+		public override JsonConverter? CreateConverter(Type type, JsonSerializerOptions options)
+		{
+			StringIdConverterAttribute? attribute = type.GetCustomAttribute<StringIdConverterAttribute>();
+			if (attribute == null)
+			{
+				return null;
+			}
+			return (JsonConverter?)Activator.CreateInstance(typeof(StringIdJsonConverter<,>).MakeGenericType(type, attribute.ConverterType));
+		}
+	}
+
+	/// <summary>
+	/// Class which serializes object id types to BSON
+	/// </summary>
+	sealed class StringIdBsonSerializer<TValue, TConverter> : SerializerBase<TValue> where TValue : struct where TConverter : StringIdConverter<TValue>, new()
+	{
+		readonly TConverter _converter = new TConverter();
+
+		/// <inheritdoc/>
+		public override TValue Deserialize(BsonDeserializationContext context, BsonDeserializationArgs args) => _converter.FromStringId(new StringId(context.Reader.ReadString()));
+
+		/// <inheritdoc/>
+		public override void Serialize(BsonSerializationContext context, BsonSerializationArgs args, TValue value) => context.Writer.WriteString(_converter.ToStringId(value).ToString());
+	}
+
+	/// <summary>
+	/// Class which serializes object id types to BSON
+	/// </summary>
+	sealed class StringIdBsonSerializationProvider : BsonSerializationProviderBase
+	{
+		/// <inheritdoc/>
+		public override IBsonSerializer? GetSerializer(Type type, IBsonSerializerRegistry serializerRegistry)
+		{
+			StringIdConverterAttribute? attribute = type.GetCustomAttribute<StringIdConverterAttribute>();
+			if (attribute == null)
+			{
+				return null;
+			}
+			return (IBsonSerializer?)Activator.CreateInstance(typeof(StringIdBsonSerializer<,>).MakeGenericType(type, attribute.ConverterType));
+		}
+	}
+
+	/// <summary>
+	/// Surrogate type for serializing StringId types to ProtoBuf
+	/// </summary>
+	[ProtoContract]
+	struct StringIdProto<TValue, TConverter> where TValue : struct where TConverter : StringIdConverter<TValue>, new()
+	{
+		static readonly TConverter s_converter = new TConverter();
+
+		[ProtoMember(1)]
+		public string? Id { get; set; }
+
+		public static implicit operator TValue(StringIdProto<TValue, TConverter> source) => s_converter.FromStringId(new StringId(source.Id!));
+		public static implicit operator StringIdProto<TValue, TConverter>(TValue source) => new StringIdProto<TValue, TConverter> { Id = s_converter.ToStringId(source).ToString() };
+	}
+}
diff --git a/Engine/Source/Programs/Horde/Horde.Build/Utilities/TracingExtensions.cs b/Engine/Source/Programs/Horde/Horde.Build/Utilities/TracingExtensions.cs
--- a/Engine/Source/Programs/Horde/Horde.Build/Utilities/TracingExtensions.cs
+++ b/Engine/Source/Programs/Horde/Horde.Build/Utilities/TracingExtensions.cs
@@ -13,36 +13,6 @@ namespace Horde.Build.Utilities
 	/// </summary>
 	public static class OpenTracingSpanExtensions
 	{
-		/// <summary>Set a key:value tag on the span</summary>
-		/// <returns>This span instance, for chaining</returns>
-		public static ISpan SetTag<T>(this ISpan span, string key, StringId<T>? value)
-		{
-			if (value != null)
-			{
-				span.SetTag(key, value.ToString());
-			}
-			return span;
-		}
-		
-		/// <summary>Set a key:value tag on the span</summary>
-		/// <returns>This span instance, for chaining</returns>
-		public static ISpan SetTag<T>(this ISpan span, string key, StringId<T> value)
-		{
-			span.SetTag(key, value.ToString());
-			return span;
-		}
-
-		/// <summary>Set a key:value tag on the span</summary>
-		/// <returns>This span instance, for chaining</returns>
-		public static ISpan SetTag<T>(this ISpan span, string key, IEnumerable<StringId<T>>? value)
-		{
-			if (value != null)
-			{
-				span.SetTag(key, String.Join(',', value));
-			}
-			return span;
-		}
-
 		/// <summary>Set a key:value tag on the span</summary>
 		/// <returns>This span instance, for chaining</returns>
 		public static ISpan SetTag<T>(this ISpan span, string key, ObjectId<T>? value)
diff --git a/Engine/Source/Programs/Shared/EpicGames.Horde/StringId.cs b/Engine/Source/Programs/Shared/EpicGames.Horde/StringId.cs
--- a/Engine/Source/Programs/Shared/EpicGames.Horde/StringId.cs
+++ b/Engine/Source/Programs/Shared/EpicGames.Horde/StringId.cs
@@ -1,6 +1,8 @@
 // Copyright Epic Games, Inc. All Rights Reserved.
 
 using System;
+using System.Diagnostics.CodeAnalysis;
+using System.Text;
 using EpicGames.Core;
 
 namespace EpicGames.Horde
@@ -10,6 +12,17 @@ namespace EpicGames.Horde
 	/// </summary>
 	public struct StringId : IEquatable<StringId>
 	{
+		/// <summary>
+		/// Enum used to disable validation on string arguments
+		/// </summary>
+		public enum Validate
+		{ 
+			/// <summary>
+			/// No validation required
+			/// </summary>
+			None,
+		};
+
 		/// <summary>
 		/// The text representing this id
 		/// </summary>
@@ -34,11 +47,53 @@ namespace EpicGames.Horde
 			Text = ValidateArgument(text, nameof(text));
 		}
 
+		/// <summary>
+		/// Constructor
+		/// </summary>
+		/// <param name="text">Unique id for the string</param>
+		/// <param name="validate">Argument used for overload resolution for pre-validated strings</param>
+		[SuppressMessage("Style", "IDE0060:Remove unused parameter")]
+		public StringId(Utf8String text, Validate validate)
+		{
+			Text = text;
+		}
+
 		/// <summary>
 		/// Checks whether this StringId is set
 		/// </summary>
 		public bool IsEmpty => Text.IsEmpty;
 
+		/// <summary>
+		/// Generates a new string id from the given text
+		/// </summary>
+		/// <param name="text">Text to generate from</param>
+		/// <returns>New string id</returns>
+		public static StringId Sanitize(string text)
+		{
+			StringBuilder result = new StringBuilder();
+			for (int idx = 0; idx < text.Length; idx++)
+			{
+				char character = (char)text[idx];
+				if (character >= 'A' && character <= 'Z')
+				{
+					result.Append((char)('a' + (character - 'A')));
+				}
+				else if (IsValidCharacter(character))
+				{
+					result.Append(character);
+				}
+				else if (result.Length > 0 && result[^1] != '-')
+				{
+					result.Append('-');
+				}
+			}
+			while (result.Length > 0 && result[^1] == '-')
+			{
+				result.Remove(result.Length - 1, 1);
+			}
+			return new StringId(result.ToString(), Validate.None);
+		}
+
 		/// <summary>
 		/// Validates the given string as a StringId, normalizing it if necessary.
 		/// </summary>
@@ -47,11 +102,6 @@ namespace EpicGames.Horde
 		/// <returns></returns>
 		public static Utf8String ValidateArgument(Utf8String text, string paramName)
 		{
-			if (text.Length == 0)
-			{
-				throw new ArgumentException("String id may not be empty", paramName);
-			}
-
 			const int MaxLength = 64;
 			if (text.Length > MaxLength)
 			{
