Horde: Remove type aliases for ObjectId<> types. Use concrete record types instead.

diff --git a/Engine/Source/Programs/Horde/Horde.Build.Tests/DebugTest.cs b/Engine/Source/Programs/Horde/Horde.Build.Tests/DebugTest.cs
--- a/Engine/Source/Programs/Horde/Horde.Build.Tests/DebugTest.cs
+++ b/Engine/Source/Programs/Horde/Horde.Build.Tests/DebugTest.cs
@@ -6,8 +6,6 @@ using Microsoft.Extensions.Options;
 
 namespace Horde.Build.Tests
 {
-	using JobId = ObjectId<IJob>;
-
     /// <summary>
     /// Tests aimed to be run from an IDE while debugging against real databases
     /// Uncomment the [Ignore] tag to run
diff --git a/Engine/Source/Programs/Horde/Horde.Build.Tests/Fixture.cs b/Engine/Source/Programs/Horde/Horde.Build.Tests/Fixture.cs
--- a/Engine/Source/Programs/Horde/Horde.Build.Tests/Fixture.cs
+++ b/Engine/Source/Programs/Horde/Horde.Build.Tests/Fixture.cs
@@ -20,8 +20,6 @@ using HordeCommon;
 
 namespace Horde.Build.Tests
 {
-	using JobId = ObjectId<IJob>;
-
 	public class Fixture
 	{
 		public IJob Job1 { get; private set; } = null!;
@@ -90,7 +88,7 @@ namespace Horde.Build.Tests
 			StreamConfig = streamConfig;
 
 			Job1 = await jobService.CreateJobAsync(
-				jobId: new JobId("5f283932841e7fdbcafb6ab5"),
+				jobId: JobId.Parse("5f283932841e7fdbcafb6ab5"),
 				streamConfig: streamConfig,
 				templateRefId: TemplateRefId1,
 				templateHash: Template.Hash,
@@ -103,7 +101,7 @@ namespace Horde.Build.Tests
 			Job1 = (await jobService.GetJobAsync(Job1.Id))!;
 
 			Job2 = await jobService.CreateJobAsync(
-				jobId: new JobId("5f69ea1b68423e921b035106"),
+				jobId: JobId.Parse("5f69ea1b68423e921b035106"),
 				streamConfig: streamConfig,
 				templateRefId: new TemplateId("template-id-1"),
 				templateHash: ContentHash.MD5("made-up-template-hash"),
diff --git a/Engine/Source/Programs/Horde/Horde.Build.Tests/Fleet/JobQueueStrategyTest.cs b/Engine/Source/Programs/Horde/Horde.Build.Tests/Fleet/JobQueueStrategyTest.cs
--- a/Engine/Source/Programs/Horde/Horde.Build.Tests/Fleet/JobQueueStrategyTest.cs
+++ b/Engine/Source/Programs/Horde/Horde.Build.Tests/Fleet/JobQueueStrategyTest.cs
@@ -150,7 +150,7 @@ namespace Horde.Build.Tests.Fleet
 			CreateJobOptions options = new CreateJobOptions();
 			options.Arguments.Add($"-Target={nodeNameToExecute}");
 
-			IJob job = await JobCollection.AddAsync(ObjectId<IJob>.GenerateNewId(), streamId,
+			IJob job = await JobCollection.AddAsync(JobId.GenerateNewId(), streamId,
 				new TemplateId("bogusTemplateRefId"), ContentHash.Empty, graph, "bogusJobName",
 				1000, 1000, options);
 
diff --git a/Engine/Source/Programs/Horde/Horde.Build.Tests/Fleet/LeaseUtilizationStrategyTest.cs b/Engine/Source/Programs/Horde/Horde.Build.Tests/Fleet/LeaseUtilizationStrategyTest.cs
--- a/Engine/Source/Programs/Horde/Horde.Build.Tests/Fleet/LeaseUtilizationStrategyTest.cs
+++ b/Engine/Source/Programs/Horde/Horde.Build.Tests/Fleet/LeaseUtilizationStrategyTest.cs
@@ -108,7 +108,7 @@ namespace Horde.Build.Tests.Fleet
 			placeholderJobTask.JobName = "placeholderJobName";
 			byte[] payload = Any.Pack(placeholderJobTask).ToByteArray();
 
-			ILease lease = await LeaseCollection.AddAsync(ObjectId<ILease>.GenerateNewId(), "placeholderLease", agent.Id, agent.SessionId!.Value, new StreamId("placeholderStream"), pool.Id, null, startTime, payload);
+			ILease lease = await LeaseCollection.AddAsync(LeaseId.GenerateNewId(), "placeholderLease", agent.Id, agent.SessionId!.Value, new StreamId("placeholderStream"), pool.Id, null, startTime, payload);
 			bool wasModified = await LeaseCollection.TrySetOutcomeAsync(lease.Id, startTime + duration, LeaseOutcome.Success, null);
 			Assert.IsTrue(wasModified);
 			return lease;
diff --git a/Engine/Source/Programs/Horde/Horde.Build.Tests/IssueServiceTests.cs b/Engine/Source/Programs/Horde/Horde.Build.Tests/IssueServiceTests.cs
--- a/Engine/Source/Programs/Horde/Horde.Build.Tests/IssueServiceTests.cs
+++ b/Engine/Source/Programs/Horde/Horde.Build.Tests/IssueServiceTests.cs
@@ -33,10 +33,6 @@ using System.Threading;
 
 namespace Horde.Build.Tests
 {
-	using JobId = ObjectId<IJob>;
-	using LogId = ObjectId<ILogFile>;
-	using UserId = ObjectId<IUser>;
-
 	[TestClass]
 	public class IssueServiceTests : TestSetup
 	{
diff --git a/Engine/Source/Programs/Horde/Horde.Build.Tests/JobCollectionTests.cs b/Engine/Source/Programs/Horde/Horde.Build.Tests/JobCollectionTests.cs
--- a/Engine/Source/Programs/Horde/Horde.Build.Tests/JobCollectionTests.cs
+++ b/Engine/Source/Programs/Horde/Horde.Build.Tests/JobCollectionTests.cs
@@ -22,10 +22,6 @@ using Horde.Build.Agents.Sessions;
 
 namespace Horde.Build.Tests
 {
-	using JobId = ObjectId<IJob>;
-	using LeaseId = ObjectId<ILease>;
-	using LogId = ObjectId<ILogFile>;
-
 	[TestClass]
 	public class JobCollectionTests : TestSetup
 	{
@@ -118,11 +114,11 @@ namespace Horde.Build.Tests
 		{
 			Fixture fixture = await CreateFixtureAsync();
 
-			ObjectId<ISession> sessionId1 = new (ObjectId.GenerateNewId());
+			SessionId sessionId1 = SessionId.GenerateNewId();
 			await JobCollection.TryAssignLeaseAsync(fixture.Job1, 0, new PoolId("foo"), fixture.Agent1.Id,
 				sessionId1, LeaseId.GenerateNewId(), LogId.GenerateNewId());
 			
-			ObjectId<ISession> sessionId2 = new (ObjectId.GenerateNewId());
+			SessionId sessionId2 = SessionId.GenerateNewId();
 			IJob job = (await JobCollection.GetAsync(fixture.Job1.Id))!;
 			await JobCollection.TryAssignLeaseAsync(job, 0, new PoolId("foo"), fixture.Agent1.Id,
 				sessionId2, LeaseId.GenerateNewId(), LogId.GenerateNewId());
diff --git a/Engine/Source/Programs/Horde/Horde.Build.Tests/JobServiceTests.cs b/Engine/Source/Programs/Horde/Horde.Build.Tests/JobServiceTests.cs
--- a/Engine/Source/Programs/Horde/Horde.Build.Tests/JobServiceTests.cs
+++ b/Engine/Source/Programs/Horde/Horde.Build.Tests/JobServiceTests.cs
@@ -26,8 +26,6 @@ using Microsoft.VisualStudio.TestTools.UnitTesting;
 
 namespace Horde.Build.Tests
 {
-	using JobId = ObjectId<IJob>;
-	using LogId = ObjectId<ILogFile>;
 	using ProjectId = Projects.ProjectId;
 
 	[TestClass]
diff --git a/Engine/Source/Programs/Horde/Horde.Build.Tests/JobTaskSourceTests.cs b/Engine/Source/Programs/Horde/Horde.Build.Tests/JobTaskSourceTests.cs
--- a/Engine/Source/Programs/Horde/Horde.Build.Tests/JobTaskSourceTests.cs
+++ b/Engine/Source/Programs/Horde/Horde.Build.Tests/JobTaskSourceTests.cs
@@ -17,9 +17,6 @@ using Horde.Build.Streams;
 
 namespace Horde.Build.Tests
 {
-	using JobId = ObjectId<IJob>;
-	using UserId = ObjectId<IUser>;
-
 	[TestClass]
 	public class JobTaskSourceTests : TestSetup
 	{
diff --git a/Engine/Source/Programs/Horde/Horde.Build.Tests/LeasesControllerDbTest.cs b/Engine/Source/Programs/Horde/Horde.Build.Tests/LeasesControllerDbTest.cs
--- a/Engine/Source/Programs/Horde/Horde.Build.Tests/LeasesControllerDbTest.cs
+++ b/Engine/Source/Programs/Horde/Horde.Build.Tests/LeasesControllerDbTest.cs
@@ -37,8 +37,8 @@ namespace Horde.Build.Tests
 
 		private async Task<ILease> CreateLease(DateTime startTime, TimeSpan duration)
 		{
-			ObjectId<ILease> id = ObjectId<ILease>.GenerateNewId();
-			ObjectId<ISession> sessionId = ObjectId<ISession>.GenerateNewId();
+			LeaseId id = LeaseId.GenerateNewId();
+			SessionId sessionId = SessionId.GenerateNewId();
 			ILease lease = await LeaseCollection.AddAsync(id, "myLease", new AgentId("agent-1"), sessionId, null, null, null, startTime, Array.Empty<byte>());
 			await LeaseCollection.TrySetOutcomeAsync(id, startTime + duration, LeaseOutcome.Success, null);
 			return lease;
diff --git a/Engine/Source/Programs/Horde/Horde.Build.Tests/LogBuilderTests.cs b/Engine/Source/Programs/Horde/Horde.Build.Tests/LogBuilderTests.cs
--- a/Engine/Source/Programs/Horde/Horde.Build.Tests/LogBuilderTests.cs
+++ b/Engine/Source/Programs/Horde/Horde.Build.Tests/LogBuilderTests.cs
@@ -14,8 +14,6 @@ using Microsoft.VisualStudio.TestTools.UnitTesting;
 
 namespace Horde.Build.Tests
 {
-	using LogId = ObjectId<ILogFile>;
-
 	[TestClass]
 	public class LogBuilderTests : DatabaseIntegrationTest
 	{
diff --git a/Engine/Source/Programs/Horde/Horde.Build.Tests/LogFileServiceTests.cs b/Engine/Source/Programs/Horde/Horde.Build.Tests/LogFileServiceTests.cs
--- a/Engine/Source/Programs/Horde/Horde.Build.Tests/LogFileServiceTests.cs
+++ b/Engine/Source/Programs/Horde/Horde.Build.Tests/LogFileServiceTests.cs
@@ -26,9 +26,6 @@ using MongoDB.Bson;
 
 namespace Horde.Build.Tests
 {
-	using JobId = ObjectId<IJob>;
-	using LogId = ObjectId<ILogFile>;
-
 	[TestClass]
     public class LogFileServiceTest : DatabaseIntegrationTest
     {
@@ -256,8 +253,8 @@ namespace Horde.Build.Tests
 
 			// Will implicitly test GetLogFileAsync(), AddCachedLogFile()
 			JobId jobId = JobId.GenerateNewId();
-            ObjectId sessionId = ObjectId.GenerateNewId();
-            ILogFile a = await _logFileService.CreateLogFileAsync(jobId, new ObjectId<ISession>(sessionId), LogType.Text, useNewStorageBackend: false, logId: null, CancellationToken.None);
+            SessionId sessionId = SessionId.GenerateNewId();
+            ILogFile a = await _logFileService.CreateLogFileAsync(jobId, sessionId, LogType.Text, useNewStorageBackend: false, logId: null, CancellationToken.None);
             ILogFile b = (await _logFileService.GetCachedLogFileAsync(a.Id, CancellationToken.None))!;
             Assert.AreEqual(a.JobId, b.JobId);
             Assert.AreEqual(a.SessionId, b.SessionId);
@@ -266,7 +263,7 @@ namespace Horde.Build.Tests
             ILogFile? notFound = await _logFileService.GetCachedLogFileAsync(LogId.GenerateNewId(), CancellationToken.None);
             Assert.IsNull(notFound);
 
-            await _logFileService.CreateLogFileAsync(JobId.GenerateNewId(), new ObjectId<ISession>(ObjectId.GenerateNewId()), LogType.Text, useNewStorageBackend: false, logId: null, cancellationToken: CancellationToken.None);
+            await _logFileService.CreateLogFileAsync(JobId.GenerateNewId(), SessionId.GenerateNewId(), LogType.Text, useNewStorageBackend: false, logId: null, cancellationToken: CancellationToken.None);
             Assert.AreEqual(2, (await _logFileService.GetLogFilesAsync()).Count);
         }
 
@@ -274,8 +271,8 @@ namespace Horde.Build.Tests
         public async Task AuthorizeForSession()
         {
 			JobId jobId = JobId.GenerateNewId();
-            ObjectId sessionId = ObjectId.GenerateNewId();
-            ILogFile logFile = await _logFileService.CreateLogFileAsync(jobId, new ObjectId<ISession>(sessionId), LogType.Text, useNewStorageBackend: false, logId: null, cancellationToken: CancellationToken.None);
+            SessionId sessionId = SessionId.GenerateNewId();
+            ILogFile logFile = await _logFileService.CreateLogFileAsync(jobId, sessionId, LogType.Text, useNewStorageBackend: false, logId: null, cancellationToken: CancellationToken.None);
             ILogFile logFileNoSession = await _logFileService.CreateLogFileAsync(jobId, null, LogType.Text, useNewStorageBackend: false, logId: null, cancellationToken: CancellationToken.None);
 
 			ClaimsPrincipal hasClaim = new ClaimsPrincipal(new ClaimsIdentity(new List<Claim>
diff --git a/Engine/Source/Programs/Horde/Horde.Build.Tests/LogIndexingTests.cs b/Engine/Source/Programs/Horde/Horde.Build.Tests/LogIndexingTests.cs
--- a/Engine/Source/Programs/Horde/Horde.Build.Tests/LogIndexingTests.cs
+++ b/Engine/Source/Programs/Horde/Horde.Build.Tests/LogIndexingTests.cs
@@ -23,8 +23,6 @@ using EpicGames.Horde.Logs;
 
 namespace Horde.Build.Tests
 {
-	using JobId = ObjectId<IJob>;
-
 	[TestClass]
 	public class LogIndexingTests : DatabaseIntegrationTest
 	{
diff --git a/Engine/Source/Programs/Horde/Horde.Build.Tests/LogTailServiceTests.cs b/Engine/Source/Programs/Horde/Horde.Build.Tests/LogTailServiceTests.cs
--- a/Engine/Source/Programs/Horde/Horde.Build.Tests/LogTailServiceTests.cs
+++ b/Engine/Source/Programs/Horde/Horde.Build.Tests/LogTailServiceTests.cs
@@ -17,8 +17,6 @@ using Microsoft.VisualStudio.TestTools.UnitTesting;
 
 namespace Horde.Build.Tests
 {
-	using LogId = ObjectId<ILogFile>;
-
 	[TestClass]
 	public class LogTailServiceTests : TestSetup
 	{
diff --git a/Engine/Source/Programs/Horde/Horde.Build.Tests/NotificationServiceTests.cs b/Engine/Source/Programs/Horde/Horde.Build.Tests/NotificationServiceTests.cs
--- a/Engine/Source/Programs/Horde/Horde.Build.Tests/NotificationServiceTests.cs
+++ b/Engine/Source/Programs/Horde/Horde.Build.Tests/NotificationServiceTests.cs
@@ -21,8 +21,6 @@ using Microsoft.Extensions.Caching.Memory;
 
 namespace Horde.Build.Tests
 {
-	using JobId = ObjectId<IJob>;
-	
 	public class FakeNotificationSink : INotificationSink
 	{
 		public List<JobScheduledNotification> JobScheduledNotifications { get; } = new();
@@ -59,7 +57,7 @@ namespace Horde.Build.Tests
 
 		public static IJob CreateJob(StreamId streamId, int change, string name, IGraph graph)
 		{
-			JobId jobId = new JobId("5ec16da1774cb4000107c2c1");
+			JobId jobId = JobId.Parse("5ec16da1774cb4000107c2c1");
 
 			List<IJobStepBatch> batches = new List<IJobStepBatch>();
 			for (int groupIdx = 0; groupIdx < graph.Groups.Count; groupIdx++)
diff --git a/Engine/Source/Programs/Horde/Horde.Build.Tests/RpcServiceTest.cs b/Engine/Source/Programs/Horde/Horde.Build.Tests/RpcServiceTest.cs
--- a/Engine/Source/Programs/Horde/Horde.Build.Tests/RpcServiceTest.cs
+++ b/Engine/Source/Programs/Horde/Horde.Build.Tests/RpcServiceTest.cs
@@ -27,12 +27,10 @@ using ISession = Microsoft.AspNetCore.Http.ISession;
 using Horde.Build.Jobs.Artifacts;
 using Horde.Build.Agents;
 using HordeCommon.Rpc.Messages;
+using Horde.Build.Agents.Sessions;
 
 namespace Horde.Build.Tests
 {
-	using LeaseId = ObjectId<ILease>;
-	using LogId = ObjectId<ILogFile>;
-	
 	public class AppLifetimeStub : IHostApplicationLifetime
 	{
 		public CancellationToken ApplicationStarted { get; }
@@ -430,7 +428,7 @@ namespace Horde.Build.Tests
 		{
 			Fixture fixture = await CreateFixtureAsync();
 
-			ObjectId sessionId = ObjectId.GenerateNewId();
+			SessionId sessionId = SessionId.GenerateNewId();
 			ServerCallContext context = new ServerCallContextStub(new ClaimsPrincipal(new ClaimsIdentity(new List<Claim>
 			{
 				new Claim(HordeClaimTypes.Role, "app-horde-admins"),
@@ -452,7 +450,7 @@ namespace Horde.Build.Tests
 
 			// Set the session ID on the job batch to pass auth later
 			Deref(await JobCollection.TryAssignLeaseAsync(fixture.Job1, 0, new PoolId("foo"),
-				new AgentId("test"), new ObjectId<Horde.Build.Agents.Sessions.ISession>(sessionId),
+				new AgentId("test"), sessionId,
 				LeaseId.GenerateNewId(), LogId.GenerateNewId()));
 /*
 			TestAsyncStreamReader<UploadArtifactRequest> RequestStream = new TestAsyncStreamReader<UploadArtifactRequest>(Context);
diff --git a/Engine/Source/Programs/Horde/Horde.Build.Tests/SchedulerTests.cs b/Engine/Source/Programs/Horde/Horde.Build.Tests/SchedulerTests.cs
--- a/Engine/Source/Programs/Horde/Horde.Build.Tests/SchedulerTests.cs
+++ b/Engine/Source/Programs/Horde/Horde.Build.Tests/SchedulerTests.cs
@@ -21,9 +21,6 @@ using Horde.Build.Perforce;
 
 namespace Horde.Build.Tests
 {
-	using JobId = ObjectId<IJob>;
-	using LogId = ObjectId<ILogFile>;
-
 	[TestClass]
     public class SchedulerTests : TestSetup
 	{
diff --git a/Engine/Source/Programs/Horde/Horde.Build.Tests/Stubs/Collections/JobStepRefCollectionStub.cs b/Engine/Source/Programs/Horde/Horde.Build.Tests/Stubs/Collections/JobStepRefCollectionStub.cs
--- a/Engine/Source/Programs/Horde/Horde.Build.Tests/Stubs/Collections/JobStepRefCollectionStub.cs
+++ b/Engine/Source/Programs/Horde/Horde.Build.Tests/Stubs/Collections/JobStepRefCollectionStub.cs
@@ -13,9 +13,6 @@ using HordeCommon;
 
 namespace Horde.Build.Tests.Stubs.Collections
 {
-	using JobId = ObjectId<IJob>;
-	using LogId = ObjectId<ILogFile>;
-
 	class JobStepRefStub : IJobStepRef
 	{
 		public JobStepRefId Id { get; set; }
diff --git a/Engine/Source/Programs/Horde/Horde.Build.Tests/Stubs/Services/PerforceServiceStub.cs b/Engine/Source/Programs/Horde/Horde.Build.Tests/Stubs/Services/PerforceServiceStub.cs
--- a/Engine/Source/Programs/Horde/Horde.Build.Tests/Stubs/Services/PerforceServiceStub.cs
+++ b/Engine/Source/Programs/Horde/Horde.Build.Tests/Stubs/Services/PerforceServiceStub.cs
@@ -16,8 +16,6 @@ using Microsoft.VisualStudio.TestTools.UnitTesting;
 
 namespace Horde.Build.Tests.Stubs.Services
 {
-	using UserId = ObjectId<IUser>;
-
 	class PerforceServiceStub : IPerforceService
 	{
 		class User : IUser
diff --git a/Engine/Source/Programs/Horde/Horde.Build.Tests/ToolTests.cs b/Engine/Source/Programs/Horde/Horde.Build.Tests/ToolTests.cs
--- a/Engine/Source/Programs/Horde/Horde.Build.Tests/ToolTests.cs
+++ b/Engine/Source/Programs/Horde/Horde.Build.Tests/ToolTests.cs
@@ -14,8 +14,6 @@ using Microsoft.VisualStudio.TestTools.UnitTesting;
 
 namespace Horde.Build.Tests
 {
-	using ToolDeploymentId = ObjectId<IToolDeployment>;
-
 	[TestClass]
     public class ToolTests : TestSetup
     {
diff --git a/Engine/Source/Programs/Horde/Horde.Build/Acls/AclService.cs b/Engine/Source/Programs/Horde/Horde.Build/Acls/AclService.cs
--- a/Engine/Source/Programs/Horde/Horde.Build/Acls/AclService.cs
+++ b/Engine/Source/Programs/Horde/Horde.Build/Acls/AclService.cs
@@ -15,8 +15,6 @@ using MongoDB.Driver;
 
 namespace Horde.Build.Acls
 {
-	using SessionId = ObjectId<ISession>;
-
 	/// <summary>
 	/// Wraps functionality for manipulating permissions
 	/// </summary>
diff --git a/Engine/Source/Programs/Horde/Horde.Build/Agents/AgentCollection.cs b/Engine/Source/Programs/Horde/Horde.Build/Agents/AgentCollection.cs
--- a/Engine/Source/Programs/Horde/Horde.Build/Agents/AgentCollection.cs
+++ b/Engine/Source/Programs/Horde/Horde.Build/Agents/AgentCollection.cs
@@ -20,8 +20,6 @@ using MongoDB.Driver;
 
 namespace Horde.Build.Agents
 {
-	using SessionId = ObjectId<ISession>;
-
 	/// <summary>
 	/// Collection of agent documents
 	/// </summary>
diff --git a/Engine/Source/Programs/Horde/Horde.Build/Agents/AgentService.cs b/Engine/Source/Programs/Horde/Horde.Build/Agents/AgentService.cs
--- a/Engine/Source/Programs/Horde/Horde.Build/Agents/AgentService.cs
+++ b/Engine/Source/Programs/Horde/Horde.Build/Agents/AgentService.cs
@@ -30,9 +30,6 @@ using StatsdClient;
 
 namespace Horde.Build.Agents
 {
-	using LeaseId = ObjectId<ILease>;
-	using SessionId = ObjectId<ISession>;
-
 	/// <summary>
 	/// Singleton used to store agent costs
 	/// </summary>
@@ -657,7 +654,7 @@ namespace Horde.Build.Agents
 				List<AgentLease> leases = new List<AgentLease>(agent.Leases);
 
 				// Remove any completed leases from the agent
-				Dictionary<LeaseId, HordeCommon.Rpc.Messages.Lease> leaseIdToNewState = newLeases.ToDictionary(x => new LeaseId(x.Id), x => x);
+				Dictionary<LeaseId, HordeCommon.Rpc.Messages.Lease> leaseIdToNewState = newLeases.ToDictionary(x => LeaseId.Parse(x.Id), x => x);
 				for (int idx = 0; idx < leases.Count; idx++)
 				{
 					AgentLease lease = leases[idx];
diff --git a/Engine/Source/Programs/Horde/Horde.Build/Agents/AgentsController.cs b/Engine/Source/Programs/Horde/Horde.Build/Agents/AgentsController.cs
--- a/Engine/Source/Programs/Horde/Horde.Build/Agents/AgentsController.cs
+++ b/Engine/Source/Programs/Horde/Horde.Build/Agents/AgentsController.cs
@@ -19,9 +19,6 @@ using Microsoft.Extensions.Options;
 
 namespace Horde.Build.Agents
 {
-	using LeaseId = ObjectId<ILease>;
-	using SessionId = ObjectId<ISession>;
-
 	/// <summary>
 	/// Controller for the /api/v1/agents endpoint
 	/// </summary>
diff --git a/Engine/Source/Programs/Horde/Horde.Build/Agents/ConformList.cs b/Engine/Source/Programs/Horde/Horde.Build/Agents/ConformList.cs
--- a/Engine/Source/Programs/Horde/Horde.Build/Agents/ConformList.cs
+++ b/Engine/Source/Programs/Horde/Horde.Build/Agents/ConformList.cs
@@ -10,8 +10,6 @@ using MongoDB.Bson.Serialization.Attributes;
 
 namespace Horde.Build.Agents
 {
-	using LeaseId = ObjectId<ILease>;
-
 	/// <summary>
 	/// Entry to the conform list
 	/// </summary>
diff --git a/Engine/Source/Programs/Horde/Horde.Build/Agents/IAgent.cs b/Engine/Source/Programs/Horde/Horde.Build/Agents/IAgent.cs
--- a/Engine/Source/Programs/Horde/Horde.Build/Agents/IAgent.cs
+++ b/Engine/Source/Programs/Horde/Horde.Build/Agents/IAgent.cs
@@ -27,10 +27,6 @@ using MongoDB.Bson.Serialization.Attributes;
 
 namespace Horde.Build.Agents
 {
-	using LeaseId = ObjectId<ILease>;
-	using LogId = ObjectId<ILogFile>;
-	using SessionId = ObjectId<ISession>;
-	
 	/// <summary>
 	/// Information about a workspace synced to an agent
 	/// </summary>
diff --git a/Engine/Source/Programs/Horde/Horde.Build/Agents/Leases/ILease.cs b/Engine/Source/Programs/Horde/Horde.Build/Agents/Leases/ILease.cs
--- a/Engine/Source/Programs/Horde/Horde.Build/Agents/Leases/ILease.cs
+++ b/Engine/Source/Programs/Horde/Horde.Build/Agents/Leases/ILease.cs
@@ -12,10 +12,6 @@ using HordeCommon;
 
 namespace Horde.Build.Agents.Leases
 {
-	using LeaseId = ObjectId<ILease>;
-	using LogId = ObjectId<ILogFile>;
-	using SessionId = ObjectId<ISession>;
-	
 	/// <summary>
 	/// Document describing a lease. This exists to permanently record a lease; the agent object tracks internal state of any active leases through AgentLease objects.
 	/// </summary>
diff --git a/Engine/Source/Programs/Horde/Horde.Build/Agents/Leases/LeaseId.cs b/Engine/Source/Programs/Horde/Horde.Build/Agents/Leases/LeaseId.cs
new file mode 100644
--- /dev/null
+++ b/Engine/Source/Programs/Horde/Horde.Build/Agents/Leases/LeaseId.cs
@@ -0,0 +1,42 @@
+// Copyright Epic Games, Inc. All Rights Reserved.
+
+using System.ComponentModel;
+using EpicGames.Core;
+using Horde.Build.Utilities;
+using MongoDB.Bson;
+
+namespace Horde.Build.Agents.Leases
+{
+	/// <summary>
+	/// Identifier for a user
+	/// </summary>
+	/// <param name="Id">Id to construct from</param>
+	[JsonSchemaString]
+	[TypeConverter(typeof(ObjectIdTypeConverter<LeaseId, LeaseIdConverter>))]
+	[ObjectIdConverter(typeof(LeaseIdConverter))]
+	public record struct LeaseId(ObjectId Id)
+	{
+		/// <summary>
+		/// Creates a new <see cref="LeaseId"/>
+		/// </summary>
+		public static LeaseId GenerateNewId() => new LeaseId(ObjectId.GenerateNewId());
+
+		/// <inheritdoc cref="ObjectId.Parse(System.String)"/>
+		public static LeaseId Parse(string text) => new LeaseId(ObjectId.Parse(text));
+
+		/// <inheritdoc/>
+		public override string ToString() => Id.ToString();
+	}
+
+	/// <summary>
+	/// Converter to and from <see cref="ObjectId"/> instances.
+	/// </summary>
+	class LeaseIdConverter : ObjectIdConverter<LeaseId>
+	{
+		/// <inheritdoc/>
+		public override LeaseId FromObjectId(ObjectId id) => new LeaseId(id);
+
+		/// <inheritdoc/>
+		public override ObjectId ToObjectId(LeaseId value) => value.Id;
+	}
+}
diff --git a/Engine/Source/Programs/Horde/Horde.Build/Agents/Leases/LeasesController.cs b/Engine/Source/Programs/Horde/Horde.Build/Agents/Leases/LeasesController.cs
--- a/Engine/Source/Programs/Horde/Horde.Build/Agents/Leases/LeasesController.cs
+++ b/Engine/Source/Programs/Horde/Horde.Build/Agents/Leases/LeasesController.cs
@@ -16,9 +16,6 @@ using OpenTracing.Util;
 
 namespace Horde.Build.Agents.Leases
 {
-	using LeaseId = ObjectId<ILease>;
-	using SessionId = ObjectId<ISession>;
-
 	/// <summary>
 	/// Controller for the /api/v1/leases endpoint
 	/// </summary>
diff --git a/Engine/Source/Programs/Horde/Horde.Build/Agents/Sessions/SessionId.cs b/Engine/Source/Programs/Horde/Horde.Build/Agents/Sessions/SessionId.cs
new file mode 100644
--- /dev/null
+++ b/Engine/Source/Programs/Horde/Horde.Build/Agents/Sessions/SessionId.cs
@@ -0,0 +1,62 @@
+// Copyright Epic Games, Inc. All Rights Reserved.
+
+using System.ComponentModel;
+using EpicGames.Core;
+using Horde.Build.Utilities;
+using MongoDB.Bson;
+
+namespace Horde.Build.Agents.Sessions
+{
+	/// <summary>
+	/// Identifier for a session
+	/// </summary>
+	/// <param name="Id">Id to construct from</param>
+	[JsonSchemaString]
+	[TypeConverter(typeof(ObjectIdTypeConverter<SessionId, SessionIdConverter>))]
+	[ObjectIdConverter(typeof(SessionIdConverter))]
+	public record struct SessionId(ObjectId Id)
+	{
+		/// <summary>
+		/// Default empty value for session id
+		/// </summary>
+		public static SessionId Empty { get; } = default;
+
+		/// <summary>
+		/// Creates a new <see cref="SessionId"/>
+		/// </summary>
+		public static SessionId GenerateNewId() => new SessionId(ObjectId.GenerateNewId());
+
+		/// <inheritdoc cref="ObjectId.Parse(System.String)"/>
+		public static SessionId Parse(string text) => new SessionId(ObjectId.Parse(text));
+
+		/// <inheritdoc cref="ObjectId.TryParse(System.String, out ObjectId)"/>
+		public static bool TryParse(string text, out SessionId id)
+		{
+			if (ObjectId.TryParse(text, out ObjectId objectId))
+			{
+				id = new SessionId(objectId);
+				return true;
+			}
+			else
+			{
+				id = default;
+				return false;
+			}
+		}
+
+		/// <inheritdoc/>
+		public override string ToString() => Id.ToString();
+	}
+
+	/// <summary>
+	/// Converter to and from <see cref="ObjectId"/> instances.
+	/// </summary>
+	class SessionIdConverter : ObjectIdConverter<SessionId>
+	{
+		/// <inheritdoc/>
+		public override SessionId FromObjectId(ObjectId id) => new SessionId(id);
+
+		/// <inheritdoc/>
+		public override ObjectId ToObjectId(SessionId value) => value.Id;
+	}
+}
diff --git a/Engine/Source/Programs/Horde/Horde.Build/Commands/Migrate/MigrateLogsCommand.cs b/Engine/Source/Programs/Horde/Horde.Build/Commands/Migrate/MigrateLogsCommand.cs
new file mode 100644
--- /dev/null
+++ b/Engine/Source/Programs/Horde/Horde.Build/Commands/Migrate/MigrateLogsCommand.cs
@@ -0,0 +1,112 @@
+// Copyright Epic Games, Inc. All Rights Reserved.
+
+using System;
+using System.Collections.Generic;
+using System.IO;
+using System.Threading;
+using System.Threading.Tasks;
+using EpicGames.Core;
+using EpicGames.Horde.Logs;
+using EpicGames.Horde.Storage;
+using EpicGames.Horde.Storage.Nodes;
+using Horde.Build.Logs;
+using Horde.Build.Perforce;
+using Horde.Build.Server;
+using Horde.Build.Storage;
+using Horde.Build.Utilities;
+using Microsoft.Extensions.Caching.Memory;
+using Microsoft.Extensions.Configuration;
+using Microsoft.Extensions.DependencyInjection;
+using Microsoft.Extensions.Logging;
+using MongoDB.Bson;
+using MongoDB.Bson.Serialization.Attributes;
+using MongoDB.Driver;
+
+namespace Horde.Build.Commands.Migrate
+{
+	[Command("migrate", "logs", "Extracts data from a bundle to the local hard drive")]
+	internal class MigrateLogsCommand : Command
+	{
+		class LogFile
+		{
+			public ObjectId Id { get; set; }
+			public RefName RefName { get; set; }
+		}
+
+		readonly IConfiguration _configuration;
+		readonly ILoggerProvider _loggerProvider;
+
+		public MigrateLogsCommand(IConfiguration configuration, ILoggerProvider loggerProvider)
+		{
+			_configuration = configuration;
+			_loggerProvider = loggerProvider;
+		}
+
+		public override async Task<int> ExecuteAsync(ILogger logger)
+		{
+			using ServiceProvider serviceProvider = Startup.CreateServiceProvider(_configuration, _loggerProvider);
+
+			MongoService mongoService = serviceProvider.GetRequiredService<MongoService>();
+			ILogFileCollection logCollection = serviceProvider.GetRequiredService<ILogFileCollection>();
+			ILogFileService logService = serviceProvider.GetRequiredService<ILogFileService>();
+
+			StorageService storageService = serviceProvider.GetRequiredService<StorageService>();
+			IStorageClient storageClient = await storageService.GetClientAsync(Namespace.Logs, CancellationToken.None);
+
+			byte[] buffer = new byte[1 * 1024];
+
+			ILogFile logFile = (await logCollection.GetLogFileAsync(LogId.Parse("6374195e5f2709066e7257a6"), CancellationToken.None))!;
+//			await foreach(ILogFile logFile in logCollection.GetLogFilesAsync())
+			{
+				if (logFile.RefName == RefName.Empty)
+				{
+//					TreeWriter writer = new TreeWriter(storageClient);
+					int bufferLength = 0;
+					int builderLength = 0;
+
+					LogBuilder builder = new LogBuilder(LogFormat.Json, logger);
+
+					using (Stream stream = await logService.OpenRawStreamAsync(logFile, CancellationToken.None))
+					{
+						for (; ; )
+						{
+							int readLength = await stream.ReadAsync(buffer.AsMemory(bufferLength));
+							if (readLength == 0)
+							{
+								break;
+							}
+
+							bufferLength += readLength;
+
+							int startIdx = 0;
+							for (int idx = 0; idx < bufferLength; idx++)
+							{
+								if (buffer[idx] == (byte)'\n')
+								{
+									ReadOnlyMemory<byte> line = buffer.AsMemory(startIdx, idx + 1 - startIdx);
+
+									builder.WriteData(line);
+									builderLength += line.Length;
+
+									if (builderLength > 1024 * 1024)
+									{
+//										await builder.FlushAsync(writer, false, CancellationToken.None);
+										builderLength = 0;
+									}
+
+									startIdx = idx + 1;
+								}
+							}
+
+							bufferLength -= startIdx;
+							buffer.AsSpan(startIdx, bufferLength).CopyTo(buffer.AsSpan());
+						}
+					}
+//					await builder.FlushAsync(writer, true, CancellationToken.None);
+				}
+			}
+
+			return 0;
+		}
+	}
+}
diff --git a/Engine/Source/Programs/Horde/Horde.Build/Compute/V1/ComputeService.cs b/Engine/Source/Programs/Horde/Horde.Build/Compute/V1/ComputeService.cs
--- a/Engine/Source/Programs/Horde/Horde.Build/Compute/V1/ComputeService.cs
+++ b/Engine/Source/Programs/Horde/Horde.Build/Compute/V1/ComputeService.cs
@@ -30,8 +30,6 @@ using StackExchange.Redis;
 
 namespace Horde.Build.Compute.V1
 {
-	using LeaseId = ObjectId<ILease>;
-
 	/// <summary>
 	/// Information about a particular task
 	/// </summary>
diff --git a/Engine/Source/Programs/Horde/Horde.Build/Devices/DeviceCollection.cs b/Engine/Source/Programs/Horde/Horde.Build/Devices/DeviceCollection.cs
--- a/Engine/Source/Programs/Horde/Horde.Build/Devices/DeviceCollection.cs
+++ b/Engine/Source/Programs/Horde/Horde.Build/Devices/DeviceCollection.cs
@@ -17,8 +17,6 @@ using MongoDB.Driver;
 
 namespace Horde.Build.Devices
 {
-	using UserId = ObjectId<IUser>;
-	
 	/// <summary>
 	/// Collection of device documents
 	/// </summary>
diff --git a/Engine/Source/Programs/Horde/Horde.Build/Devices/DeviceService.cs b/Engine/Source/Programs/Horde/Horde.Build/Devices/DeviceService.cs
--- a/Engine/Source/Programs/Horde/Horde.Build/Devices/DeviceService.cs
+++ b/Engine/Source/Programs/Horde/Horde.Build/Devices/DeviceService.cs
@@ -25,9 +25,6 @@ using OpenTracing.Util;
 
 namespace Horde.Build.Devices
 {
-	using JobId = ObjectId<IJob>;
-	using UserId = ObjectId<IUser>;
-
 	/// <summary>
 	///  Device Pool Authorization (convenience class for pool ACL's)
 	/// </summary>
@@ -420,7 +417,7 @@ namespace Horde.Build.Devices
 
 			if (jobId != null)
 			{
-				IJob? job = await _jobService.GetJobAsync(new JobId(jobId));
+				IJob? job = await _jobService.GetJobAsync(JobId.Parse(jobId));
 				if (job != null)
 				{					
 					streamId = job.StreamId.ToString();
@@ -519,7 +516,7 @@ namespace Horde.Build.Devices
 
 				if (jobId != null)
 				{
-					job = await _jobService.GetJobAsync(new JobId(jobId));
+					job = await _jobService.GetJobAsync(JobId.Parse(jobId));
 
 					if (job != null)
 					{
diff --git a/Engine/Source/Programs/Horde/Horde.Build/Issues/IIssueCollection.cs b/Engine/Source/Programs/Horde/Horde.Build/Issues/IIssueCollection.cs
--- a/Engine/Source/Programs/Horde/Horde.Build/Issues/IIssueCollection.cs
+++ b/Engine/Source/Programs/Horde/Horde.Build/Issues/IIssueCollection.cs
@@ -16,10 +16,6 @@ using MongoDB.Driver;
 
 namespace Horde.Build.Issues
 {
-	using JobId = ObjectId<IJob>;
-	using LogId = ObjectId<ILogFile>;
-	using UserId = ObjectId<IUser>;
-
 	/// <summary>
 	/// Fingerprint for an issue
 	/// </summary>
diff --git a/Engine/Source/Programs/Horde/Horde.Build/Issues/IssuesController.cs b/Engine/Source/Programs/Horde/Horde.Build/Issues/IssuesController.cs
--- a/Engine/Source/Programs/Horde/Horde.Build/Issues/IssuesController.cs
+++ b/Engine/Source/Programs/Horde/Horde.Build/Issues/IssuesController.cs
@@ -24,10 +24,6 @@ using Horde.Build.Server;
 
 namespace Horde.Build.Issues
 {
-	using JobId = ObjectId<IJob>;
-	using LogId = ObjectId<ILogFile>;
-	using UserId = ObjectId<IUser>;
-
 	/// <summary>
 	/// Controller for the /api/v1/issues endpoint
 	/// </summary>
@@ -228,7 +224,7 @@ namespace Horde.Build.Issues
 			UserId? ownerIdValue = null;
 			if (ownerId != null)
 			{
-				ownerIdValue = new UserId(ownerId);
+				ownerIdValue = UserId.Parse(ownerId);
 			}
 
 			List<IIssue> issues;
@@ -491,7 +487,7 @@ namespace Horde.Build.Issues
 			}
 			if(logIds != null)
 			{
-				logIdValues.UnionWith(logIds.Select(x => new LogId(x)));
+				logIdValues.UnionWith(logIds.Select(x => LogId.Parse(x)));
 			}
 
 			List<IIssueSpan> spans = await _issueCollection.FindSpansAsync(issueId);
@@ -549,13 +545,13 @@ namespace Horde.Build.Issues
 			UserId? newOwnerId = null;
 			if (request.OwnerId != null)
 			{
-				newOwnerId = request.OwnerId.Length == 0 ? UserId.Empty : new UserId(request.OwnerId);
+				newOwnerId = request.OwnerId.Length == 0 ? UserId.Empty : UserId.Parse(request.OwnerId);
 			}
 
 			UserId? newNominatedById = null;
 			if (request.NominatedById != null)
 			{
-				newNominatedById = new UserId(request.NominatedById);
+				newNominatedById = UserId.Parse(request.NominatedById);
 			}
 			else if (request.OwnerId != null)
 			{
@@ -577,13 +573,13 @@ namespace Horde.Build.Issues
 			UserId? newQuarantinedById = null;
 			if (request.QuarantinedById != null)
 			{
-				newQuarantinedById = request.QuarantinedById.Length > 0 ? new UserId(request.QuarantinedById) : UserId.Empty;
+				newQuarantinedById = request.QuarantinedById.Length > 0 ? UserId.Parse(request.QuarantinedById) : UserId.Empty;
 			}
 
 			UserId? newForceClosedById = null;
 			if (request.ForceClosedById != null)
 			{
-				newForceClosedById = request.ForceClosedById.Length > 0 ? new UserId(request.ForceClosedById) : UserId.Empty;
+				newForceClosedById = request.ForceClosedById.Length > 0 ? UserId.Parse(request.ForceClosedById) : UserId.Empty;
 			}
 
 
diff --git a/Engine/Source/Programs/Horde/Horde.Build/Jobs/Artifacts/ArtifactsControllerV1.cs b/Engine/Source/Programs/Horde/Horde.Build/Jobs/Artifacts/ArtifactsControllerV1.cs
--- a/Engine/Source/Programs/Horde/Horde.Build/Jobs/Artifacts/ArtifactsControllerV1.cs
+++ b/Engine/Source/Programs/Horde/Horde.Build/Jobs/Artifacts/ArtifactsControllerV1.cs
@@ -23,8 +23,6 @@ using MongoDB.Bson;
 
 namespace Horde.Build.Jobs.Artifacts
 {
-	using JobId = ObjectId<IJob>;
-
 	/// <summary>
 	/// Controller for the /api/artifacts endpoint
 	/// </summary>
@@ -338,7 +336,7 @@ namespace Horde.Build.Jobs.Artifacts
 				return BadRequest("Must specify a JobId");
 			}
 
-			IJob? job = await _jobService.GetJobAsync(new JobId(artifactZipRequest.JobId!));
+			IJob? job = await _jobService.GetJobAsync(JobId.Parse(artifactZipRequest.JobId!));
 			if (job == null)
 			{
 				return NotFound();
diff --git a/Engine/Source/Programs/Horde/Horde.Build/Jobs/IJob.cs b/Engine/Source/Programs/Horde/Horde.Build/Jobs/IJob.cs
--- a/Engine/Source/Programs/Horde/Horde.Build/Jobs/IJob.cs
+++ b/Engine/Source/Programs/Horde/Horde.Build/Jobs/IJob.cs
@@ -27,12 +27,7 @@ using MongoDB.Driver;
 
 namespace Horde.Build.Jobs
 {
-	using JobId = ObjectId<IJob>;
-	using LeaseId = ObjectId<ILease>;
-	using LogId = ObjectId<ILogFile>;
 	using ReportPlacement = HordeCommon.Rpc.ReportPlacement;
-	using SessionId = ObjectId<ISession>;
-	using UserId = ObjectId<IUser>;
 
 	/// <summary>
 	/// Report for a job or jobstep
diff --git a/Engine/Source/Programs/Horde/Horde.Build/Jobs/IJobCollection.cs b/Engine/Source/Programs/Horde/Horde.Build/Jobs/IJobCollection.cs
--- a/Engine/Source/Programs/Horde/Horde.Build/Jobs/IJobCollection.cs
+++ b/Engine/Source/Programs/Horde/Horde.Build/Jobs/IJobCollection.cs
@@ -20,12 +20,6 @@ using MongoDB.Bson;
 
 namespace Horde.Build.Jobs
 {
-	using JobId = ObjectId<IJob>;
-	using LeaseId = ObjectId<ILease>;
-	using LogId = ObjectId<ILogFile>;
-	using SessionId = ObjectId<ISession>;
-	using UserId = ObjectId<IUser>;
-
 	/// <summary>
 	/// Options for creating a new job
 	/// </summary>
diff --git a/Engine/Source/Programs/Horde/Horde.Build/Jobs/JobCollection.cs b/Engine/Source/Programs/Horde/Horde.Build/Jobs/JobCollection.cs
--- a/Engine/Source/Programs/Horde/Horde.Build/Jobs/JobCollection.cs
+++ b/Engine/Source/Programs/Horde/Horde.Build/Jobs/JobCollection.cs
@@ -34,12 +34,6 @@ using OpenTracing.Util;
 
 namespace Horde.Build.Jobs
 {
-	using JobId = ObjectId<IJob>;
-	using LeaseId = ObjectId<ILease>;
-	using LogId = ObjectId<ILogFile>;
-	using SessionId = ObjectId<ISession>;
-	using UserId = ObjectId<IUser>;
-
 	/// <summary>
 	/// Wrapper around the jobs collection in a mongo DB
 	/// </summary>
diff --git a/Engine/Source/Programs/Horde/Horde.Build/Jobs/JobId.cs b/Engine/Source/Programs/Horde/Horde.Build/Jobs/JobId.cs
new file mode 100644
--- /dev/null
+++ b/Engine/Source/Programs/Horde/Horde.Build/Jobs/JobId.cs
@@ -0,0 +1,47 @@
+// Copyright Epic Games, Inc. All Rights Reserved.
+
+using System.ComponentModel;
+using EpicGames.Core;
+using Horde.Build.Utilities;
+using MongoDB.Bson;
+
+namespace Horde.Build.Jobs
+{
+	/// <summary>
+	/// Identifier for a job
+	/// </summary>
+	/// <param name="Id">Id to construct from</param>
+	[JsonSchemaString]
+	[TypeConverter(typeof(ObjectIdTypeConverter<JobId, JobIdConverter>))]
+	[ObjectIdConverter(typeof(JobIdConverter))]
+	public record struct JobId(ObjectId Id)
+	{
+		/// <summary>
+		/// Constant value for an empty job id
+		/// </summary>
+		public static JobId Empty { get; } = default;
+
+		/// <summary>
+		/// Creates a new <see cref="JobId"/>
+		/// </summary>
+		public static JobId GenerateNewId() => new JobId(ObjectId.GenerateNewId());
+
+		/// <inheritdoc cref="ObjectId.Parse(System.String)"/>
+		public static JobId Parse(string text) => new JobId(ObjectId.Parse(text));
+
+		/// <inheritdoc/>
+		public override string ToString() => Id.ToString();
+	}
+
+	/// <summary>
+	/// Converter to and from <see cref="ObjectId"/> instances.
+	/// </summary>
+	class JobIdConverter : ObjectIdConverter<JobId>
+	{
+		/// <inheritdoc/>
+		public override JobId FromObjectId(ObjectId id) => new JobId(id);
+
+		/// <inheritdoc/>
+		public override ObjectId ToObjectId(JobId value) => value.Id;
+	}
+}
diff --git a/Engine/Source/Programs/Horde/Horde.Build/Jobs/JobService.cs b/Engine/Source/Programs/Horde/Horde.Build/Jobs/JobService.cs
--- a/Engine/Source/Programs/Horde/Horde.Build/Jobs/JobService.cs
+++ b/Engine/Source/Programs/Horde/Horde.Build/Jobs/JobService.cs
@@ -35,12 +35,6 @@ using OpenTracing.Util;
 
 namespace Horde.Build.Jobs
 {
-	using JobId = ObjectId<IJob>;
-	using LeaseId = ObjectId<ILease>;
-	using LogId = ObjectId<ILogFile>;
-	using SessionId = ObjectId<ISession>;
-	using UserId = ObjectId<IUser>;
-
 	/// <summary>
 	/// Exception thrown when attempting to retry executing a node that does not allow retries
 	/// </summary>
@@ -167,7 +161,7 @@ namespace Horde.Build.Jobs
 		public async Task<IJob> CreateJobAsync(JobId? jobId, StreamConfig streamConfig, TemplateId templateRefId, ContentHash templateHash, IGraph graph, string name, int change, int codeChange, CreateJobOptions options)
 		{
 			using IScope traceScope = GlobalTracer.Instance.BuildSpan("JobService.CreateJobAsync").StartActive();
-			traceScope.Span.SetTag("JobId", jobId);
+			traceScope.Span.SetTag("JobId", jobId?.ToString());
 			traceScope.Span.SetTag("Stream", streamConfig.Name);
 			traceScope.Span.SetTag("TemplateRefId", templateRefId);
 			traceScope.Span.SetTag("TemplateHash", templateHash);
@@ -197,7 +191,7 @@ namespace Horde.Build.Jobs
 				traceScope.Span.SetTag("JobTriggers.Count", options.JobTriggers.Count);
 			}
 
-			JobId jobIdValue = jobId ?? Horde.Build.Utilities.ObjectId<IJob>.GenerateNewId();
+			JobId jobIdValue = jobId ?? JobId.GenerateNewId();
 			using IDisposable scope = _logger.BeginScope("CreateJobAsync({JobId})", jobIdValue);
 
 			if (options.PreflightChange != null && ShouldClonePreflightChange(streamConfig.Id))
@@ -550,15 +544,15 @@ namespace Horde.Build.Jobs
 		public async Task<List<IJob>> FindJobsAsync(JobId[]? jobIds = null, StreamId? streamId = null, string? name = null, TemplateId[]? templates = null, int? minChange = null, int? maxChange = null, int? preflightChange = null, bool? preflightOnly = null, UserId? preflightStartedByUser = null, UserId? startedByUser = null, DateTimeOffset ? minCreateTime = null, DateTimeOffset? maxCreateTime = null, string? target = null, JobStepState[]? state = null, JobStepOutcome[]? outcome = null, DateTimeOffset? modifiedBefore = null, DateTimeOffset? modifiedAfter = null, int? index = null, int? count = null, bool consistentRead = true, bool? excludeUserJobs = null)
 		{
 			using IScope scope = GlobalTracer.Instance.BuildSpan("JobService.FindJobsAsync").StartActive();
-			scope.Span.SetTag("JobIds", jobIds);
+			scope.Span.SetTag("JobIds", (jobIds == null)? null : String.Join(',', jobIds));
 			scope.Span.SetTag("StreamId", streamId);
 			scope.Span.SetTag("Name", name);
 			scope.Span.SetTag("Templates", templates);
 			scope.Span.SetTag("MinChange", minChange);
 			scope.Span.SetTag("MaxChange", maxChange);
 			scope.Span.SetTag("PreflightChange", preflightChange);
-			scope.Span.SetTag("PreflightStartedByUser", preflightStartedByUser);
-			scope.Span.SetTag("StartedByUser", startedByUser);
+			scope.Span.SetTag("PreflightStartedByUser", preflightStartedByUser?.ToString());
+			scope.Span.SetTag("StartedByUser", startedByUser?.ToString());
 			scope.Span.SetTag("MinCreateTime", minCreateTime);
 			scope.Span.SetTag("MaxCreateTime", maxCreateTime);
 			scope.Span.SetTag("Target", target);
@@ -637,7 +631,7 @@ namespace Horde.Build.Jobs
 			using IScope scope = GlobalTracer.Instance.BuildSpan("JobService.FindJobsByStreamWithTemplatesAsync").StartActive();
 			scope.Span.SetTag("StreamId", streamId);
 			scope.Span.SetTag("Templates", templates);
-			scope.Span.SetTag("PreflightStartedByUser", preflightStartedByUser);
+			scope.Span.SetTag("PreflightStartedByUser", preflightStartedByUser?.ToString());
 			scope.Span.SetTag("MaxCreateTime", maxCreateTime);
 			scope.Span.SetTag("ModifiedAfter", modifiedAfter);
 			scope.Span.SetTag("Index", index);
@@ -655,7 +649,7 @@ namespace Horde.Build.Jobs
 		public async Task<IJob?> TryUpdateGraphAsync(IJob job, IGraph newGraph)
 		{
 			using IScope traceScope = GlobalTracer.Instance.BuildSpan("JobService.TryUpdateGraphAsync").StartActive();
-			traceScope.Span.SetTag("Job", job.Id);
+			traceScope.Span.SetTag("Job", job.Id.ToString());
 			traceScope.Span.SetTag("NewGraph", newGraph.Id);
 
 			using IDisposable scope = _logger.BeginScope("TryUpdateGraphAsync({JobId})", job.Id);
@@ -679,7 +673,7 @@ namespace Horde.Build.Jobs
 		public async Task<IJobTiming> GetJobTimingAsync(IJob job)
 		{
 			using IScope traceScope = GlobalTracer.Instance.BuildSpan("JobService.GetJobTimingAsync").StartActive();
-			traceScope.Span.SetTag("Job", job.Id);
+			traceScope.Span.SetTag("Job", job.Id.ToString());
 
 			using IDisposable scope = _logger.BeginScope("GetJobTimingAsync({JobId})", job.Id);
 
@@ -800,9 +794,9 @@ namespace Horde.Build.Jobs
 		public async Task<IJob?> UpdateBatchAsync(IJob job, SubResourceId batchId, StreamConfig streamConfig, LogId? newLogId = null, JobStepBatchState? newState = null)
 		{
 			using IScope traceScope = GlobalTracer.Instance.BuildSpan("JobService.UpdateBatchAsync").StartActive();
-			traceScope.Span.SetTag("Job", job.Id);
+			traceScope.Span.SetTag("Job", job.Id.ToString());
 			traceScope.Span.SetTag("BatchId", batchId);
-			traceScope.Span.SetTag("NewLogId", newLogId);
+			traceScope.Span.SetTag("NewLogId", newLogId?.ToString());
 			traceScope.Span.SetTag("NewState", newState.ToString());
 
 			using IDisposable scope = _logger.BeginScope("UpdateBatchAsync({JobId})", job.Id);
@@ -935,7 +929,7 @@ namespace Horde.Build.Jobs
 		public async Task<IJob?> UpdateStepAsync(IJob job, SubResourceId batchId, SubResourceId stepId, StreamConfig streamConfig, JobStepState newState = JobStepState.Unspecified, JobStepOutcome newOutcome = JobStepOutcome.Unspecified, JobStepError? newError = null, bool? newAbortRequested = null, UserId? newAbortByUserId = null, LogId? newLogId = null, ObjectId? newNotificationTriggerId = null, UserId? newRetryByUserId = null, Priority? newPriority = null, List<Report>? newReports = null, Dictionary<string, string?>? newProperties = null)
 		{
 			using IScope traceScope = GlobalTracer.Instance.BuildSpan("JobService.UpdateStepAsync").StartActive();
-			traceScope.Span.SetTag("Job", job.Id);
+			traceScope.Span.SetTag("Job", job.Id.ToString());
 			traceScope.Span.SetTag("BatchId", batchId);
 			traceScope.Span.SetTag("StepId", stepId);
 			
@@ -980,7 +974,7 @@ namespace Horde.Build.Jobs
 		public async Task<IJob?> TryUpdateStepAsync(IJob job, SubResourceId batchId, SubResourceId stepId, StreamConfig streamConfig, JobStepState newState = JobStepState.Unspecified, JobStepOutcome newOutcome = JobStepOutcome.Unspecified, JobStepError? newError = null, bool? newAbortRequested = null, UserId? newAbortByUserId = null, LogId? newLogId = null, ObjectId? newTriggerId = null, UserId? newRetryByUserId = null, Priority? newPriority = null, List<Report>? newReports = null, Dictionary<string, string?>? newProperties = null)
 		{
 			using IScope traceScope = GlobalTracer.Instance.BuildSpan("JobService.TryUpdateStepAsync").StartActive();
-			traceScope.Span.SetTag("Job", job.Id);
+			traceScope.Span.SetTag("Job", job.Id.ToString());
 			traceScope.Span.SetTag("BatchId", batchId);
 			traceScope.Span.SetTag("StepId", stepId);
 
@@ -1124,7 +1118,7 @@ namespace Horde.Build.Jobs
 		private async Task<IJob> AutoSubmitChangeAsync(StreamConfig streamConfig, IJob job, IGraph graph)
 		{
 			using IScope traceScope = GlobalTracer.Instance.BuildSpan("JobService.AutoSubmitChangeAsync").StartActive();
-			traceScope.Span.SetTag("Job", job.Id);
+			traceScope.Span.SetTag("Job", job.Id.ToString());
 			traceScope.Span.SetTag("Graph", graph.Id);
 			
 			int? change;
@@ -1255,10 +1249,10 @@ namespace Horde.Build.Jobs
 		private async Task<IJob?> FireJobTriggerAsync(IJob job, IGraph graph, IChainedJob jobTrigger, StreamConfig streamConfig)
 		{
 			using IScope traceScope = GlobalTracer.Instance.BuildSpan("JobService.FireJobTriggerAsync").StartActive();
-			traceScope.Span.SetTag("Job", job.Id);
+			traceScope.Span.SetTag("Job", job.Id.ToString());
 			traceScope.Span.SetTag("Graph", graph.Id);
 			traceScope.Span.SetTag("JobTrigger.Target", jobTrigger.Target);
-			traceScope.Span.SetTag("JobTrigger.JobId", jobTrigger.JobId);
+			traceScope.Span.SetTag("JobTrigger.JobId", jobTrigger.JobId?.ToString());
 			traceScope.Span.SetTag("JobTrigger.TemplateRefId", jobTrigger.TemplateRefId);
 			
 			for (; ; )
diff --git a/Engine/Source/Programs/Horde/Horde.Build/Jobs/JobTaskSource.cs b/Engine/Source/Programs/Horde/Horde.Build/Jobs/JobTaskSource.cs
--- a/Engine/Source/Programs/Horde/Horde.Build/Jobs/JobTaskSource.cs
+++ b/Engine/Source/Programs/Horde/Horde.Build/Jobs/JobTaskSource.cs
@@ -33,10 +33,6 @@ using MongoDB.Driver;
 
 namespace Horde.Build.Jobs
 {
-	using JobId = ObjectId<IJob>;
-	using LeaseId = ObjectId<ILease>;
-	using LogId = ObjectId<ILogFile>;
-	
 	/// <summary>
 	/// Background service to dispatch pending work to agents in priority order.
 	/// </summary>
@@ -137,7 +133,7 @@ namespace Horde.Build.Jobs
 				int delta = y.Batch.SchedulePriority - x.Batch.SchedulePriority;
 				if (delta == 0)
 				{
-					delta = x._job.Id.CompareTo(y._job.Id);
+					delta = x._job.Id.Id.CompareTo(y._job.Id.Id);
 					if (delta == 0)
 					{
 						delta = (int)x.Batch.Id.Value - (int)y.Batch.Id.Value;
@@ -242,7 +238,7 @@ namespace Horde.Build.Jobs
 			_settings = settings;
 			_logger = logger;
 
-			OnLeaseStartedProperties.Add(nameof(ExecuteJobTask.JobId), x => new JobId(x.JobId)).Add(nameof(ExecuteJobTask.BatchId), x => SubResourceId.Parse(x.BatchId)).Add(nameof(ExecuteJobTask.LogId), x => new LogId(x.LogId));
+			OnLeaseStartedProperties.Add(nameof(ExecuteJobTask.JobId), x => JobId.Parse(x.JobId)).Add(nameof(ExecuteJobTask.BatchId), x => SubResourceId.Parse(x.BatchId)).Add(nameof(ExecuteJobTask.LogId), x => LogId.Parse(x.LogId));
 		}
 
 		/// <inheritdoc/>
@@ -715,7 +711,7 @@ namespace Horde.Build.Jobs
 		/// <inheritdoc/>
 		public override Task CancelLeaseAsync(IAgent agent, LeaseId leaseId, ExecuteJobTask task)
 		{
-			return CancelLeaseAsync(agent, new JobId(task.JobId), task.BatchId.ToSubResourceId());
+			return CancelLeaseAsync(agent, JobId.Parse(task.JobId), task.BatchId.ToSubResourceId());
 		}
 
 		/// <summary>
@@ -986,7 +982,7 @@ namespace Horde.Build.Jobs
 			if (outcome != LeaseOutcome.Success)
 			{
 				AgentId agentId = agent.Id;
-				JobId jobId = new JobId(task.JobId);
+				JobId jobId = JobId.Parse(task.JobId);
 				SubResourceId batchId = task.BatchId.ToSubResourceId();
 
 				// Update the batch
diff --git a/Engine/Source/Programs/Horde/Horde.Build/Jobs/JobsController.cs b/Engine/Source/Programs/Horde/Horde.Build/Jobs/JobsController.cs
--- a/Engine/Source/Programs/Horde/Horde.Build/Jobs/JobsController.cs
+++ b/Engine/Source/Programs/Horde/Horde.Build/Jobs/JobsController.cs
@@ -30,9 +30,6 @@ using OpenTracing.Util;
 
 namespace Horde.Build.Jobs
 {
-	using JobId = ObjectId<IJob>;
-	using UserId = ObjectId<IUser>;
-
 	/// <summary>
 	/// Controller for the /api/v1/jobs endpoing
 	/// </summary>
@@ -754,7 +751,7 @@ namespace Horde.Build.Jobs
 			[FromQuery] int index = 0,
 			[FromQuery] int count = 100)
 		{
-			JobId[]? jobIdValues = (ids == null) ? (JobId[]?)null : Array.ConvertAll(ids, x => new JobId(x));
+			JobId[]? jobIdValues = (ids == null) ? (JobId[]?)null : Array.ConvertAll(ids, x => JobId.Parse(x));
 			StreamId? streamIdValue = (streamId == null)? (StreamId?)null : new StreamId(streamId);
 			
 			TemplateId[]? templateRefIds = (templates != null && templates.Length > 0) ? templates.Select(x => new TemplateId(x)).ToArray() : null;
@@ -768,14 +765,14 @@ namespace Horde.Build.Jobs
 
 			if (preflightStartedByUserId != null)
 			{
-				preflightStartedByUserIdValue = new UserId(preflightStartedByUserId);
+				preflightStartedByUserIdValue = UserId.Parse(preflightStartedByUserId);
 			}
 
 			UserId? startedByUserIdValue = null;
 
 			if (startedByUserId != null)
 			{
-				startedByUserIdValue = new UserId(startedByUserId);
+				startedByUserIdValue = UserId.Parse(startedByUserId);
 			}
 
 			List<IJob> jobs;
@@ -818,7 +815,7 @@ namespace Horde.Build.Jobs
 		{
 			StreamId streamIdValue = new StreamId(streamId);
 			TemplateId[] templateRefIds = templates.Select(x => new TemplateId(x)).ToArray();
-			UserId? preflightStartedByUserIdValue = preflightStartedByUserId != null ? new UserId(preflightStartedByUserId) : null;
+			UserId? preflightStartedByUserIdValue = preflightStartedByUserId != null ? UserId.Parse(preflightStartedByUserId) : null;
 			count = Math.Min(1000, count);
 
 			List<IJob> jobs = await _jobService.FindJobsByStreamWithTemplatesAsync(streamIdValue, templateRefIds, preflightStartedByUserIdValue, maxCreateTime, modifiedAfter, index, count, consistentRead);
@@ -1096,7 +1093,7 @@ namespace Horde.Build.Jobs
 				return Forbid("Missing session claim for job {JobId} batch {BatchId}", jobId, batchId);
 			}
 
-			IJob? newJob = await _jobService.UpdateBatchAsync(job, batchId, streamConfig, request.LogId?.ToObjectId<ILogFile>(), request.State);
+			IJob? newJob = await _jobService.UpdateBatchAsync(job, batchId, streamConfig, (request.LogId == null)? null : LogId.Parse(request.LogId), request.State);
 			if (newJob == null)
 			{
 				return NotFound(jobId);
@@ -1247,7 +1244,7 @@ namespace Horde.Build.Jobs
 
 			try
 			{
-				IJob? newJob = await _jobService.UpdateStepAsync(job, batchId, stepId, streamConfig, request.State, request.Outcome, null, request.AbortRequested, abortByUser, request.LogId?.ToObjectId<ILogFile>(), null, retryByUser, request.Priority, null, request.Properties);
+				IJob? newJob = await _jobService.UpdateStepAsync(job, batchId, stepId, streamConfig, request.State, request.Outcome, null, request.AbortRequested, abortByUser, (request.LogId == null)? null : LogId.Parse(request.LogId), null, retryByUser, request.Priority, null, request.Properties);
 				if (newJob == null)
 				{
 					return NotFound(jobId);
diff --git a/Engine/Source/Programs/Horde/Horde.Build/Jobs/TestData/ITestData.cs b/Engine/Source/Programs/Horde/Horde.Build/Jobs/TestData/ITestData.cs
--- a/Engine/Source/Programs/Horde/Horde.Build/Jobs/TestData/ITestData.cs
+++ b/Engine/Source/Programs/Horde/Horde.Build/Jobs/TestData/ITestData.cs
@@ -9,13 +9,6 @@ using MongoDB.Bson.Serialization.Attributes;
 
 namespace Horde.Build.Jobs.TestData
 {
-	using JobId = ObjectId<IJob>;
-	
-	using TestId = ObjectId<ITest>;
-	using TestSuiteId = ObjectId<ITestSuite>;
-	using TestMetaId = ObjectId<ITestMeta>;
-	using TestRefId = ObjectId<ITestDataRef>;
-
 	/// <summary>
 	/// Test outcome
 	/// </summary>
diff --git a/Engine/Source/Programs/Horde/Horde.Build/Jobs/TestData/ITestDataCollection.cs b/Engine/Source/Programs/Horde/Horde.Build/Jobs/TestData/ITestDataCollection.cs
--- a/Engine/Source/Programs/Horde/Horde.Build/Jobs/TestData/ITestDataCollection.cs
+++ b/Engine/Source/Programs/Horde/Horde.Build/Jobs/TestData/ITestDataCollection.cs
@@ -9,12 +9,6 @@ using MongoDB.Bson;
 
 namespace Horde.Build.Jobs.TestData
 {	
-	using JobId = ObjectId<IJob>;
-	using TestId = ObjectId<ITest>;
-	using TestSuiteId = ObjectId<ITestSuite>;
-	using TestMetaId = ObjectId<ITestMeta>;
-	using TestRefId = ObjectId<ITestDataRef>;
-
 	/// <summary>
 	/// Collection of test data documents
 	/// </summary>
diff --git a/Engine/Source/Programs/Horde/Horde.Build/Jobs/TestData/TestDataCollection.cs b/Engine/Source/Programs/Horde/Horde.Build/Jobs/TestData/TestDataCollection.cs
--- a/Engine/Source/Programs/Horde/Horde.Build/Jobs/TestData/TestDataCollection.cs
+++ b/Engine/Source/Programs/Horde/Horde.Build/Jobs/TestData/TestDataCollection.cs
@@ -20,13 +20,6 @@ using Microsoft.Extensions.Options;
 
 namespace Horde.Build.Jobs.TestData
 {
-	using JobId = ObjectId<IJob>;
-	
-	using TestId = ObjectId<ITest>;
-	using TestSuiteId = ObjectId<ITestSuite>;
-	using TestMetaId = ObjectId<ITestMeta>;
-	using TestRefId = ObjectId<ITestDataRef>;
-
 	/// <summary>
 	/// Collection of test data documents
 	/// </summary>
diff --git a/Engine/Source/Programs/Horde/Horde.Build/Jobs/TestData/TestDataController.cs b/Engine/Source/Programs/Horde/Horde.Build/Jobs/TestData/TestDataController.cs
--- a/Engine/Source/Programs/Horde/Horde.Build/Jobs/TestData/TestDataController.cs
+++ b/Engine/Source/Programs/Horde/Horde.Build/Jobs/TestData/TestDataController.cs
@@ -20,12 +20,6 @@ using Horde.Build.Server;
 
 namespace Horde.Build.Jobs.TestData
 {
-	using JobId = ObjectId<IJob>;
-	using TestId = ObjectId<ITest>;
-	using TestSuiteId = ObjectId<ITestSuite>;
-	using TestMetaId = ObjectId<ITestMeta>;
-	using TestRefId = ObjectId<ITestDataRef>;
-
 	/// <summary>
 	/// Controller for the /api/v1/testdata endpoint
 	/// </summary>
@@ -95,7 +89,7 @@ namespace Horde.Build.Jobs.TestData
 		[ProducesResponseType(typeof(List<GetTestDataDetailsResponse>), 200)]
 		public async Task<ActionResult<List<GetTestDataDetailsResponse>>> GetTestDetailsAsync([FromQuery(Name = "id")] string[] ids)
 		{
-			TestRefId[] idValues = Array.ConvertAll(ids, x => new TestRefId(x));
+			TestRefId[] idValues = Array.ConvertAll(ids, x => TestRefId.Parse(x));
 			List<ITestDataDetails> details = await _testDataService.FindTestDetails(idValues);
 			return details.Select(d => new GetTestDataDetailsResponse(d)).ToList();
 		}
@@ -112,7 +106,7 @@ namespace Horde.Build.Jobs.TestData
 		{
 			HashSet<string> testIds = new HashSet<string>(request.testIds);
 
-			List<ITest> testValues = await _testDataService.FindTests(testIds.Select(x => new TestId(x)).ToArray());			
+			List<ITest> testValues = await _testDataService.FindTests(testIds.Select(x => TestId.Parse(x)).ToArray());			
 
 			return testValues.Select(x => new GetTestResponse(x)).ToList();
 		}
@@ -301,7 +295,7 @@ namespace Horde.Build.Jobs.TestData
 		[Route("/api/v1/testdata")]
 		public async Task<ActionResult<CreateTestDataResponse>> CreateAsync(CreateTestDataRequest request)
 		{
-			IJob? job = await _jobService.GetJobAsync(new JobId(request.JobId));
+			IJob? job = await _jobService.GetJobAsync(JobId.Parse(request.JobId));
 			if (job == null)
 			{
 				return NotFound();
@@ -347,7 +341,7 @@ namespace Horde.Build.Jobs.TestData
 
 			List<object> results = new List<object>();
 
-			List<ITestData> documents = await _testDataCollection.FindAsync(streamIdValue, minChange, maxChange, jobId?.ToObjectId<IJob>(), jobStepId?.ToSubResourceId(), key, index, count);
+			List<ITestData> documents = await _testDataCollection.FindAsync(streamIdValue, minChange, maxChange, (jobId == null)? null : JobId.Parse(jobId), jobStepId?.ToSubResourceId(), key, index, count);
 			foreach (ITestData document in documents)
 			{
 				if (await _jobService.AuthorizeAsync(document.JobId, AclAction.ViewJob, User, _globalConfig.Value))
diff --git a/Engine/Source/Programs/Horde/Horde.Build/Jobs/TestData/TestDataService.cs b/Engine/Source/Programs/Horde/Horde.Build/Jobs/TestData/TestDataService.cs
--- a/Engine/Source/Programs/Horde/Horde.Build/Jobs/TestData/TestDataService.cs
+++ b/Engine/Source/Programs/Horde/Horde.Build/Jobs/TestData/TestDataService.cs
@@ -14,11 +14,6 @@ using Microsoft.Extensions.Options;
 
 namespace Horde.Build.Jobs.TestData
 {
-	using TestId = ObjectId<ITest>;
-	using TestSuiteId = ObjectId<ITestSuite>;
-	using TestMetaId = ObjectId<ITestMeta>;
-	using TestRefId = ObjectId<ITestDataRef>;
-
 	/// <summary>
 	/// Device management service
 	/// </summary>
@@ -153,8 +148,8 @@ namespace Horde.Build.Jobs.TestData
 		/// <returns></returns>
 		public async Task<List<ITestDataRef>> FindTestRefs(StreamId[] streamIds, TestMetaId[] metaIds, string[]? testIds = null, string[]? suiteIds = null, DateTime? minCreateTime = null, DateTime? maxCreateTime = null, int? minChange = null, int? maxChange = null)
 		{
-			TestId[]? tids = testIds?.ConvertAll(x => new TestId(x));
-			TestSuiteId[]? sids = suiteIds?.ConvertAll(x => new TestSuiteId(x));
+			TestId[]? tids = testIds?.ConvertAll(x => TestId.Parse(x));
+			TestSuiteId[]? sids = suiteIds?.ConvertAll(x => TestSuiteId.Parse(x));
 
 			return await _testData.FindTestRefs(streamIds, metaIds, tids, sids, minCreateTime, maxCreateTime, minChange, maxChange);
 		}
diff --git a/Engine/Source/Programs/Horde/Horde.Build/Logs/Builder/RedisLogBuilder.cs b/Engine/Source/Programs/Horde/Horde.Build/Logs/Builder/RedisLogBuilder.cs
--- a/Engine/Source/Programs/Horde/Horde.Build/Logs/Builder/RedisLogBuilder.cs
+++ b/Engine/Source/Programs/Horde/Horde.Build/Logs/Builder/RedisLogBuilder.cs
@@ -17,7 +17,6 @@ using StackExchange.Redis;
 namespace Horde.Build.Logs.Builder
 {
 	using Condition = StackExchange.Redis.Condition;
-	using LogId = ObjectId<ILogFile>;
 
 	/// <summary>
 	/// Redis-based cache for log file chunks
diff --git a/Engine/Source/Programs/Horde/Horde.Build/Logs/ILogFile.cs b/Engine/Source/Programs/Horde/Horde.Build/Logs/ILogFile.cs
--- a/Engine/Source/Programs/Horde/Horde.Build/Logs/ILogFile.cs
+++ b/Engine/Source/Programs/Horde/Horde.Build/Logs/ILogFile.cs
@@ -11,10 +11,6 @@ using Horde.Build.Utilities;
 
 namespace Horde.Build.Logs
 {
-	using JobId = ObjectId<IJob>;
-	using LogId = ObjectId<ILogFile>;
-	using SessionId = ObjectId<ISession>;
-
 	/// <summary>
 	/// Information about a log file chunk
 	/// </summary>
diff --git a/Engine/Source/Programs/Horde/Horde.Build/Logs/ILogFileCollection.cs b/Engine/Source/Programs/Horde/Horde.Build/Logs/ILogFileCollection.cs
--- a/Engine/Source/Programs/Horde/Horde.Build/Logs/ILogFileCollection.cs
+++ b/Engine/Source/Programs/Horde/Horde.Build/Logs/ILogFileCollection.cs
@@ -9,10 +9,6 @@ using Horde.Build.Utilities;
 
 namespace Horde.Build.Logs
 {
-	using JobId = ObjectId<IJob>;
-	using LogId = ObjectId<ILogFile>;
-	using SessionId = ObjectId<ISession>;
-
 	/// <summary>
 	/// Updates a log file chunk
 	/// </summary>
diff --git a/Engine/Source/Programs/Horde/Horde.Build/Logs/LogEventCollection.cs b/Engine/Source/Programs/Horde/Horde.Build/Logs/LogEventCollection.cs
--- a/Engine/Source/Programs/Horde/Horde.Build/Logs/LogEventCollection.cs
+++ b/Engine/Source/Programs/Horde/Horde.Build/Logs/LogEventCollection.cs
@@ -14,8 +14,6 @@ using MongoDB.Driver;
 
 namespace Horde.Build.Logs
 {
-	using LogId = ObjectId<ILogFile>;
-
 	/// <summary>
 	/// Collection of event documents
 	/// </summary>
@@ -134,7 +132,7 @@ namespace Horde.Build.Logs
 		/// <inheritdoc/>
 		public async Task<List<ILogEvent>> FindAsync(LogId logId, ObjectId? spanId = null, int? index = null, int? count = null)
 		{
-			_logger.LogInformation("Querying for log events for log {LogId} creation time {CreateTime}", logId, logId.Value.CreationTime);
+			_logger.LogInformation("Querying for log events for log {LogId} creation time {CreateTime}", logId, logId.Id.CreationTime);
 
 			FilterDefinitionBuilder<LogEventDocument> builder = Builders<LogEventDocument>.Filter;
 
diff --git a/Engine/Source/Programs/Horde/Horde.Build/Logs/LogId.cs b/Engine/Source/Programs/Horde/Horde.Build/Logs/LogId.cs
new file mode 100644
--- /dev/null
+++ b/Engine/Source/Programs/Horde/Horde.Build/Logs/LogId.cs
@@ -0,0 +1,62 @@
+// Copyright Epic Games, Inc. All Rights Reserved.
+
+using System.ComponentModel;
+using EpicGames.Core;
+using Horde.Build.Utilities;
+using MongoDB.Bson;
+
+namespace Horde.Build.Logs
+{
+	/// <summary>
+	/// Identifier for a log
+	/// </summary>
+	/// <param name="Id">Id to construct from</param>
+	[JsonSchemaString]
+	[TypeConverter(typeof(ObjectIdTypeConverter<LogId, LogIdConverter>))]
+	[ObjectIdConverter(typeof(LogIdConverter))]
+	public record struct LogId(ObjectId Id)
+	{
+		/// <summary>
+		/// Constant value for empty user id
+		/// </summary>
+		public static LogId Empty { get; } = new LogId(ObjectId.Empty);
+
+		/// <summary>
+		/// Creates a new <see cref="LogId"/>
+		/// </summary>
+		public static LogId GenerateNewId() => new LogId(ObjectId.GenerateNewId());
+
+		/// <inheritdoc cref="ObjectId.Parse(System.String)"/>
+		public static LogId Parse(string text) => new LogId(ObjectId.Parse(text));
+
+		/// <inheritdoc cref="ObjectId.TryParse(System.String, out ObjectId)"/>
+		public static bool TryParse(string text, out LogId id)
+		{
+			if (ObjectId.TryParse(text, out ObjectId objectId))
+			{
+				id = new LogId(objectId);
+				return true;
+			}
+			else
+			{
+				id = default;
+				return false;
+			}
+		}
+
+		/// <inheritdoc/>
+		public override string ToString() => Id.ToString();
+	}
+
+	/// <summary>
+	/// Converter to and from <see cref="ObjectId"/> instances.
+	/// </summary>
+	class LogIdConverter : ObjectIdConverter<LogId>
+	{
+		/// <inheritdoc/>
+		public override LogId FromObjectId(ObjectId id) => new LogId(id);
+
+		/// <inheritdoc/>
+		public override ObjectId ToObjectId(LogId value) => value.Id;
+	}
+}
diff --git a/Engine/Source/Programs/Horde/Horde.Build/Logs/LogRpcService.cs b/Engine/Source/Programs/Horde/Horde.Build/Logs/LogRpcService.cs
--- a/Engine/Source/Programs/Horde/Horde.Build/Logs/LogRpcService.cs
+++ b/Engine/Source/Programs/Horde/Horde.Build/Logs/LogRpcService.cs
@@ -12,8 +12,6 @@ using System.Threading;
 
 namespace Horde.Build.Logs
 {
-	using LogId = ObjectId<ILogFile>;
-
 	/// <summary>
 	/// Implements the Horde gRPC service for bots updating their status and dequeing work
 	/// </summary>
@@ -41,7 +39,7 @@ namespace Horde.Build.Logs
 		/// <inheritdoc/>
 		public override async Task<UpdateLogResponse> UpdateLog(UpdateLogRequest request, ServerCallContext context)
 		{
-			ILogFile? logFile = await _logFileService.GetCachedLogFileAsync(new LogId(request.LogId), context.CancellationToken);
+			ILogFile? logFile = await _logFileService.GetCachedLogFileAsync(LogId.Parse(request.LogId), context.CancellationToken);
 			if (logFile == null)
 			{
 				throw new StructuredRpcException(StatusCode.NotFound, "Resource not found");
@@ -72,7 +70,7 @@ namespace Horde.Build.Logs
 			while (await moveNextTask)
 			{
 				UpdateLogTailRequest request = requestStream.Current;
-				LogId logId = new LogId(request.LogId);
+				LogId logId = LogId.Parse(request.LogId);
 
 				if (request.TailData.Length > 0)
 				{
diff --git a/Engine/Source/Programs/Horde/Horde.Build/Notifications/Sinks/SlackNotificationSink.cs b/Engine/Source/Programs/Horde/Horde.Build/Notifications/Sinks/SlackNotificationSink.cs
--- a/Engine/Source/Programs/Horde/Horde.Build/Notifications/Sinks/SlackNotificationSink.cs
+++ b/Engine/Source/Programs/Horde/Horde.Build/Notifications/Sinks/SlackNotificationSink.cs
@@ -45,10 +45,6 @@ using MongoDB.Driver;
 
 namespace Horde.Build.Notifications.Sinks
 {
-	using JobId = ObjectId<IJob>;
-	using LogId = ObjectId<ILogFile>;
-	using UserId = ObjectId<IUser>;
-
 	/// <summary>
 	/// Maintains a connection to Slack, in order to receive socket-mode notifications of user interactions
 	/// </summary>
@@ -2627,7 +2623,7 @@ namespace Horde.Build.Notifications.Sinks
 							{
 								int issueId = Int32.Parse(match.Groups[1].Value, NumberStyles.Integer, CultureInfo.InvariantCulture);
 								string verb = match.Groups[2].Value;
-								UserId userId = match.Groups[3].Value.ToObjectId<IUser>();
+								UserId userId = new UserId(ObjectId.Parse(match.Groups[3].Value));
 								await HandleIssueDmResponseAsync(issueId, verb, userId, payload.User.Id);
 							}
 							else if (TryMatch(action.Value, @"^issue_(\d+)_([a-zA-Z]+)$", out match) && payload.TriggerId != null)
@@ -2651,7 +2647,7 @@ namespace Horde.Build.Notifications.Sinks
 						if (TryMatch(payload.View.CallbackId, @"^issue_(\d+)_markfixed_([a-fA-F0-9]{24})$", out match))
 						{
 							int issueId = Int32.Parse(match.Groups[1].Value, NumberStyles.Integer, CultureInfo.InvariantCulture);
-							UserId userId = match.Groups[2].Value.ToObjectId<IUser>();
+							UserId userId = UserId.Parse(match.Groups[2].Value);
 
 							string? fixChangeStr;
 							if (payload.View.State.TryGetValue("fix_cl", "fix_cl_action", out fixChangeStr))
@@ -2671,7 +2667,7 @@ namespace Horde.Build.Notifications.Sinks
 						else if (TryMatch(payload.View.CallbackId, @"^issue_(\d+)_ack_([a-fA-F0-9]{24})$", out match))
 						{
 							int issueId = Int32.Parse(match.Groups[1].Value, NumberStyles.Integer, CultureInfo.InvariantCulture);
-							UserId userId = match.Groups[2].Value.ToObjectId<IUser>();
+							UserId userId = UserId.Parse(match.Groups[2].Value);
 							await _issueService.UpdateIssueAsync(issueId, acknowledged: true, ownerId: userId, initiatedById: userId);
 						}
 					}
diff --git a/Engine/Source/Programs/Horde/Horde.Build/Server/DebugController.cs b/Engine/Source/Programs/Horde/Horde.Build/Server/DebugController.cs
--- a/Engine/Source/Programs/Horde/Horde.Build/Server/DebugController.cs
+++ b/Engine/Source/Programs/Horde/Horde.Build/Server/DebugController.cs
@@ -437,7 +437,7 @@ namespace Horde.Build.Server
 		/// <returns>Information about the requested project</returns>
 		[HttpGet]
 		[Route("/api/v1/debug/logs/{LogFileId}")]
-		public async Task<ActionResult<object>> GetLogAsync(ObjectId<ILogFile> logFileId, [FromQuery] PropertyFilter? filter = null)
+		public async Task<ActionResult<object>> GetLogAsync(LogId logFileId, [FromQuery] PropertyFilter? filter = null)
 		{
 			if (!_globalConfig.Value.Authorize(AclAction.AdminRead, User))
 			{
diff --git a/Engine/Source/Programs/Horde/Horde.Build/Server/RpcService.cs b/Engine/Source/Programs/Horde/Horde.Build/Server/RpcService.cs
--- a/Engine/Source/Programs/Horde/Horde.Build/Server/RpcService.cs
+++ b/Engine/Source/Programs/Horde/Horde.Build/Server/RpcService.cs
@@ -43,8 +43,6 @@ using Microsoft.Extensions.Options;
 
 namespace Horde.Build.Server
 {
-	using JobId = ObjectId<IJob>;
-	using LogId = ObjectId<ILogFile>;
 	using RpcAgentCapabilities = HordeCommon.Rpc.Messages.AgentCapabilities;
 	using RpcDeviceCapabilities = HordeCommon.Rpc.Messages.DeviceCapabilities;
 	using RpcGetStreamResponse = HordeCommon.Rpc.GetStreamResponse;
@@ -52,7 +50,6 @@ namespace Horde.Build.Server
 	using RpcGetStepResponse = HordeCommon.Rpc.GetStepResponse;
 	using RpcUpdateJobRequest = HordeCommon.Rpc.UpdateJobRequest;
 	using RpcUpdateStepRequest = HordeCommon.Rpc.UpdateStepRequest;
-	using SessionId = ObjectId<ISession>;
 	
 	/// <summary>
 	/// Implements the Horde gRPC service for bots updating their status and dequeing work
@@ -469,7 +466,7 @@ namespace Horde.Build.Server
 		/// <returns>Information about the new agent</returns>
 		public override async Task<Empty> UpdateJob(RpcUpdateJobRequest request, ServerCallContext context)
 		{
-			JobId jobIdValue = new JobId(request.JobId);
+			JobId jobIdValue = JobId.Parse(request.JobId);
 
 			IJob? job = await _jobService.GetJobAsync(jobIdValue);
 			if (job == null)
@@ -495,7 +492,7 @@ namespace Horde.Build.Server
 		{
 			SubResourceId batchId = request.BatchId.ToSubResourceId();
 
-			IJob? job = await _jobService.GetJobAsync(new JobId(request.JobId));
+			IJob? job = await _jobService.GetJobAsync(JobId.Parse(request.JobId));
 			if (job == null)
 			{
 				throw new StructuredRpcException(StatusCode.NotFound, "Job {JobId} not found", request.JobId);
@@ -529,7 +526,7 @@ namespace Horde.Build.Server
 		/// <returns>Information about the new agent</returns>
 		public override async Task<Empty> FinishBatch(FinishBatchRequest request, ServerCallContext context)
 		{
-			IJob? job = await _jobService.GetJobAsync(new JobId(request.JobId));
+			IJob? job = await _jobService.GetJobAsync(JobId.Parse(request.JobId));
 			if (job == null)
 			{
 				throw new StructuredRpcException(StatusCode.NotFound, "Job {JobId} not found", request.JobId);
@@ -566,7 +563,7 @@ namespace Horde.Build.Server
 		async Task<BeginStepResponse?> TryBeginStep(BeginStepRequest request, Boxed<ILogFile?> log, ServerCallContext context)
 		{
 			// Check the job exists and we can access it
-			IJob? job = await _jobService.GetJobAsync(new JobId(request.JobId));
+			IJob? job = await _jobService.GetJobAsync(JobId.Parse(request.JobId));
 			if (job == null)
 			{
 				throw new StructuredRpcException(StatusCode.NotFound, "Job {JobId} not found", request.JobId);
@@ -803,7 +800,7 @@ namespace Horde.Build.Server
 		/// <returns>Information about the new agent</returns>
 		public override async Task<Empty> UpdateStep(RpcUpdateStepRequest request, ServerCallContext context)
 		{
-			IJob? job = await _jobService.GetJobAsync(new JobId(request.JobId));
+			IJob? job = await _jobService.GetJobAsync(JobId.Parse(request.JobId));
 			if (job == null)
 			{
 				throw new StructuredRpcException(StatusCode.NotFound, "Job {JobId} not found", request.JobId);
@@ -839,7 +836,7 @@ namespace Horde.Build.Server
 		/// <returns>Information about the step</returns>
 		public override async Task<RpcGetStepResponse> GetStep(GetStepRequest request, ServerCallContext context)
 		{
-			IJob job = await GetJobAsync(new JobId(request.JobId));
+			IJob job = await GetJobAsync(JobId.Parse(request.JobId));
 			IJobStepBatch batch = AuthorizeBatch(job, request.BatchId.ToSubResourceId(), context);
 
 			SubResourceId stepId = request.StepId.ToSubResourceId();
@@ -892,7 +889,7 @@ namespace Horde.Build.Server
 				newLabels.Add(newLabel);
 			}
 
-			JobId jobIdValue = new JobId(request.JobId);
+			JobId jobIdValue = JobId.Parse(request.JobId);
 			for (; ; )
 			{
 				IJob? job = await _jobService.GetJobAsync(jobIdValue);
@@ -933,7 +930,7 @@ namespace Horde.Build.Server
 			foreach (CreateEventRequest createEvent in request.Events)
 			{
 				NewLogEventData newEvent = new NewLogEventData();
-				newEvent.LogId = new LogId(createEvent.LogId);
+				newEvent.LogId = LogId.Parse(createEvent.LogId);
 				newEvent.Severity = createEvent.Severity;
 				newEvent.LineIndex = createEvent.LineIndex;
 				newEvent.LineCount = createEvent.LineCount;
@@ -951,7 +948,7 @@ namespace Horde.Build.Server
 		/// <returns>Information about the new agent</returns>
 		public override async Task<Empty> WriteOutput(WriteOutputRequest request, ServerCallContext context)
 		{
-			ILogFile? logFile = await _logFileService.GetCachedLogFileAsync(new LogId(request.LogId), context.CancellationToken);
+			ILogFile? logFile = await _logFileService.GetCachedLogFileAsync(LogId.Parse(request.LogId), context.CancellationToken);
 			if (logFile == null)
 			{
 				throw new StructuredRpcException(StatusCode.NotFound, "Resource not found");
@@ -1044,7 +1041,7 @@ namespace Horde.Build.Server
 			}
 
 			// Get the job and step
-			IJob job = await GetJobAsync(new JobId(metadata.JobId));
+			IJob job = await GetJobAsync(JobId.Parse(metadata.JobId));
 			AuthorizeBatch(job, metadata.BatchId.ToSubResourceId(), context);
 
 			IJobStep? step;
@@ -1081,7 +1078,7 @@ namespace Horde.Build.Server
 			{
 				UploadTestDataRequest request = reader.Current;
 
-				JobId jobId = new JobId(request.JobId);
+				JobId jobId = JobId.Parse(request.JobId);
 				if (job == null)
 				{					
 					job = await _jobService.GetJobAsync(jobId);
@@ -1135,7 +1132,7 @@ namespace Horde.Build.Server
 		/// <returns></returns>
 		public override async Task<CreateReportResponse> CreateReport(CreateReportRequest request, ServerCallContext context)
 		{
-			IJob job = await GetJobAsync(new JobId(request.JobId));
+			IJob job = await GetJobAsync(JobId.Parse(request.JobId));
 			if (!_globalConfig.Value.TryGetStream(job.StreamId, out StreamConfig? streamConfig))
 			{
 				throw new StructuredRpcException(StatusCode.NotFound, "Stream {StreamId} not found", job.StreamId);
diff --git a/Engine/Source/Programs/Horde/Horde.Build/Startup.cs b/Engine/Source/Programs/Horde/Horde.Build/Startup.cs
--- a/Engine/Source/Programs/Horde/Horde.Build/Startup.cs
+++ b/Engine/Source/Programs/Horde/Horde.Build/Startup.cs
@@ -97,11 +97,6 @@ using StatusCode = Grpc.Core.StatusCode;
 
 namespace Horde.Build
 {
-	using JobId = ObjectId<IJob>;
-	using LeaseId = ObjectId<ILease>;
-	using LogId = ObjectId<ILogFile>;
-	using UserId = ObjectId<IUser>;
-
 	class Startup
 	{
 		static Startup()
@@ -960,7 +955,7 @@ namespace Horde.Build
 				BsonSerializer.RegisterSerializer(new ConditionSerializer());
 				BsonSerializer.RegisterSerializationProvider(new BsonSerializationProvider());
 				BsonSerializer.RegisterSerializationProvider(new StringIdBsonSerializationProvider());
-				BsonSerializer.RegisterSerializationProvider(new ObjectIdSerializationProvider());
+				BsonSerializer.RegisterSerializationProvider(new ObjectIdBsonSerializationProvider());
 			}
 		}
 
diff --git a/Engine/Source/Programs/Horde/Horde.Build/Streams/IStream.cs b/Engine/Source/Programs/Horde/Horde.Build/Streams/IStream.cs
--- a/Engine/Source/Programs/Horde/Horde.Build/Streams/IStream.cs
+++ b/Engine/Source/Programs/Horde/Horde.Build/Streams/IStream.cs
@@ -16,9 +16,6 @@ using Horde.Build.Utilities;
 
 namespace Horde.Build.Streams
 {
-	using JobId = ObjectId<IJob>;
-	using UserId = ObjectId<IUser>;
-
 	/// <summary>
 	/// Exception thrown when stream validation fails
 	/// </summary>
diff --git a/Engine/Source/Programs/Horde/Horde.Build/Streams/StreamCollection.cs b/Engine/Source/Programs/Horde/Horde.Build/Streams/StreamCollection.cs
--- a/Engine/Source/Programs/Horde/Horde.Build/Streams/StreamCollection.cs
+++ b/Engine/Source/Programs/Horde/Horde.Build/Streams/StreamCollection.cs
@@ -17,9 +17,6 @@ using MongoDB.Driver;
 
 namespace Horde.Build.Streams
 {
-	using JobId = ObjectId<IJob>;
-	using UserId = ObjectId<IUser>;
-
 	/// <summary>
 	/// Collection of stream documents
 	/// </summary>
@@ -321,7 +318,7 @@ namespace Horde.Build.Streams
 				{
 					int stateIndex = newStepStates.FindIndex(x => x.Name == updateState.Name);
 
-					UserId? pausedByUserId = updateState.PausedByUserId != null ? new UserId(updateState.PausedByUserId) : null;
+					UserId? pausedByUserId = updateState.PausedByUserId != null ? UserId.Parse(updateState.PausedByUserId) : null;
 
 					if (stateIndex == -1)
 					{
diff --git a/Engine/Source/Programs/Horde/Horde.Build/Tasks/ConformTaskSource.cs b/Engine/Source/Programs/Horde/Horde.Build/Tasks/ConformTaskSource.cs
--- a/Engine/Source/Programs/Horde/Horde.Build/Tasks/ConformTaskSource.cs
+++ b/Engine/Source/Programs/Horde/Horde.Build/Tasks/ConformTaskSource.cs
@@ -23,10 +23,6 @@ using Microsoft.Extensions.Options;
 
 namespace Horde.Build.Tasks
 {
-	using JobId = ObjectId<IJob>;
-	using LeaseId = ObjectId<ILease>;
-	using LogId = ObjectId<ILogFile>;
-
 	/// <summary>
 	/// Generates tasks telling agents to sync their workspaces
 	/// </summary>
@@ -67,7 +63,7 @@ namespace Horde.Build.Tasks
 			_logger = logger;
 			_tickConformList = clock.AddSharedTicker<ConformTaskSource>(TimeSpan.FromMinutes(1.0), CleanConformListAsync, logger);
 
-			OnLeaseStartedProperties.Add(nameof(ConformTask.LogId), x => new LogId(x.LogId));
+			OnLeaseStartedProperties.Add(nameof(ConformTask.LogId), x => LogId.Parse(x.LogId));
 		}
 
 		/// <inheritdoc/>
diff --git a/Engine/Source/Programs/Horde/Horde.Build/Tasks/RestartTaskSource.cs b/Engine/Source/Programs/Horde/Horde.Build/Tasks/RestartTaskSource.cs
--- a/Engine/Source/Programs/Horde/Horde.Build/Tasks/RestartTaskSource.cs
+++ b/Engine/Source/Programs/Horde/Horde.Build/Tasks/RestartTaskSource.cs
@@ -14,10 +14,6 @@ using HordeCommon.Rpc.Tasks;
 
 namespace Horde.Build.Tasks
 {
-	using JobId = ObjectId<IJob>;
-	using LeaseId = ObjectId<ILease>;
-	using LogId = ObjectId<ILogFile>;
-
 	class RestartTaskSource : TaskSourceBase<RestartTask>
 	{
 		public override string Type => "Restart";
@@ -30,7 +26,7 @@ namespace Horde.Build.Tasks
 		{
 			_logService = logService;
 
-			OnLeaseStartedProperties.Add(nameof(RestartTask.LogId), x => new LogId(x.LogId));
+			OnLeaseStartedProperties.Add(nameof(RestartTask.LogId), x => LogId.Parse(x.LogId));
 		}
 
 		public override async Task<Task<AgentLease?>> AssignLeaseAsync(IAgent agent, CancellationToken cancellationToken)
diff --git a/Engine/Source/Programs/Horde/Horde.Build/Tasks/ShutdownTaskSource.cs b/Engine/Source/Programs/Horde/Horde.Build/Tasks/ShutdownTaskSource.cs
--- a/Engine/Source/Programs/Horde/Horde.Build/Tasks/ShutdownTaskSource.cs
+++ b/Engine/Source/Programs/Horde/Horde.Build/Tasks/ShutdownTaskSource.cs
@@ -14,10 +14,6 @@ using HordeCommon.Rpc.Tasks;
 
 namespace Horde.Build.Tasks
 {
-	using JobId = ObjectId<IJob>;
-	using LeaseId = ObjectId<ILease>;
-	using LogId = ObjectId<ILogFile>;
-
 	class ShutdownTaskSource : TaskSourceBase<ShutdownTask>
 	{
 		public override string Type => "Shutdown";
@@ -29,7 +25,7 @@ namespace Horde.Build.Tasks
 		public ShutdownTaskSource(ILogFileService logService)
 		{
 			_logService = logService;
-			OnLeaseStartedProperties.Add(nameof(ShutdownTask.LogId), x => new LogId(x.LogId));
+			OnLeaseStartedProperties.Add(nameof(ShutdownTask.LogId), x => LogId.Parse(x.LogId));
 		}
 
 		public override async Task<Task<AgentLease?>> AssignLeaseAsync(IAgent agent, CancellationToken cancellationToken)
diff --git a/Engine/Source/Programs/Horde/Horde.Build/Tasks/UpgradeTaskSource.cs b/Engine/Source/Programs/Horde/Horde.Build/Tasks/UpgradeTaskSource.cs
--- a/Engine/Source/Programs/Horde/Horde.Build/Tasks/UpgradeTaskSource.cs
+++ b/Engine/Source/Programs/Horde/Horde.Build/Tasks/UpgradeTaskSource.cs
@@ -18,10 +18,6 @@ using Microsoft.Extensions.Options;
 
 namespace Horde.Build.Tasks
 {
-	using JobId = ObjectId<IJob>;
-	using LeaseId = ObjectId<ILease>;
-	using LogId = ObjectId<ILogFile>;
-
 	class UpgradeTaskSource : TaskSourceBase<UpgradeTask>
 	{
 		public override string Type => "Upgrade";
@@ -42,7 +38,7 @@ namespace Horde.Build.Tasks
 			_serverSettings = serverSettings;
 			_clock = clock;
 
-			OnLeaseStartedProperties.Add(nameof(UpgradeTask.LogId), x => new LogId(x.LogId));
+			OnLeaseStartedProperties.Add(nameof(UpgradeTask.LogId), x => LogId.Parse(x.LogId));
 		}
 
 		public override async Task<Task<AgentLease?>> AssignLeaseAsync(IAgent agent, CancellationToken cancellationToken)
diff --git a/Engine/Source/Programs/Horde/Horde.Build/Tools/ITool.cs b/Engine/Source/Programs/Horde/Horde.Build/Tools/ITool.cs
--- a/Engine/Source/Programs/Horde/Horde.Build/Tools/ITool.cs
+++ b/Engine/Source/Programs/Horde/Horde.Build/Tools/ITool.cs
@@ -10,8 +10,6 @@ using System.Threading.Tasks;
 
 namespace Horde.Build.Tools
 {
-	using ToolDeploymentId = ObjectId<IToolDeployment>;
-
 	/// <summary>
 	/// Describes a standalone, external tool hosted and deployed by Horde. Provides basic functionality for performing
 	/// gradual roll-out, versioning, etc...
diff --git a/Engine/Source/Programs/Horde/Horde.Build/Tools/ToolDeploymentId.cs b/Engine/Source/Programs/Horde/Horde.Build/Tools/ToolDeploymentId.cs
new file mode 100644
--- /dev/null
+++ b/Engine/Source/Programs/Horde/Horde.Build/Tools/ToolDeploymentId.cs
@@ -0,0 +1,39 @@
+// Copyright Epic Games, Inc. All Rights Reserved.
+
+using System.ComponentModel;
+using EpicGames.Core;
+using Horde.Build.Utilities;
+using MongoDB.Bson;
+
+namespace Horde.Build.Tools
+{
+	/// <summary>
+	/// Identifier for a tool deployment
+	/// </summary>
+	/// <param name="Id">Id to construct from</param>
+	[JsonSchemaString]
+	[TypeConverter(typeof(ObjectIdTypeConverter<ToolDeploymentId, ToolDeploymentIdConverter>))]
+	[ObjectIdConverter(typeof(ToolDeploymentIdConverter))]
+	public record struct ToolDeploymentId(ObjectId Id)
+	{
+		/// <summary>
+		/// Creates a new <see cref="ToolDeploymentId"/>
+		/// </summary>
+		public static ToolDeploymentId GenerateNewId() => new ToolDeploymentId(ObjectId.GenerateNewId());
+
+		/// <inheritdoc/>
+		public override string ToString() => Id.ToString();
+	}
+
+	/// <summary>
+	/// Converter to and from <see cref="ObjectId"/> instances.
+	/// </summary>
+	class ToolDeploymentIdConverter : ObjectIdConverter<ToolDeploymentId>
+	{
+		/// <inheritdoc/>
+		public override ToolDeploymentId FromObjectId(ObjectId id) => new ToolDeploymentId(id);
+
+		/// <inheritdoc/>
+		public override ObjectId ToObjectId(ToolDeploymentId value) => value.Id;
+	}
+}
diff --git a/Engine/Source/Programs/Horde/Horde.Build/Ugs/UgsController.cs b/Engine/Source/Programs/Horde/Horde.Build/Ugs/UgsController.cs
--- a/Engine/Source/Programs/Horde/Horde.Build/Ugs/UgsController.cs
+++ b/Engine/Source/Programs/Horde/Horde.Build/Ugs/UgsController.cs
@@ -16,9 +16,6 @@ using MongoDB.Bson;
 
 namespace Horde.Build.Ugs
 {
-	using LogId = ObjectId<ILogFile>;
-	using UserId = ObjectId<IUser>;
-
 	/// <summary>
 	/// Controller for the /api/v1/issues endpoint
 	/// </summary>
diff --git a/Engine/Source/Programs/Horde/Horde.Build/Users/IUser.cs b/Engine/Source/Programs/Horde/Horde.Build/Users/IUser.cs
--- a/Engine/Source/Programs/Horde/Horde.Build/Users/IUser.cs
+++ b/Engine/Source/Programs/Horde/Horde.Build/Users/IUser.cs
@@ -7,9 +7,6 @@ using MongoDB.Bson;
 
 namespace Horde.Build.Users
 {
-	using JobId = ObjectId<IJob>;
-	using UserId = ObjectId<IUser>;
-
 	/// <summary>
 	/// Known user ids
 	/// </summary>
@@ -18,7 +15,7 @@ namespace Horde.Build.Users
 		/// <summary>
 		/// The system user. Used for automated processes.
 		/// </summary>
-		public static UserId System { get; } = new UserId("6170b423b94a2c7c2d6b6f87");
+		public static UserId System { get; } = UserId.Parse("6170b423b94a2c7c2d6b6f87");
 	}
 
 	/// <summary>
diff --git a/Engine/Source/Programs/Horde/Horde.Build/Users/UserController.cs b/Engine/Source/Programs/Horde/Horde.Build/Users/UserController.cs
--- a/Engine/Source/Programs/Horde/Horde.Build/Users/UserController.cs
+++ b/Engine/Source/Programs/Horde/Horde.Build/Users/UserController.cs
@@ -10,9 +10,6 @@ using Microsoft.AspNetCore.Mvc;
 
 namespace Horde.Build.Users
 {
-	using JobId = ObjectId<IJob>;
-	using UserId = ObjectId<IUser>;
-
 	/// <summary>
 	/// Controller for the /api/v1/user endpoint
 	/// </summary>
@@ -76,7 +73,7 @@ namespace Horde.Build.Users
 			{
 				return BadRequest("Current user does not have a registered profile");
 			}
-			await UserCollection.UpdateSettingsAsync(userId.Value, request.EnableExperimentalFeatures, request.DashboardSettings?.ToBsonValue(), request.AddPinnedJobIds?.Select(x => new JobId(x)), request.RemovePinnedJobIds?.Select(x => new JobId(x)));
+			await UserCollection.UpdateSettingsAsync(userId.Value, request.EnableExperimentalFeatures, request.DashboardSettings?.ToBsonValue(), request.AddPinnedJobIds?.Select(x => JobId.Parse(x)), request.RemovePinnedJobIds?.Select(x => JobId.Parse(x)));
 			return Ok();
 		}
 	}
diff --git a/Engine/Source/Programs/Horde/Horde.Build/Users/UserId.cs b/Engine/Source/Programs/Horde/Horde.Build/Users/UserId.cs
new file mode 100644
--- /dev/null
+++ b/Engine/Source/Programs/Horde/Horde.Build/Users/UserId.cs
@@ -0,0 +1,62 @@
+// Copyright Epic Games, Inc. All Rights Reserved.
+
+using System.ComponentModel;
+using EpicGames.Core;
+using Horde.Build.Utilities;
+using MongoDB.Bson;
+
+namespace Horde.Build.Users
+{
+	/// <summary>
+	/// Identifier for a user
+	/// </summary>
+	/// <param name="Id">Id to construct from</param>
+	[JsonSchemaString]
+	[TypeConverter(typeof(ObjectIdTypeConverter<UserId, UserIdConverter>))]
+	[ObjectIdConverter(typeof(UserIdConverter))]
+	public record struct UserId(ObjectId Id)
+	{
+		/// <summary>
+		/// Constant value for empty user id
+		/// </summary>
+		public static UserId Empty { get; } = new UserId(ObjectId.Empty);
+
+		/// <summary>
+		/// Creates a new <see cref="UserId"/>
+		/// </summary>
+		public static UserId GenerateNewId() => new UserId(ObjectId.GenerateNewId());
+
+		/// <inheritdoc cref="ObjectId.Parse(System.String)"/>
+		public static UserId Parse(string text) => new UserId(ObjectId.Parse(text));
+
+		/// <inheritdoc cref="ObjectId.TryParse(System.String, out ObjectId)"/>
+		public static bool TryParse(string text, out UserId id)
+		{
+			if (ObjectId.TryParse(text, out ObjectId objectId))
+			{
+				id = new UserId(objectId);
+				return true;
+			}
+			else
+			{
+				id = default;
+				return false;
+			}
+		}
+
+		/// <inheritdoc/>
+		public override string ToString() => Id.ToString();
+	}
+
+	/// <summary>
+	/// Converter to and from <see cref="ObjectId"/> instances.
+	/// </summary>
+	class UserIdConverter : ObjectIdConverter<UserId>
+	{
+		/// <inheritdoc/>
+		public override UserId FromObjectId(ObjectId id) => new UserId(id);
+
+		/// <inheritdoc/>
+		public override ObjectId ToObjectId(UserId value) => value.Id;
+	}
+}
diff --git a/Engine/Source/Programs/Horde/Horde.Build/Users/UsersController.cs b/Engine/Source/Programs/Horde/Horde.Build/Users/UsersController.cs
--- a/Engine/Source/Programs/Horde/Horde.Build/Users/UsersController.cs
+++ b/Engine/Source/Programs/Horde/Horde.Build/Users/UsersController.cs
@@ -10,8 +10,6 @@ using Microsoft.AspNetCore.Mvc;
 
 namespace Horde.Build.Users
 {
-	using UserId = ObjectId<IUser>;
-
 	/// <summary>
 	/// Controller for the /api/v1/users endpoint
 	/// </summary>
@@ -87,7 +85,7 @@ namespace Horde.Build.Users
 			UserId[]? userIds = null;
 			if (ids != null && ids.Length > 0)
 			{
-				userIds = ids.Select(x => new UserId(x)).ToArray();
+				userIds = ids.Select(x => UserId.Parse(x)).ToArray();
 			}
 
 			List<IUser> users = await UserCollection.FindUsersAsync(userIds, nameRegex, index, count);
diff --git a/Engine/Source/Programs/Horde/Horde.Build/Utilities/HordeClaimTypes.cs b/Engine/Source/Programs/Horde/Horde.Build/Utilities/HordeClaimTypes.cs
--- a/Engine/Source/Programs/Horde/Horde.Build/Utilities/HordeClaimTypes.cs
+++ b/Engine/Source/Programs/Horde/Horde.Build/Utilities/HordeClaimTypes.cs
@@ -6,8 +6,6 @@ using Horde.Build.Users;
 
 namespace Horde.Build.Utilities
 {
-	using UserId = ObjectId<IUser>;
-
 	/// <summary>
 	/// Claim types that are specific to horde
 	/// </summary>
@@ -84,7 +82,7 @@ namespace Horde.Build.Utilities
 			}
 			else
 			{
-				return new UserId(idValue);
+				return UserId.Parse(idValue);
 			}
 		}
 
diff --git a/Engine/Source/Programs/Horde/Horde.Build/Utilities/HordeControllerBase.cs b/Engine/Source/Programs/Horde/Horde.Build/Utilities/HordeControllerBase.cs
--- a/Engine/Source/Programs/Horde/Horde.Build/Utilities/HordeControllerBase.cs
+++ b/Engine/Source/Programs/Horde/Horde.Build/Utilities/HordeControllerBase.cs
@@ -16,10 +16,6 @@ using Microsoft.Extensions.Logging;
 
 namespace Horde.Build.Utilities
 {
-	using JobId = ObjectId<IJob>;
-	using LeaseId = ObjectId<ILease>;
-	using ToolDeploymentId = ObjectId<IToolDeployment>;
-	
 	/// <summary>
 	/// Base class for Horde controllers
 	/// </summary>
diff --git a/Engine/Source/Programs/Horde/Horde.Build/Utilities/ObjectIdHelpers.cs b/Engine/Source/Programs/Horde/Horde.Build/Utilities/ObjectIdHelpers.cs
--- a/Engine/Source/Programs/Horde/Horde.Build/Utilities/ObjectIdHelpers.cs
+++ b/Engine/Source/Programs/Horde/Horde.Build/Utilities/ObjectIdHelpers.cs
@@ -17,10 +17,5 @@ namespace Horde.Build.Utilities
 				return ObjectId.Parse(text);
 			}
 		}
-
-		public static ObjectId<T> ToObjectId<T>(this string text)
-		{
-			return new ObjectId<T>(text.ToObjectId());
-		}
 	}
 }
diff --git a/Engine/Source/Programs/Horde/Horde.Build/Utilities/StringIdConverter.cs b/Engine/Source/Programs/Horde/Horde.Build/Utilities/StringIdConverter.cs
--- a/Engine/Source/Programs/Horde/Horde.Build/Utilities/StringIdConverter.cs
+++ b/Engine/Source/Programs/Horde/Horde.Build/Utilities/StringIdConverter.cs
@@ -16,7 +16,7 @@ using ProtoBuf;
 namespace Horde.Build.Utilities
 {
 	/// <summary>
-	/// Base class for converting to and from types containing an object id. Useful pattern for reducing boilerplate with strongly typed records.
+	/// Base class for converting to and from types containing a <see cref="StringId"/>. Useful pattern for reducing boilerplate with strongly typed records.
 	/// </summary>
 	/// <typeparam name="T"></typeparam>
 	abstract class StringIdConverter<T> where T : struct
@@ -27,13 +27,13 @@ namespace Horde.Build.Utilities
 		public abstract StringId ToStringId(T value);
 
 		/// <summary>
-		/// Constructs a type from an object id
+		/// Constructs a type from a <see cref="StringId"/>
 		/// </summary>
 		public abstract T FromStringId(StringId id);
 	}
 
 	/// <summary>
-	/// Attribute declaring an object id converter for a particular type
+	/// Attribute declaring a <see cref="StringIdConverter{T}"/> for a particular type
 	/// </summary>
 	[AttributeUsage(AttributeTargets.Struct)]
 	sealed class StringIdConverterAttribute : Attribute
diff --git a/Engine/Source/Programs/Horde/Horde.Build/Utilities/TracingExtensions.cs b/Engine/Source/Programs/Horde/Horde.Build/Utilities/TracingExtensions.cs
--- a/Engine/Source/Programs/Horde/Horde.Build/Utilities/TracingExtensions.cs
+++ b/Engine/Source/Programs/Horde/Horde.Build/Utilities/TracingExtensions.cs
@@ -13,25 +13,6 @@ namespace Horde.Build.Utilities
 	/// </summary>
 	public static class OpenTracingSpanExtensions
 	{
-		/// <summary>Set a key:value tag on the span</summary>
-		/// <returns>This span instance, for chaining</returns>
-		public static ISpan SetTag<T>(this ISpan span, string key, ObjectId<T>? value)
-		{
-			if (value != null)
-			{
-				span.SetTag(key, value.ToString());
-			}
-			return span;
-		}
-		
-		/// <summary>Set a key:value tag on the span</summary>
-		/// <returns>This span instance, for chaining</returns>
-		public static ISpan SetTag<T>(this ISpan span, string key, ObjectId<T> value)
-		{
-			span.SetTag(key, value.ToString());
-			return span;
-		}
-		
 		/// <summary>Set a key:value tag on the span</summary>
 		/// <returns>This span instance, for chaining</returns>
 		public static ISpan SetTag(this ISpan span, string key, ContentHash value)
@@ -48,17 +29,6 @@ namespace Horde.Build.Utilities
 			return span;
 		}
 		
-		/// <summary>Set a key:value tag on the span</summary>
-		/// <returns>This span instance, for chaining</returns>
-		public static ISpan SetTag<T>(this ISpan span, string key, IEnumerable<ObjectId<T>>? value)
-		{
-			if (value != null)
-			{
-				span.SetTag(key, String.Join(',', value));
-			}
-			return span;
-		}
-		
 		/// <summary>Set a key:value tag on the span</summary>
 		/// <returns>This span instance, for chaining</returns>
 		public static ISpan SetTag(this ISpan span, string key, int? value)
