Horde: Collect metrics for XGE builds invocations

Outputs a structured log message with a summary after each job execution.

diff --git a/Engine/Source/Programs/Horde/Horde.Agent.Tests/XgeMetadataExtractorTest.cs b/Engine/Source/Programs/Horde/Horde.Agent.Tests/XgeMetadataExtractorTest.cs
new file mode 100644
--- /dev/null
+++ b/Engine/Source/Programs/Horde/Horde.Agent.Tests/XgeMetadataExtractorTest.cs
@@ -0,0 +1,91 @@
+// Copyright Epic Games, Inc. All Rights Reserved.
+
+using System;
+using System.Text;
+using Horde.Agent.Utility;
+using Microsoft.VisualStudio.TestTools.UnitTesting;
+
+namespace Horde.Agent.Tests;
+
+[TestClass]
+public class XgeMetadataExtractorTest
+{
+	[TestMethod]
+	public void GetSummaryFromJson()
+	{
+		byte[] jsonData = Encoding.UTF8.GetBytes(JsonFromIbMon);
+
+		XgeTaskMetadataSummary summary = XgeMetadataExtractor.GetSummaryFromJson(jsonData)!;
+		Assert.AreEqual("UAT_XGE", summary.Title);
+		
+		Assert.AreEqual(1, summary.LocalTaskCount);
+		Assert.AreEqual(1, summary.RemoteTaskCount);
+		
+		Assert.AreEqual(TimeSpan.FromMilliseconds(2471), summary.TotalLocalTaskDuration);
+		Assert.AreEqual(TimeSpan.FromMilliseconds(11706), summary.TotalRemoteTaskDuration);
+	}
+
+	private const string JsonFromIbMon = @"
+{
+ ""Metadata"":{
+	""ID"":""{BF63E95E-F7F7-4551-8F6F-C85D62FF62DD}"",
+	""Caption"":""UAT_XGE"",
+	""ComputerName"":""MY-COMPUTER"",
+	""UserName"":""foo.bar"",
+	""ErrorCount"":0,
+	""WarningCount"":0,
+	""SystemErrorCount"":0,
+	""SystemWarningCount"":2,
+	""StartTime"":""20-01-2023 13;40;22"",
+	""Duration"":""12:07:14 AM"",
+	""Cancelled"":""false"",
+	""ProductVersion"":2005956,
+	""ProductBuildNum"":4956,
+	""ProductVersionText"":""10.1.7 (build 4956)"",
+	""SystemAutoRecoveryCount"":0,
+	""ProjectFormat"":"""",
+	""LoggingLevel"":4,
+	""PrimaryFilename"":""D:\\Workspaces\\ue5-main\\Engine\\Programs\\AutomationTool\\Saved\\Logs\\UAT_XGE.xml"",
+	""Incomplete"":""false"",
+	""OverrideBuildStatusToSuccess"":""false""
+ },
+ ""Tasks"":[
+  {
+   ""TaskID"":44,
+   ""Name"":""GPUSkinVertexFactory.ispc"",
+   ""BeginTime"":1823,
+   ""GUID"":""{2A2D4ECD-A1F2-4045-8D12-45C336585FCD}"",
+   ""CmdLine"":"" @\""../Intermediate/Build/Win64/UnrealEditor/Development/Engine/GPUSkinVertexFactory.ispc.generated.dummy.h.response\"""",
+   ""AppName"":""D:\\Workspaces\\ue5-main\\Engine\\Source\\ThirdParty\\Intel\\ISPC\\bin\\Windows\\ispc.exe"",
+   ""WorkerID"":""MY-COMPUTER\\8"",
+   ""Remote"":""false"",
+   ""EndTime"":4294,
+   ""Errors"":0,
+   ""Warnings"":0,
+   ""ExitCodeExists"":""true"",
+   ""ExitCode"":0,
+   ""Double"":""false""
+  },
+  {
+   ""TaskID"":428,
+   ""Name"":""Module.ChaosCore.cpp"",
+   ""BeginTime"":6077,
+   ""GUID"":""{BB171990-E012-4A48-AA7B-E38F10F7F089}"",
+   ""CmdLine"":"" @..\\Intermediate\\Build\\Win64\\UnrealEditor\\Development\\ChaosCore\\Module.ChaosCore.cpp.obj.response"",
+   ""AppName"":""C:\\Program Files\\Microsoft Visual Studio\\2022\\Professional\\VC\\Tools\\MSVC\\14.31.31103\\bin\\Hostx64\\x64\\cl.exe"",
+   ""WorkerID"":""OTHER-COMPUTER\\8"",
+   ""Remote"":""true"",
+   ""EndTime"":17783,
+   ""Errors"":0,
+   ""Warnings"":0,
+   ""ExitCodeExists"":""true"",
+   ""ExitCode"":0,
+   ""Double"":""false""
+  }
+ ]
+}
+";
+
+}
+
+
diff --git a/Engine/Source/Programs/Horde/Horde.Agent/Execution/JobExecutor.cs b/Engine/Source/Programs/Horde/Horde.Agent/Execution/JobExecutor.cs
--- a/Engine/Source/Programs/Horde/Horde.Agent/Execution/JobExecutor.cs
+++ b/Engine/Source/Programs/Horde/Horde.Agent/Execution/JobExecutor.cs
@@ -199,6 +199,8 @@ namespace Horde.Agent.Execution
 
 		protected Dictionary<string, string> _envVars = new Dictionary<string, string>(StringComparer.OrdinalIgnoreCase);
 
+		private XgeMetadataExtractor? _xgeMetadataExtractor;
+
 		public JobExecutor(JobExecutorOptions options, ILogger logger)
 		{
 			_session = options.Session;
@@ -288,6 +290,23 @@ namespace Horde.Agent.Execution
 			{
 				_additionalArguments.Add($"-set:PreflightChange={_job.PreflightChange}");
 			}
+
+			// Only run XGE metadata extraction on Windows
+			if (RuntimeInformation.IsOSPlatform(OSPlatform.Windows))
+			{
+				DirectoryReference? xgeDir = XgeMetadataExtractor.FindXgeDir();
+				if (xgeDir != null)
+				{
+					_xgeMetadataExtractor = new XgeMetadataExtractor(xgeDir);
+					
+					// Clear any *.ib_mon files from previous builds
+					_xgeMetadataExtractor.ClearLocalIbMonFiles();
+				}
+				else
+				{
+					_logger.LogInformation("Unable to locate XGE directory. Not installed?");
+				}
+			}
 		}
 
 		public static bool IsUserAdministrator()
@@ -339,6 +358,17 @@ namespace Horde.Agent.Execution
 
 		public virtual Task FinalizeAsync(ILogger logger, CancellationToken cancellationToken)
 		{
+			if (_xgeMetadataExtractor != null)
+			{
+				// Collect any new *.ib_mon files and generate a summary for each build invocation, grouped by title.
+				// Title will segregate for example UnrealBuildTool from invocations coming from the engine/cooks.
+				foreach (XgeTaskMetadataSummary s in _xgeMetadataExtractor.GetSummariesForIbMonFiles())
+				{
+					_logger.LogInformation("XGE build invocation {Title} with {LocalTaskCount} local tasks for {LocalTaskDuration} secs total and {RemoteTaskCount} remote tasks for {RemoteTaskDuration} secs total",
+						s.Title, s.LocalTaskCount, s.TotalLocalTaskDuration, s.RemoteTaskCount, s.TotalRemoteTaskDuration);
+				}
+			}
+			
 			return Task.CompletedTask;
 		}
 
diff --git a/Engine/Source/Programs/Horde/Horde.Agent/Utility/XgeMetadataExtractor.cs b/Engine/Source/Programs/Horde/Horde.Agent/Utility/XgeMetadataExtractor.cs
new file mode 100644
--- /dev/null
+++ b/Engine/Source/Programs/Horde/Horde.Agent/Utility/XgeMetadataExtractor.cs
@@ -0,0 +1,416 @@
+// Copyright Epic Games, Inc. All Rights Reserved.
+
+using System;
+using System.Collections.Generic;
+using System.Diagnostics;
+using System.IO;
+using System.Linq;
+using System.Runtime.Versioning;
+using System.Text.Json;
+using System.Text.Json.Serialization;
+using EpicGames.Core;
+using Microsoft.Win32;
+
+namespace Horde.Agent.Utility;
+
+/// <summary>
+/// Summary of an XGE run
+/// </summary>
+public class XgeTaskMetadataSummary
+{
+	/// <summary>
+	/// Title of the XGE invocation
+	/// </summary>
+	public string Title { get; }
+	
+	/// <summary>
+	/// Number of tasks executed locally
+	/// </summary>
+	public int LocalTaskCount { get; }
+	
+	/// <summary>
+	/// Number of tasks executed remotely on other machines
+	/// </summary>
+	public int RemoteTaskCount { get; }
+	
+	/// <summary>
+	/// Total time spent executing local tasks
+	/// </summary>
+	public TimeSpan TotalLocalTaskDuration { get; }
+	
+	/// <summary>
+	/// Total time spent executing remote tasks
+	/// </summary>
+	public TimeSpan TotalRemoteTaskDuration { get; }
+
+	/// <summary>
+	/// Constructor
+	/// </summary>
+	/// <param name="title"></param>
+	/// <param name="localTaskCount"></param>
+	/// <param name="remoteTaskCount"></param>
+	/// <param name="totalLocalTaskDuration"></param>
+	/// <param name="totalRemoteTaskDuration"></param>
+	public XgeTaskMetadataSummary(string title, int localTaskCount, int remoteTaskCount, TimeSpan totalLocalTaskDuration, TimeSpan totalRemoteTaskDuration)
+	{
+		Title = title;
+		LocalTaskCount = localTaskCount;
+		RemoteTaskCount = remoteTaskCount;
+		TotalLocalTaskDuration = totalLocalTaskDuration;
+		TotalRemoteTaskDuration = totalRemoteTaskDuration;
+	}
+
+	/// <summary>
+	/// Return a copy of this instance with values incremented by the argument
+	/// </summary>
+	/// <param name="other">The other summary instance to increment with</param>
+	/// <returns></returns>
+	public XgeTaskMetadataSummary Add(XgeTaskMetadataSummary other)
+	{
+		return new XgeTaskMetadataSummary(
+			Title,
+			LocalTaskCount + other.LocalTaskCount,
+			RemoteTaskCount + other.RemoteTaskCount,
+			TotalLocalTaskDuration.Add(other.TotalLocalTaskDuration),
+			TotalRemoteTaskDuration.Add(other.TotalRemoteTaskDuration));
+	}
+
+	/// <inheritdoc/>
+	public override string ToString()
+	{
+		return $"Title={Title} Local({LocalTaskCount} tasks for {TotalLocalTaskDuration.TotalMilliseconds} ms) Remote({RemoteTaskCount} tasks for {TotalRemoteTaskDuration.TotalMilliseconds} ms)";
+	}
+}
+
+/// <summary>
+/// Extracts metadata about executed XGE / IncrediBuild tasks by locating local *.ib_mon files
+/// </summary>
+public class XgeMetadataExtractor
+{
+	private const string XgConsoleExeWin = "xgConsole.exe";
+	private const string RegistryKey = "SOFTWARE\\Xoreax\\IncrediBuild\\Builder";
+
+	private readonly DirectoryReference _xgeDir; 
+	private readonly FileReference _buildMonitorExePath; 
+	
+	/// <summary>
+	/// Constructor
+	/// </summary>
+	/// <param name="xgeDir">Base directory for XGE</param>
+	/// <exception cref="ArgumentException">If directory is not a valid base dir for XGE</exception>
+	public XgeMetadataExtractor(DirectoryReference xgeDir)
+	{
+		if (!OperatingSystem.IsWindows())
+		{
+			throw new Exception("Only Windows is supported for XGE metadata extraction");
+		}
+
+		if (!IsXgeDirectory(xgeDir.FullName))
+		{
+			throw new ArgumentException("Path is not XGE base directory: " + xgeDir.FullName);
+		}
+		
+		string buildMonitorExePath = Path.Join(xgeDir.FullName, "BuildMonitor.exe");
+		if (!File.Exists(buildMonitorExePath))
+		{
+			throw new FileNotFoundException("BuildMonitor.exe not found in XGE base dir " + xgeDir.FullName);
+		}
+		
+		_xgeDir = xgeDir;
+		_buildMonitorExePath = new FileReference(buildMonitorExePath);
+	}
+
+	/// <summary>
+	/// Removes all *.ib_mon from local history directory
+	/// </summary>
+	public void ClearLocalIbMonFiles()
+	{
+		List<FileReference> files = GetLocalIbMonFilePaths();
+		foreach (FileReference file in files)
+		{
+			FileReference.Delete(file);	
+		}
+
+		if (GetLocalIbMonFilePaths().Count != 0)
+		{
+			throw new Exception("Unable to remove all *.ib_mon files from history");
+		}
+	}
+	
+	/// <summary>
+	/// Find all local *.ib_mon from history 
+	/// </summary>
+	/// <returns>A list of *.ib_mon files</returns>
+	public List<FileReference> GetLocalIbMonFilePaths()
+	{
+		DirectoryReference localHistoryDir = DirectoryReference.Combine(_xgeDir, "History", "Local");
+		if (!DirectoryReference.Exists(localHistoryDir))
+		{
+			return new List<FileReference>();
+		}
+
+		return localHistoryDir
+			.ToDirectoryInfo()
+			.EnumerateFiles("*.ib_mon", SearchOption.AllDirectories)
+			.Where(x => x.Exists)
+			.Select(x => new FileReference(x))
+			.ToList();
+	}
+	
+	/// <summary>
+	/// Get a summary of task execution metadata for a given .ib_mon file
+	/// </summary>
+	/// <param name="ibMonFile">Path to .ib_mon file</param>
+	/// <returns>A summary of task count and durations</returns>
+	/// <exception cref="FileNotFoundException">If input wasn't found</exception>
+	/// <exception cref="Exception">If JSON file conversion wasn't successful</exception>
+	public XgeTaskMetadataSummary? GetSummaryFromIbMonFile(FileReference ibMonFile)
+	{
+		if (!FileReference.Exists(ibMonFile))
+		{
+			throw new FileNotFoundException("Path to ib_mon file not found " + ibMonFile.FullName);
+		}
+		
+		string tempFile = Path.GetTempFileName();
+		string arguments = ibMonFile.FullName + $" /json=\"{tempFile}\"";
+		string workingDir = _xgeDir.FullName;
+		
+		using ManagedProcessGroup processGroup = new ();
+		using ManagedProcess process = new (processGroup, _buildMonitorExePath.FullName, arguments, workingDir, null, null, ProcessPriorityClass.Normal);
+		process.WaitForExit();
+
+		if (process.ExitCode != 0)
+		{
+			throw new Exception("XGE build monitor executable finished with non-zero exit code: " + process.ExitCode);
+		}
+
+		FileInfo tempFileInfo = new (tempFile);
+		if (!tempFileInfo.Exists || tempFileInfo.Length == 0)
+		{
+			throw new Exception("Generated JSON file does not exist or is empty");
+		}
+
+		byte[] jsonData = File.ReadAllBytes(tempFile);
+		try
+		{
+			tempFileInfo.Delete();
+		}
+		catch (Exception) { /* Ignore any exceptions */ }  
+		
+		return GetSummaryFromJson(jsonData);
+	}
+
+	/// <summary>
+	/// Convert JSON output from .ib_mon file conversion to C# objects
+	/// </summary>
+	/// <param name="data">Path to .ib_mon file</param>
+	/// <returns>A C# object representation</returns>
+	/// <exception cref="Exception">If JSON file conversion wasn't successful</exception>
+	public static XgeTaskMetadataSummary? GetSummaryFromJson(byte[] data)
+	{
+		JsonSerializerOptions jsonOptions = new () { Converters = { new XgeMetadataBoolConverter() } };
+		BuildMonitorData? buildMonitorData = JsonSerializer.Deserialize<BuildMonitorData>(data, jsonOptions);
+		if (buildMonitorData == null)
+		{
+			throw new Exception("Failed reading JSON file converted from .ib_mon file");
+		}
+
+		if (buildMonitorData.Tasks.Count == 0)
+		{
+			return null;
+		}
+
+		int localTaskCount = 0;
+		long localTaskDuration = 0;
+		int remoteTaskCount = 0;
+		long remoteTaskDuration = 0;
+		
+		foreach (XgeTaskMetadata taskMetadata in buildMonitorData.Tasks)
+		{
+			if (!taskMetadata.IsRemote)
+			{
+				localTaskDuration += taskMetadata.EndTimeMs - taskMetadata.BeginTimeMs;
+				localTaskCount += 1;
+			}
+			else
+			{
+				remoteTaskDuration += taskMetadata.EndTimeMs - taskMetadata.BeginTimeMs;
+				remoteTaskCount += 1;
+			}
+		}
+
+		return new XgeTaskMetadataSummary(
+			buildMonitorData.Metadata.Caption,
+			localTaskCount, remoteTaskCount,
+			TimeSpan.FromMilliseconds(localTaskDuration),
+			TimeSpan.FromMilliseconds(remoteTaskDuration));
+	}
+
+	/// <summary>
+	/// Get all summaries for local .ib_mon files, aggregated by XGE title (almost always translating to UBT and cook invocations)
+	/// </summary>
+	/// <returns>A list of summaries</returns>
+	public List<XgeTaskMetadataSummary> GetSummariesForIbMonFiles()
+	{
+		Dictionary<string, XgeTaskMetadataSummary> result = new();
+		foreach (XgeTaskMetadataSummary? s in GetLocalIbMonFilePaths().Select(GetSummaryFromIbMonFile))
+		{
+			if (s == null || (s.LocalTaskCount == 0 && s.RemoteTaskCount == 0)) continue;
+			
+			if (!result.TryGetValue(s.Title, out XgeTaskMetadataSummary? aggregatedSummary))
+			{
+				aggregatedSummary = new XgeTaskMetadataSummary(s.Title, 0, 0, TimeSpan.Zero, TimeSpan.Zero);
+			}
+
+			result[s.Title] = aggregatedSummary.Add(s);
+		}
+
+		return result.Values.ToList();
+	}
+	
+	/// <summary>
+	/// Locate local directory containing XGE / IncrediBuild
+	/// </summary>
+	/// <returns>XGE directory</returns>
+	public static DirectoryReference? FindXgeDir()
+	{
+		// Try to get the path from the registry
+		if (OperatingSystem.IsWindows())
+		{
+			DirectoryReference? dir = GetXgeDirFromRegistry(RegistryView.Registry32) ?? GetXgeDirFromRegistry(RegistryView.Registry64);
+			if (dir != null)
+			{
+				return dir;
+			}
+		}
+
+		// Search the path for it
+		string? pathEnvVar = Environment.GetEnvironmentVariable("PATH");
+		string? dirPath = pathEnvVar?.Split(Path.PathSeparator).ToList().Find(IsXgeDirectory);
+		return dirPath != null ? new DirectoryReference(dirPath) : null;
+	}
+
+	private static bool IsXgeDirectory(string dirPath)
+	{
+		return File.Exists(Path.Join(dirPath, XgConsoleExeWin));
+	}
+
+	[SupportedOSPlatform("windows")]
+	private static DirectoryReference? GetXgeDirFromRegistry(RegistryView view)
+	{
+		try
+		{
+			using RegistryKey baseKey = Microsoft.Win32.RegistryKey.OpenBaseKey(RegistryHive.LocalMachine, view);
+			using RegistryKey? key = baseKey.OpenSubKey(RegistryKey, false);
+			if (key != null)
+			{
+				string? folder = key.GetValue("Folder", null) as string;
+				if (!String.IsNullOrEmpty(folder))
+				{
+					DirectoryReference dir = new (folder);
+					if (IsXgeDirectory(dir.FullName))
+					{
+						return dir;
+					}
+				}
+			}
+		}
+		catch (Exception)
+		{
+			return null;
+		}
+
+		return null;
+	}
+	
+	#region JSON marshaling for .ib_mon file
+	#pragma warning disable CS8618
+	
+	/// <summary>
+	/// Root object of a deserialized *.ib_mon file
+	/// </summary>
+	private class BuildMonitorData
+	{
+		[JsonPropertyName("Metadata")]
+		public XgeMetadata Metadata { get; set; }
+
+		/// <summary>
+		/// List of metadata for XGE tasks
+		/// </summary>
+		[JsonPropertyName("Tasks")]
+		public List<XgeTaskMetadata> Tasks { get; set; }
+	}
+
+	private class XgeMetadata
+	{
+		[JsonPropertyName("ID")]
+		public string Id { get; set; }
+
+		[JsonPropertyName("Caption")]
+		public string Caption { get; set; }
+	}
+	
+	/// <summary>
+	/// Metadata for an XGE task as deserialized from JSON
+	/// </summary>
+	private class XgeTaskMetadata
+	{
+		[JsonPropertyName("TaskID")]
+		public long TaskId { get; set; }
+
+		[JsonPropertyName("Name")]
+		public string Name { get; set; }
+
+		[JsonPropertyName("BeginTime")]
+		public long BeginTimeMs { get; set; }
+
+		[JsonPropertyName("GUID")]
+		public string Guid { get; set; }
+
+		[JsonPropertyName("CmdLine")]
+		public string CmdLine { get; set; }
+
+		[JsonPropertyName("AppName")]
+		public string AppName { get; set; }
+
+		[JsonPropertyName("WorkerID")]
+		public string WorkerId { get; set; }
+
+		[JsonPropertyName("Remote")]
+		[JsonConverter(typeof(XgeMetadataBoolConverter))]
+		public bool IsRemote { get; set; }
+
+		[JsonPropertyName("EndTime")]
+		public long EndTimeMs { get; set; }
+
+		[JsonPropertyName("Errors")]
+		public long Errors { get; set; }
+
+		[JsonPropertyName("Warnings")]
+		public long Warnings { get; set; }
+
+		[JsonPropertyName("ExitCode")]
+		public long ExitCode { get; set; }
+	}
+	
+	private class XgeMetadataBoolConverter : JsonConverter<bool>
+	{
+		public override bool CanConvert(Type t) => t == typeof(bool);
+
+		public override bool Read(ref Utf8JsonReader reader, Type typeToConvert, JsonSerializerOptions options)
+		{
+			if (Boolean.TryParse(reader.GetString(), out bool value)) { return value; }
+			throw new Exception("Cannot unmarshal type bool");
+		}
+
+		public override void Write(Utf8JsonWriter writer, bool value, JsonSerializerOptions options)
+		{
+			string boolString = value ? "true" : "false";
+			JsonSerializer.Serialize(writer, boolString, options);
+		}
+	}
+	
+	#pragma warning restore CS8618
+	#endregion
+}
+
