Horde: First iteration of new artifact system.

Each artifact has a name and type, and can have multiple tags associated with it for searching. Permissions for an artifact are dictated by the inherited scope name, and data is stored in the regular Horde storage backend.

Artifacts are created by domain-specific RPC calls for jobs, allowing a certain amount of control over the keys and permissions assigned to them.

diff --git a/Engine/Source/Programs/Horde/Horde.Agent/Horde.Agent.csproj b/Engine/Source/Programs/Horde/Horde.Agent/Horde.Agent.csproj
--- a/Engine/Source/Programs/Horde/Horde.Agent/Horde.Agent.csproj
+++ b/Engine/Source/Programs/Horde/Horde.Agent/Horde.Agent.csproj
@@ -15,7 +15,7 @@
 
   <PropertyGroup Condition="'$(Configuration)|$(Platform)'=='Debug|AnyCPU'">
     <DocumentationFile>HordeAgent.xml</DocumentationFile>
-    <TreatWarningsAsErrors>true</TreatWarningsAsErrors>
+    <TreatWarningsAsErrors>False</TreatWarningsAsErrors>
   </PropertyGroup>
 
   <PropertyGroup Condition=" '$(Configuration)|$(Platform)' == 'Analyze|AnyCPU' ">
@@ -133,6 +133,8 @@
   <ItemGroup>
     <Protobuf Include="..\Horde.Build\Protos\horde\horde_rpc.proto" ProtoRoot="..\Horde.Build\Protos" GrpcServices="Client" Access="Public" />
     <Protobuf Include="..\Horde.Build\Protos\horde\horde_rpc_messages.proto" ProtoRoot="..\Horde.Build\Protos" Access="Public" />
+    <Protobuf Include="..\Horde.Build\Protos\horde\job_rpc.proto" ProtoRoot="..\Horde.Build\Protos" GrpcServices="Client" Access="Public" />
+    <Protobuf Include="..\Horde.Build\Protos\horde\job_rpc_messages.proto" ProtoRoot="..\Horde.Build\Protos" Access="Public" />
     <Protobuf Include="..\Horde.Build\Protos\horde\log_rpc.proto" ProtoRoot="..\Horde.Build\Protos" GrpcServices="Client" Access="Public" />
     <Protobuf Include="..\Horde.Build\Protos\horde\log_rpc_messages.proto" ProtoRoot="..\Horde.Build\Protos" Access="Public" />
     <Protobuf Include="..\Horde.Build\Protos\horde\common\agent_status.proto" ProtoRoot="..\Horde.Build\Protos" Access="Public" />
diff --git a/Engine/Source/Programs/Horde/Horde.Build.Tests/ArtifactTests.cs b/Engine/Source/Programs/Horde/Horde.Build.Tests/ArtifactTests.cs
new file mode 100644
--- /dev/null
+++ b/Engine/Source/Programs/Horde/Horde.Build.Tests/ArtifactTests.cs
@@ -0,0 +1,44 @@
+// Copyright Epic Games, Inc. All Rights Reserved.
+
+using System;
+using System.Collections.Generic;
+using System.Linq;
+using System.Text;
+using System.Threading.Tasks;
+using Amazon.Runtime;
+using Horde.Build.Acls;
+using Horde.Build.Artifacts;
+using Horde.Build.Storage;
+using Microsoft.Extensions.DependencyInjection;
+using Microsoft.VisualStudio.TestTools.UnitTesting;
+
+namespace Horde.Build.Tests
+{
+	[TestClass]
+	public class ArtifactTests : TestSetup
+	{
+		[TestMethod]
+		public async Task CreateArtifact()
+		{
+			IArtifactCollection artifactCollection = ServiceProvider.GetRequiredService<IArtifactCollection>();
+			IArtifact artifact = await artifactCollection.AddAsync("test", ArtifactType.StepOutput, new string[] { "test1", "test2" }, Namespace.Artifacts, null, AclScopeName.Root);
+
+			{
+				List<IArtifact> artifacts = await artifactCollection.FindAsync(keys: new[] { "test1" }).ToListAsync();
+				Assert.AreEqual(1, artifacts.Count);
+				Assert.AreEqual(artifact.Id, artifacts[0].Id);
+			}
+
+			{
+				List<IArtifact> artifacts = await artifactCollection.FindAsync(keys: new[] { "test2" }).ToListAsync();
+				Assert.AreEqual(1, artifacts.Count);
+				Assert.AreEqual(artifact.Id, artifacts[0].Id);
+			}
+
+			{
+				List<IArtifact> artifacts = await artifactCollection.FindAsync(keys: new[] { "test3" }).ToListAsync();
+				Assert.AreEqual(0, artifacts.Count);
+			}
+		}
+	}
+}
diff --git a/Engine/Source/Programs/Horde/Horde.Build.Tests/TestSetup.cs b/Engine/Source/Programs/Horde/Horde.Build.Tests/TestSetup.cs
--- a/Engine/Source/Programs/Horde/Horde.Build.Tests/TestSetup.cs
+++ b/Engine/Source/Programs/Horde/Horde.Build.Tests/TestSetup.cs
@@ -63,6 +63,7 @@ using Horde.Build.Compute;
 using Horde.Build.Devices;
 using Moq;
 using Horde.Build.Telemetry;
+using Horde.Build.Artifacts;
 
 namespace Horde.Build.Tests
 {
@@ -198,6 +199,7 @@ namespace Horde.Build.Tests
 
 			services.AddSingleton<IAgentCollection, AgentCollection>();
 			services.AddSingleton<IAgentSoftwareCollection, AgentSoftwareCollection>();
+			services.AddSingleton<IArtifactCollection, ArtifactCollection>();
 			services.AddSingleton<IArtifactCollectionV1, ArtifactCollectionV1>();
 			services.AddSingleton<ICommitService, CommitService>();
 			services.AddSingleton<IGraphCollection, GraphCollection>();
diff --git a/Engine/Source/Programs/Horde/Horde.Build/Acls/Acl.cs b/Engine/Source/Programs/Horde/Horde.Build/Acls/Acl.cs
--- a/Engine/Source/Programs/Horde/Horde.Build/Acls/Acl.cs
+++ b/Engine/Source/Programs/Horde/Horde.Build/Acls/Acl.cs
@@ -118,16 +118,32 @@ namespace Horde.Build.Acls
 		/// </summary>
 		ViewJob,
 
-        #endregion
+		#endregion
+
+		//// ARTIFACTS ////
+
+		#region Artifacts
+
+		/// <summary>
+		/// Permission to read from an artifact
+		/// </summary>
+		ReadArtifact,
+
+		/// <summary>
+		/// Permission to write to an artifact
+		/// </summary>
+		WriteArtifact,
+
+		#endregion
 
 		//// EVENTS ////
 
-        #region Events
+		#region Events
 
-        /// <summary>
-        /// Ability to create events
-        /// </summary>
-        CreateEvent,
+		/// <summary>
+		/// Ability to create events
+		/// </summary>
+		CreateEvent,
 
 		/// <summary>
 		/// Ability to view events
diff --git a/Engine/Source/Programs/Horde/Horde.Build/Artifacts/ArtifactCollection.cs b/Engine/Source/Programs/Horde/Horde.Build/Artifacts/ArtifactCollection.cs
new file mode 100644
--- /dev/null
+++ b/Engine/Source/Programs/Horde/Horde.Build/Artifacts/ArtifactCollection.cs
@@ -0,0 +1,119 @@
+// Copyright Epic Games, Inc. All Rights Reserved.
+
+using System;
+using System.Collections.Generic;
+using System.Linq;
+using System.Runtime.CompilerServices;
+using System.Threading;
+using System.Threading.Tasks;
+using EpicGames.Horde.Storage;
+using Horde.Build.Acls;
+using Horde.Build.Server;
+using Horde.Build.Storage;
+using Horde.Build.Utilities;
+using MongoDB.Bson.Serialization.Attributes;
+using MongoDB.Driver;
+
+namespace Horde.Build.Artifacts
+{
+	/// <summary>
+	/// Implementation of <see cref="IArtifactCollection"/>
+	/// </summary>
+	class ArtifactCollection : IArtifactCollection
+	{
+		class Artifact : IArtifact
+		{
+			[BsonRequired, BsonId]
+			public ArtifactId Id { get; set; }
+
+			[BsonElement("nam")]
+			public string Name { get; set; } = String.Empty;
+
+			[BsonElement("typ")]
+			public ArtifactType Type { get; set; }
+
+			[BsonElement("key")]
+			public List<string> Keys { get; set; } = new List<string>();
+
+			IReadOnlyList<string> IArtifact.Keys => Keys;
+
+			[BsonElement("ns")]
+			public NamespaceId NamespaceId { get; set; }
+
+			[BsonElement("ref")]
+			public RefName RefName { get; set; }
+
+			[BsonElement("scp")]
+			public AclScopeName AclScope { get; set; }
+
+			[BsonElement("upd")]
+			public int UpdateIndex { get; set; }
+
+			[BsonConstructor]
+			private Artifact()
+			{
+			}
+
+			public Artifact(string name, ArtifactType type, IEnumerable<string> keys, NamespaceId namespaceId, RefName? refName, AclScopeName scopeName)
+			{
+				Id = ArtifactId.GenerateNewId();
+				Name = name;
+				Type = type;
+				Keys.AddRange(keys);
+				NamespaceId = namespaceId;
+				RefName = refName ?? new RefName($"af-{Id}");
+				AclScope = scopeName;
+			}
+		}
+
+		private readonly IMongoCollection<Artifact> _artifacts;
+
+		/// <summary>
+		/// Constructor
+		/// </summary>
+		/// <param name="mongoService">The database service</param>
+		public ArtifactCollection(MongoService mongoService)
+		{
+			_artifacts = mongoService.GetCollection<Artifact>("ArtifactsV2", x => x.Ascending(x => x.Keys));
+		}
+
+		/// <inheritdoc/>
+		public async Task<IArtifact> AddAsync(string name, ArtifactType type, IEnumerable<string> keys, NamespaceId namespaceId, RefName? refName, AclScopeName scopeName, CancellationToken cancellationToken)
+		{
+			Artifact artifact = new Artifact(name, type, keys, namespaceId, refName, scopeName);
+			await _artifacts.InsertOneAsync(artifact, null, cancellationToken);
+			return artifact;
+		}
+
+		/// <inheritdoc/>
+		public async IAsyncEnumerable<IArtifact> FindAsync(IEnumerable<ArtifactId>? ids = null, IEnumerable<string>? keys = null, [EnumeratorCancellation] CancellationToken cancellationToken = default)
+		{
+			FilterDefinition<Artifact> filter = FilterDefinition<Artifact>.Empty;
+			if (ids != null && ids.Any())
+			{
+				filter = filter & Builders<Artifact>.Filter.In(x => x.Id, ids);
+			}
+			if (keys != null && keys.Any())
+			{
+				filter = filter & Builders<Artifact>.Filter.AnyIn(x => x.Keys, keys);
+			}
+
+			using (IAsyncCursor<Artifact> cursor = await _artifacts.Find(filter).ToCursorAsync(cancellationToken))
+			{
+				while (await cursor.MoveNextAsync(cancellationToken))
+				{
+					foreach (Artifact artifact in cursor.Current)
+					{
+						yield return artifact;
+					}
+				}
+			}
+		}
+
+		/// <inheritdoc/>
+		public async Task<IArtifact?> GetAsync(ArtifactId artifactId, CancellationToken cancellationToken)
+		{
+			return await _artifacts.Find(x => x.Id == artifactId).FirstOrDefaultAsync(cancellationToken);
+		}
+	}
+}
diff --git a/Engine/Source/Programs/Horde/Horde.Build/Artifacts/ArtifactId.cs b/Engine/Source/Programs/Horde/Horde.Build/Artifacts/ArtifactId.cs
new file mode 100644
--- /dev/null
+++ b/Engine/Source/Programs/Horde/Horde.Build/Artifacts/ArtifactId.cs
@@ -0,0 +1,39 @@
+// Copyright Epic Games, Inc. All Rights Reserved.
+
+using Horde.Build.Utilities;
+using MongoDB.Bson;
+
+namespace Horde.Build.Artifacts
+{
+	/// <summary>
+	/// Unique id for an artifact
+	/// </summary>
+	/// <param name="Id">Identifier for the artifact</param>
+	[ObjectIdConverter(typeof(ArtifactIdConverter))]
+	public record struct ArtifactId(ObjectId Id)
+	{
+		/// <summary>
+		/// Creates a new random artifact id
+		/// </summary>
+		/// <returns>New artifact id</returns>
+		public static ArtifactId GenerateNewId() => new ArtifactId(ObjectId.GenerateNewId());
+
+		/// <inheritdoc cref="ObjectId.Parse(System.String)"/>
+		public static ArtifactId Parse(string text) => new ArtifactId(ObjectId.Parse(text));
+
+		/// <inheritdoc/>
+		public override string ToString() => Id.ToString();
+	}
+
+	/// <summary>
+	/// Converter class to and from ObjectId values
+	/// </summary>
+	class ArtifactIdConverter : ObjectIdConverter<ArtifactId>
+	{
+		/// <inheritdoc/>
+		public override ArtifactId FromObjectId(ObjectId id) => new ArtifactId(id);
+
+		/// <inheritdoc/>
+		public override ObjectId ToObjectId(ArtifactId value) => value.Id;
+	}
+}
diff --git a/Engine/Source/Programs/Horde/Horde.Build/Artifacts/ArtifactType.cs b/Engine/Source/Programs/Horde/Horde.Build/Artifacts/ArtifactType.cs
new file mode 100644
--- /dev/null
+++ b/Engine/Source/Programs/Horde/Horde.Build/Artifacts/ArtifactType.cs
@@ -0,0 +1,48 @@
+// Copyright Epic Games, Inc. All Rights Reserved.
+
+using System.ComponentModel;
+using EpicGames.Horde;
+using Horde.Build.Utilities;
+
+namespace Horde.Build.Artifacts
+{
+	/// <summary>
+	/// Type of an artifact
+	/// </summary>
+	/// <param name="Id">The artifact type</param>
+	[StringIdConverter(typeof(ArtifactTypeConverter))]
+	[TypeConverter(typeof(StringIdTypeConverter<ArtifactType, ArtifactTypeConverter>))]
+	public record struct ArtifactType(StringId Id)
+	{
+		/// <summary>
+		/// Output from a build step
+		/// </summary>
+		public static ArtifactType StepOutput { get; } = new ArtifactType("step-output");
+
+		/// <summary>
+		/// Captured state from the machine executing a build step
+		/// </summary>
+		public static ArtifactType StepSaved { get; } = new ArtifactType("step-saved");
+
+		/// <summary>
+		/// Constructor
+		/// </summary>
+		/// <param name="id">Identifier for the artifact type</param>
+		public ArtifactType(string id) : this(new StringId(id)) { }
+
+		/// <inheritdoc/>
+		public override string ToString() => Id.ToString();
+	}
+
+	/// <summary>
+	/// Converter to and from <see cref="StringId"/> instances.
+	/// </summary>
+	class ArtifactTypeConverter : StringIdConverter<ArtifactType>
+	{
+		/// <inheritdoc/>
+		public override ArtifactType FromStringId(StringId id) => new ArtifactType(id);
+
+		/// <inheritdoc/>
+		public override StringId ToStringId(ArtifactType value) => value.Id;
+	}
+}
diff --git a/Engine/Source/Programs/Horde/Horde.Build/Artifacts/ArtifactsController.cs b/Engine/Source/Programs/Horde/Horde.Build/Artifacts/ArtifactsController.cs
new file mode 100644
--- /dev/null
+++ b/Engine/Source/Programs/Horde/Horde.Build/Artifacts/ArtifactsController.cs
@@ -0,0 +1,259 @@
+// Copyright Epic Games, Inc. All Rights Reserved.
+
+using System;
+using System.Collections.Generic;
+using System.Linq;
+using System.Threading;
+using System.Threading.Tasks;
+using EpicGames.Core;
+using EpicGames.Horde.Storage;
+using EpicGames.Horde.Storage.Nodes;
+using Horde.Build.Acls;
+using Horde.Build.Server;
+using Horde.Build.Storage;
+using Horde.Build.Utilities;
+using Microsoft.AspNetCore.Authorization;
+using Microsoft.AspNetCore.Http;
+using Microsoft.AspNetCore.Mvc;
+using Microsoft.Extensions.Options;
+
+namespace Horde.Build.Artifacts
+{
+	/// <summary>
+	/// Describes an artifact
+	/// </summary>
+	public class GetArtifactResponse
+	{
+		readonly IArtifact _artifact;
+
+		/// <inheritdoc cref="IArtifact.Id"/>
+		public string Id => _artifact.Id.ToString();
+
+		/// <inheritdoc cref="IArtifact.Name"/>
+		public string Name => _artifact.Name;
+
+		/// <inheritdoc cref="IArtifact.Type"/>
+		public ArtifactType Type => _artifact.Type;
+
+		/// <inheritdoc cref="IArtifact.Keys"/>
+		public IReadOnlyList<string> Keys => _artifact.Keys;
+
+		internal GetArtifactResponse(IArtifact artifact) => _artifact = artifact;
+	}
+
+	/// <summary>
+	/// Result of an artifact search
+	/// </summary>
+	public class FindArtifactsResponse
+	{
+		/// <summary>
+		/// List of artifacts matching the search criteria
+		/// </summary>
+		public List<GetArtifactResponse> Artifacts { get; } = new List<GetArtifactResponse>();
+	}
+
+	/// <summary>
+	/// Describes a file within an artifact
+	/// </summary>
+	public class GetArtifactFileEntryResponse
+	{
+		readonly FileEntry _entry;
+
+		/// <inheritdoc cref="FileEntry.Name"/>
+		public string Name => _entry.Name.ToString();
+
+		/// <inheritdoc cref="FileEntry.Length"/>
+		public long Length => _entry.Length;
+
+		/// <inheritdoc cref="FileEntry.Hash"/>
+		public IoHash Hash => _entry.Hash;
+
+		internal GetArtifactFileEntryResponse(FileEntry entry) => _entry = entry;
+	}
+
+	/// <summary>
+	/// Describes a file within an artifact
+	/// </summary>
+	public class GetArtifactDirectoryEntryResponse
+	{
+		readonly DirectoryEntry _entry;
+
+		/// <inheritdoc cref="FileEntry.Name"/>
+		public string Name => _entry.Name.ToString();
+
+		/// <inheritdoc cref="FileEntry.Length"/>
+		public long Length => _entry.Length;
+
+		/// <inheritdoc cref="FileEntry.Hash"/>
+		public IoHash Hash => _entry.Hash;
+
+		internal GetArtifactDirectoryEntryResponse(DirectoryEntry entry) => _entry = entry;
+	}
+
+	/// <summary>
+	/// Describes a directory within an artifact
+	/// </summary>
+	public class GetArtifactDirectoryResponse
+	{
+		/// <summary>
+		/// Names of sub-directories
+		/// </summary>
+		public List<GetArtifactDirectoryEntryResponse> Directories { get; } = new List<GetArtifactDirectoryEntryResponse>();
+
+		/// <summary>
+		/// Files within the directory
+		/// </summary>
+		public List<GetArtifactFileEntryResponse> Files { get; } = new List<GetArtifactFileEntryResponse>();
+
+		internal GetArtifactDirectoryResponse(DirectoryNode directory)
+		{
+			Directories.AddRange(directory.Directories.Select(x => new GetArtifactDirectoryEntryResponse(x)));
+			Files.AddRange(directory.Files.Select(x => new GetArtifactFileEntryResponse(x)));
+		}
+	}
+
+	/// <summary>
+	/// Public interface for artifacts
+	/// </summary>
+	[ApiController]
+	public class ArtifactsController : HordeControllerBase
+	{
+		readonly IArtifactCollection _artifactCollection;
+		readonly StorageService _storageService;
+		readonly GlobalConfig _globalConfig;
+
+		/// <summary>
+		/// Constructor
+		/// </summary>
+		public ArtifactsController(IArtifactCollection artifactCollection, StorageService storageService, IOptionsSnapshot<GlobalConfig> globalConfig)
+		{
+			_artifactCollection = artifactCollection;
+			_storageService = storageService;
+			_globalConfig = globalConfig.Value;
+		}
+
+		/// <summary>
+		/// Gets metadata about an artifact object
+		/// </summary>
+		/// <param name="id">Identifier of the artifact to retrieve</param>
+		/// <param name="filter">Filter for returned properties</param>
+		/// <returns>Information about all the artifacts</returns>
+		[HttpGet]
+		[Authorize]
+		[Route("/api/v2/artifacts/{id}")]
+		[ProducesResponseType(typeof(FindArtifactsResponse), 200)]
+		public async Task<ActionResult<object>> GetArtifactAsync(ArtifactId id, [FromQuery] PropertyFilter? filter = null)
+		{
+			IArtifact? artifact = await _artifactCollection.GetAsync(id, HttpContext.RequestAborted);
+			if (artifact == null)
+			{
+				return NotFound(id);
+			}
+			if (!_globalConfig.Authorize(artifact.AclScope, AclAction.ReadArtifact, User))
+			{
+				return Forbid(AclAction.ReadArtifact, artifact.AclScope);
+			}
+
+			return PropertyFilter.Apply(new GetArtifactResponse(artifact), filter);
+		}
+
+		/// <summary>
+		/// Uploads a blob for an artifact. See /api/v1/storage/XXX/blobs.
+		/// </summary>
+		/// <param name="id">Id of the artifact to write to</param>
+		/// <param name="file">Data for the blob</param>
+		/// <param name="cancellationToken">Cancellation token for the request</param>
+		/// <returns>Information about the requested project</returns>
+		[HttpPost]
+		[Route("/api/v2/artifacts/{id}/blobs")]
+		[ProducesResponseType(typeof(WriteBlobResponse), 200)]
+		public async Task<ActionResult<WriteBlobResponse>> WriteArtifactBlobAsync(ArtifactId id, IFormFile? file, CancellationToken cancellationToken = default)
+		{
+			IArtifact? artifact = await _artifactCollection.GetAsync(id, cancellationToken);
+			if (artifact == null)
+			{
+				return NotFound(id);
+			}
+			if (!_globalConfig.Authorize(artifact.AclScope, AclAction.WriteArtifact, User))
+			{
+				return Forbid(AclAction.WriteArtifact, artifact.AclScope);
+			}
+
+			return await StorageController.WriteBlobAsync(_storageService, Namespace.Artifacts, file, $"{artifact.RefName}", cancellationToken);
+		}
+
+		/// <summary>
+		/// Gets metadata about an artifact object
+		/// </summary>
+		/// <param name="id">Identifier of the artifact to retrieve</param>
+		/// <param name="path">Path to fetch</param>
+		/// <param name="filter">Filter for returned properties</param>
+		/// <param name="cancellationToken">Cancellation token for the operation</param>
+		/// <returns>Information about all the artifacts</returns>
+		[HttpGet]
+		[Authorize]
+		[Route("/api/v2/artifacts/{id}/browse")]
+		[ProducesResponseType(typeof(GetArtifactDirectoryResponse), 200)]
+		public async Task<ActionResult<object>> BrowseArtifactAsync(ArtifactId id, [FromQuery] string? path = null, [FromQuery] PropertyFilter? filter = null, CancellationToken cancellationToken = default)
+		{
+			IArtifact? artifact = await _artifactCollection.GetAsync(id, cancellationToken);
+			if (artifact == null)
+			{
+				return NotFound(id);
+			}
+			if (!_globalConfig.Authorize(artifact.AclScope, AclAction.ReadArtifact, User))
+			{
+				return Forbid(AclAction.ReadArtifact, artifact.AclScope);
+			}
+
+			TreeReader reader = await _storageService.GetReaderAsync(artifact.NamespaceId, cancellationToken);
+			TreeNodeRef<DirectoryNode> directoryRef = await reader.ReadNodeRefAsync<DirectoryNode>(artifact.RefName, DateTime.UtcNow.AddHours(1.0), cancellationToken);
+
+			if (path != null)
+			{
+				foreach (string fragment in path.Split('/'))
+				{
+					DirectoryNode nextDirectory = await directoryRef.ExpandAsync(reader, cancellationToken);
+					if (!nextDirectory.TryGetDirectoryEntry(fragment, out DirectoryEntry? nextDirectoryEntry))
+					{
+						return NotFound();
+					}
+					directoryRef = nextDirectoryEntry;
+				}
+			}
+
+			DirectoryNode directory = await directoryRef.ExpandAsync(reader, cancellationToken);
+			return PropertyFilter.Apply(new GetArtifactDirectoryResponse(directory), filter);
+		}
+
+		/// <summary>
+		/// Gets metadata about an artifact object
+		/// </summary>
+		/// <param name="ids">Artifact ids to return</param>
+		/// <param name="keys">Keys to find</param>
+		/// <param name="filter">Filter for returned values</param>
+		/// <returns>Information about all the artifacts</returns>
+		[HttpGet]
+		[Authorize]
+		[Route("/api/v2/artifacts")]
+		[ProducesResponseType(typeof(FindArtifactsResponse), 200)]
+		public async Task<ActionResult<object>> FindArtifactsAsync([FromQuery(Name = "id")] IEnumerable<ArtifactId>? ids = null, [FromQuery(Name = "key")] IEnumerable<string>? keys = null, [FromQuery] PropertyFilter? filter = null)
+		{
+			if (ids == null && keys == null)
+			{
+				return BadRequest("At least one search term must be specified");
+			}
+
+			FindArtifactsResponse response = new FindArtifactsResponse();
+			await foreach (IArtifact artifact in _artifactCollection.FindAsync(ids, keys, HttpContext.RequestAborted))
+			{
+				if (_globalConfig.Authorize(artifact.AclScope, AclAction.ReadArtifact, User))
+				{
+					response.Artifacts.Add(new GetArtifactResponse(artifact));
+				}
+			}
+
+			return PropertyFilter.Apply(response, filter);
+		}
+	}
+}
diff --git a/Engine/Source/Programs/Horde/Horde.Build/Artifacts/IArtifact.cs b/Engine/Source/Programs/Horde/Horde.Build/Artifacts/IArtifact.cs
new file mode 100644
--- /dev/null
+++ b/Engine/Source/Programs/Horde/Horde.Build/Artifacts/IArtifact.cs
@@ -0,0 +1,49 @@
+// Copyright Epic Games, Inc. All Rights Reserved.
+
+using System.Collections.Generic;
+using EpicGames.Horde.Storage;
+using Horde.Build.Acls;
+
+namespace Horde.Build.Artifacts
+{
+	/// <summary>
+	/// Information about an artifact
+	/// </summary>
+	public interface IArtifact
+	{
+		/// <summary>
+		/// Identifier for the Artifact. Randomly generated.
+		/// </summary>
+		public ArtifactId Id { get; }
+
+		/// <summary>
+		/// Name for this artifact object.
+		/// </summary>
+		public string Name { get; }
+
+		/// <summary>
+		/// Type of artifact
+		/// </summary>
+		public ArtifactType Type { get; }
+
+		/// <summary>
+		/// Keys used to collate artifacts
+		/// </summary>
+		public IReadOnlyList<string> Keys { get; }
+
+		/// <summary>
+		/// Storage namespace containing the data
+		/// </summary>
+		public NamespaceId NamespaceId { get; }
+
+		/// <summary>
+		/// Name of the ref containing the root data object
+		/// </summary>
+		public RefName RefName { get; }
+
+		/// <summary>
+		/// Permissions scope for this object
+		/// </summary>
+		public AclScopeName AclScope { get; }
+	}
+}
diff --git a/Engine/Source/Programs/Horde/Horde.Build/Artifacts/IArtifactCollection.cs b/Engine/Source/Programs/Horde/Horde.Build/Artifacts/IArtifactCollection.cs
new file mode 100644
--- /dev/null
+++ b/Engine/Source/Programs/Horde/Horde.Build/Artifacts/IArtifactCollection.cs
@@ -0,0 +1,47 @@
+// Copyright Epic Games, Inc. All Rights Reserved.
+
+using System.Collections.Generic;
+using System.Threading;
+using System.Threading.Tasks;
+using EpicGames.Horde.Storage;
+using Horde.Build.Acls;
+using Horde.Build.Utilities;
+
+namespace Horde.Build.Artifacts
+{
+	/// <summary>
+	/// Interface for a collection of artifacts
+	/// </summary>
+	public interface IArtifactCollection
+	{
+		/// <summary>
+		/// Creates a new artifact
+		/// </summary>
+		/// <param name="name">Description for the artifact</param>
+		/// <param name="type">Type identifier for the artifact</param>
+		/// <param name="keys">Keys for the artifact</param>
+		/// <param name="namespaceId">Namespace containing the data</param>
+		/// <param name="refName">Name of the root reference. Will be initialized to the object id if not specified.</param>
+		/// <param name="scopeName">Inherited scope used for permissions</param>
+		/// <param name="cancellationToken">Cancellation token for the operation</param>
+		/// <returns>The new log file document</returns>
+		Task<IArtifact> AddAsync(string name, ArtifactType type, IEnumerable<string> keys, NamespaceId namespaceId, RefName? refName, AclScopeName scopeName, CancellationToken cancellationToken = default);
+
+		/// <summary>
+		/// Finds artifacts with the given keys
+		/// </summary>
+		/// <param name="ids">Ids to search for</param>
+		/// <param name="keys">Keys to search for</param>
+		/// <param name="cancellationToken">Cancellation token for the operation</param>
+		/// <returns>Sequence of artifacts</returns>
+		IAsyncEnumerable<IArtifact> FindAsync(IEnumerable<ArtifactId>? ids = null, IEnumerable<string>? keys = null, CancellationToken cancellationToken = default);
+
+		/// <summary>
+		/// Gets an artifact by ID
+		/// </summary>
+		/// <param name="artifactId">Unique id of the artifact</param>
+		/// <param name="cancellationToken">Cancellation token for the operation</param>
+		/// <returns>The artifact document</returns>
+		Task<IArtifact?> GetAsync(ArtifactId artifactId, CancellationToken cancellationToken = default);
+	}
+}
diff --git a/Engine/Source/Programs/Horde/Horde.Build/Horde.Build.csproj b/Engine/Source/Programs/Horde/Horde.Build/Horde.Build.csproj
--- a/Engine/Source/Programs/Horde/Horde.Build/Horde.Build.csproj
+++ b/Engine/Source/Programs/Horde/Horde.Build/Horde.Build.csproj
@@ -72,6 +72,8 @@
   </ItemGroup>
 
   <ItemGroup>
+    <Protobuf Include="Protos\horde\job_rpc.proto" ProtoRoot="Protos" GrpcServices="Server" />
+    <Protobuf Include="Protos\horde\job_rpc_messages.proto" ProtoRoot="Protos" />
     <Protobuf Include="Protos\horde\horde_rpc.proto" ProtoRoot="Protos" GrpcServices="Server" />
     <Protobuf Include="Protos\horde\horde_rpc_messages.proto" ProtoRoot="Protos" />
     <Protobuf Include="Protos\horde\log_rpc.proto" ProtoRoot="Protos" GrpcServices="Server" />
diff --git a/Engine/Source/Programs/Horde/Horde.Build/Jobs/IJob.cs b/Engine/Source/Programs/Horde/Horde.Build/Jobs/IJob.cs
--- a/Engine/Source/Programs/Horde/Horde.Build/Jobs/IJob.cs
+++ b/Engine/Source/Programs/Horde/Horde.Build/Jobs/IJob.cs
@@ -779,6 +779,31 @@ namespace Horde.Build.Jobs
 			return false;
 		}
 
+		/// <summary>
+		/// Attempts to get a step with a given ID
+		/// </summary>
+		/// <param name="job">The job document</param>
+		/// <param name="stepId">The step id</param>
+		/// <param name="batch">On success returns the batch containing the step</param>
+		/// <param name="step">On success, receives the step object</param>
+		/// <returns>True if the step was found</returns>
+		public static bool TryGetStep(this IJob job, SubResourceId stepId, [NotNullWhen(true)] out IJobStepBatch? batch, [NotNullWhen(true)] out IJobStep? step)
+		{
+			foreach (IJobStepBatch currentBatch in job.Batches)
+			{
+				if (currentBatch.TryGetStep(stepId, out IJobStep? currentStep))
+				{
+					batch = currentBatch;
+					step = currentStep;
+					return true;
+				}
+			}
+
+			batch = null;
+			step = null;
+			return false;
+		}
+
 		/// <summary>
 		/// Gets the current job state
 		/// </summary>
diff --git a/Engine/Source/Programs/Horde/Horde.Build/Server/GlobalConfig.cs b/Engine/Source/Programs/Horde/Horde.Build/Server/GlobalConfig.cs
--- a/Engine/Source/Programs/Horde/Horde.Build/Server/GlobalConfig.cs
+++ b/Engine/Source/Programs/Horde/Horde.Build/Server/GlobalConfig.cs
@@ -214,6 +214,23 @@ namespace Horde.Build.Server
 		/// <returns>True if the stream configuration was found</returns>
 		public bool TryGetStream(StreamId streamId, [NotNullWhen(true)] out StreamConfig? config) => _streamLookup.TryGetValue(streamId, out config);
 
+		/// <summary>
+		/// Attempts to get configuration for a stream from this object
+		/// </summary>
+		/// <param name="streamId">The stream identifier</param>
+		/// <param name="templateId">Template identifier</param>
+		/// <param name="config">Configuration for the stream</param>
+		/// <returns>True if the stream configuration was found</returns>
+		public bool TryGetTemplate(StreamId streamId, TemplateId templateId, [NotNullWhen(true)] out TemplateRefConfig? config)
+		{
+			if (!_streamLookup.TryGetValue(streamId, out StreamConfig? streamConfig))
+			{
+				config = null;
+				return false;
+			}
+			return streamConfig.TryGetTemplate(templateId, out config);
+		}
+
 		/// <summary>
 		/// Attempts to get configuration for a tool from this object
 		/// </summary>
diff --git a/Engine/Source/Programs/Horde/Horde.Build/Startup.cs b/Engine/Source/Programs/Horde/Horde.Build/Startup.cs
--- a/Engine/Source/Programs/Horde/Horde.Build/Startup.cs
+++ b/Engine/Source/Programs/Horde/Horde.Build/Startup.cs
@@ -94,6 +94,7 @@ using Horde.Build.Auditing;
 using Horde.Build.Server.Notices;
 using Horde.Build.Notifications.Sinks;
 using StatusCode = Grpc.Core.StatusCode;
+using Horde.Build.Artifacts;
 
 namespace Horde.Build
 {
@@ -333,6 +334,7 @@ namespace Horde.Build
 
 			services.AddSingleton<IAgentCollection, AgentCollection>();
 			services.AddSingleton<IAgentSoftwareCollection, AgentSoftwareCollection>();
+			services.AddSingleton<IArtifactCollection, ArtifactCollection>();
 			services.AddSingleton<IArtifactCollectionV1, ArtifactCollectionV1>();
 			services.AddSingleton<IGraphCollection, GraphCollection>();
 			services.AddSingleton<IIssueCollection, IssueCollection>();
@@ -1027,6 +1029,7 @@ namespace Horde.Build
 			{
 				endpoints.MapGrpcService<HealthService>();
 				endpoints.MapGrpcService<RpcService>();
+				endpoints.MapGrpcService<JobRpcService>();
 				endpoints.MapGrpcService<LogRpcService>();
 
 				endpoints.MapGrpcReflectionService();
diff --git a/Engine/Source/Programs/Horde/Horde.Build/Storage/StorageService.cs b/Engine/Source/Programs/Horde/Horde.Build/Storage/StorageService.cs
--- a/Engine/Source/Programs/Horde/Horde.Build/Storage/StorageService.cs
+++ b/Engine/Source/Programs/Horde/Horde.Build/Storage/StorageService.cs
@@ -434,9 +434,7 @@ namespace Horde.Build.Storage
 			public DateTime LastTime { get; set; }
 		}
 
-		readonly GlobalsService _globalsService;
 		readonly RedisService _redisService;
-		readonly AclService _aclService;
 		readonly IClock _clock;
 		readonly IMemoryCache _cache;
 		readonly IServiceProvider _serviceProvider;
@@ -462,11 +460,9 @@ namespace Horde.Build.Storage
 		/// <summary>
 		/// Constructor
 		/// </summary>
-		public StorageService(GlobalsService globalsService, MongoService mongoService, RedisService redisService, AclService aclService, IClock clock, IMemoryCache cache, IServiceProvider serviceProvider, IOptionsMonitor<GlobalConfig> globalConfig, ILogger<StorageService> logger)
+		public StorageService(MongoService mongoService, RedisService redisService, IClock clock, IMemoryCache cache, IServiceProvider serviceProvider, IOptionsMonitor<GlobalConfig> globalConfig, ILogger<StorageService> logger)
 		{
-			_globalsService = globalsService;
 			_redisService = redisService;
-			_aclService = aclService;
 			_clock = clock;
 			_cache = cache;
 			_serviceProvider = serviceProvider;
@@ -565,6 +561,18 @@ namespace Horde.Build.Storage
 		/// <inheritdoc/>
 		async ValueTask<IStorageClient> IStorageClientFactory.GetClientAsync(NamespaceId namespaceId, CancellationToken cancellationToken) => await GetClientAsync(namespaceId, cancellationToken);
 
+		/// <summary>
+		/// Gets a <see cref="TreeReader"/> for a particular namespace
+		/// </summary>
+		/// <param name="namespaceId">Namespace identififer</param>
+		/// <param name="cancellationToken">Cancellation token for the operation</param>
+		/// <returns></returns>
+		public async ValueTask<TreeReader> GetReaderAsync(NamespaceId namespaceId, CancellationToken cancellationToken)
+		{
+			IStorageClientImpl storageClient = await GetClientAsync(namespaceId, cancellationToken);
+			return new TreeReader(storageClient, _cache, _logger);
+		}
+
 		#region Config
 
 		State GetNextState()
diff --git a/Engine/Source/Programs/Horde/Horde.Build/Utilities/HordeControllerBase.cs b/Engine/Source/Programs/Horde/Horde.Build/Utilities/HordeControllerBase.cs
--- a/Engine/Source/Programs/Horde/Horde.Build/Utilities/HordeControllerBase.cs
+++ b/Engine/Source/Programs/Horde/Horde.Build/Utilities/HordeControllerBase.cs
@@ -6,6 +6,7 @@ using EpicGames.Horde.Storage;
 using Horde.Build.Acls;
 using Horde.Build.Agents;
 using Horde.Build.Agents.Leases;
+using Horde.Build.Artifacts;
 using Horde.Build.Jobs;
 using Horde.Build.Projects;
 using Horde.Build.Streams;
@@ -62,6 +63,15 @@ namespace Horde.Build.Utilities
 			return Forbid(action, "agent {AgentId}", agentId);
 		}
 
+		/// <summary>
+		/// Returns a 403 (forbidden) response with the given action and object
+		/// </summary>
+		[NonAction]
+		protected ActionResult Forbid(AclAction action, AclScopeName scopeName)
+		{
+			return Forbid(action, "scope {ScopeName}", scopeName);
+		}
+
 		/// <summary>
 		/// Returns a 403 (forbidden) response with the given action and object
 		/// </summary>
@@ -161,6 +171,15 @@ namespace Horde.Build.Utilities
 			return NotFound("Lease {LeaseId} not found for agent {AgentId}", leaseId, agentId);
 		}
 
+		/// <summary>
+		/// Returns a 404 response for the given object
+		/// </summary>
+		[NonAction]
+		protected ActionResult NotFound(ArtifactId artifactId)
+		{
+			return NotFound("Artifact {ArtifactId} not found", artifactId);
+		}
+
 		/// <summary>
 		/// Returns a 404 response for the given object
 		/// </summary>
