Horde: Missing file.

diff --git a/Engine/Source/Programs/Horde/Horde.Agent/Execution/JobExecutor.cs b/Engine/Source/Programs/Horde/Horde.Agent/Execution/JobExecutor.cs
--- a/Engine/Source/Programs/Horde/Horde.Agent/Execution/JobExecutor.cs
+++ b/Engine/Source/Programs/Horde/Horde.Agent/Execution/JobExecutor.cs
@@ -23,6 +23,7 @@ using Horde.Agent.Parser;
 using Horde.Agent.Services;
 using Horde.Agent.Utility;
 using Horde.Common;
+using Horde.Common.Rpc;
 using Horde.Storage.Utility;
 using HordeCommon;
 using HordeCommon.Rpc;
@@ -729,10 +730,34 @@ namespace Horde.Agent.Execution
 			}
 		}
 
+		async Task TestArtifactAsync(string stepId, CancellationToken cancellationToken)
+		{
+			try
+			{
+				using IRpcClientRef<JobRpc.JobRpcClient> jobRpc = await RpcConnection.GetClientRefAsync<JobRpc.JobRpcClient>(cancellationToken);
+				CreateJobArtifactResponse artifact = await jobRpc.Client.CreateArtifactAsync(new CreateJobArtifactRequest { JobId = _jobId, StepId = stepId, Name = "test artifact", Type = JobArtifactType.Saved }, cancellationToken: cancellationToken);
+				_logger.LogInformation("Created artifact {ArtifactId} with ref {RefName} in ns {Namespace}", artifact.Id, artifact.RefName, artifact.NamespaceId);
+
+				using MemoryCache cache = new MemoryCache(new MemoryCacheOptions { });
+				IStorageClient storage = _storageFactory.CreateStorageClient(_session, $"/api/v1/storage/{artifact.NamespaceId}");
+
+				using TreeWriter writer = new TreeWriter(storage, new RefName(artifact.RefName));
+
+				DirectoryNode dir = new DirectoryNode();
+				await writer.WriteAsync(new RefName(artifact.RefName), dir, cancellationToken: cancellationToken);
+			}
+			catch (Exception ex)
+			{
+				_logger.LogInformation(ex, "Error creating artifact");
+			}
+		}
+
 		private async Task<bool> ExecuteWithTempStorageAsync(BeginStepResponse step, DirectoryReference workspaceDir, string arguments, ILogger logger, CancellationToken cancellationToken)
 		{
 			DirectoryReference manifestDir = DirectoryReference.Combine(workspaceDir, "Engine", "Saved", "BuildGraph");
 
+			await TestArtifactAsync(step.StepId, cancellationToken);
+
 			using MemoryCache cache = new MemoryCache(new MemoryCacheOptions { });
 			IStorageClient storage = _storageFactory.CreateStorageClient(_session, $"/api/v1/jobs/{_jobId}/batches/{_batchId}/steps/{step.StepId}/temp");
 			TreeReader reader = new TreeReader(storage, cache, _logger);
