Horde: Missing file.

diff --git a/Engine/Source/Programs/Horde/Horde.Build/Jobs/JobRpcService.cs b/Engine/Source/Programs/Horde/Horde.Build/Jobs/JobRpcService.cs
new file mode 100644
--- /dev/null
+++ b/Engine/Source/Programs/Horde/Horde.Build/Jobs/JobRpcService.cs
@@ -0,0 +1,114 @@
+// Copyright Epic Games, Inc. All Rights Reserved.
+
+using System;
+using System.Collections.Generic;
+using System.Security.Claims;
+using System.Threading.Tasks;
+using EpicGames.Horde.Storage;
+using Grpc.Core;
+using Horde.Build.Acls;
+using Horde.Build.Agents.Sessions;
+using Horde.Build.Artifacts;
+using Horde.Build.Server;
+using Horde.Build.Storage;
+using Horde.Build.Streams;
+using Horde.Build.Utilities;
+using Horde.Common.Rpc;
+using Microsoft.AspNetCore.Authorization;
+using Microsoft.Extensions.Options;
+
+namespace Horde.Build.Jobs
+{
+	/// <summary>
+	/// Implements the Horde gRPC service for bots updating their status and dequeing work
+	/// </summary>
+	[Authorize]
+	public class JobRpcService : JobRpc.JobRpcBase
+	{
+		readonly IJobCollection _jobCollection;
+		readonly IArtifactCollection _artifactCollection;
+		readonly GlobalConfig _globalConfig;
+
+		/// <summary>
+		/// Constructor
+		/// </summary>
+		public JobRpcService(IJobCollection jobCollection, IArtifactCollection artifactCollection, IOptionsSnapshot<GlobalConfig> globalConfig)
+		{
+			_jobCollection = jobCollection;
+			_artifactCollection = artifactCollection;
+			_globalConfig = globalConfig.Value;
+		}
+
+		/// <inheritdoc/>
+		public override async Task<CreateJobArtifactResponse> CreateArtifact(CreateJobArtifactRequest request, ServerCallContext context)
+		{
+			(IJob job, _, IJobStep step) = await AuthorizeAsync(request.JobId, request.StepId, context);
+
+			ArtifactType type = request.Type switch
+			{
+				JobArtifactType.Output => ArtifactType.StepOutput,
+				JobArtifactType.Saved => ArtifactType.StepSaved,
+				_ => throw new StructuredRpcException(StatusCode.InvalidArgument, "Invalid artifact type")
+			};
+
+			List<string> keys = new List<string>();
+			keys.Add($"job:{job.Id}");
+			keys.Add($"job:{job.Id}/step:{step.Id}");
+
+			if (!_globalConfig.TryGetTemplate(job.StreamId, job.TemplateId, out TemplateRefConfig? templateConfig))
+			{
+				throw new StructuredRpcException(StatusCode.NotFound, "Couldn't find template {TemplateId} in stream {StreamId}", job.TemplateId, job.StreamId);
+			}
+
+			NamespaceId namespaceId = Namespace.Artifacts;
+			IArtifact artifact = await _artifactCollection.AddAsync(request.Name, type, keys, namespaceId, null, templateConfig.ScopeName, context.CancellationToken);
+			// TODO: Token
+			return new CreateJobArtifactResponse { Id = artifact.Id.ToString(), NamespaceId = artifact.NamespaceId.ToString(), RefName = artifact.RefName.ToString() };
+		}
+		/*
+		/// <inheritdoc/>
+		public override async Task<FinalizeJobArtifactResponse> FinalizeArtifact(FinalizeJobArtifactRequest request, ServerCallContext context)
+		{
+			IArtifact? artifact = await _artifactCollection.GetAsync(ArtifactId.Parse(request.Id), context.CancellationToken);
+			if (artifact == null)
+			{
+				throw new StructuredRpcException(StatusCode.NotFound, "Missing artifact {ArtifactId}", request.Id);
+			}
+
+			IStorageClientImpl storage = await _storageService.GetClientAsync(artifact.NamespaceId, context.CancellationToken);
+			await storage.WriteRefTargetAsync(artifact.RefName, NodeHandle.Parse(request.Locator), cancellationToken: context.CancellationToken);
+
+			return new FinalizeJobArtifactResponse();
+		}
+		*/
+		Task<(IJob, IJobStepBatch, IJobStep)> AuthorizeAsync(string jobId, string stepId, ServerCallContext context)
+		{
+			return AuthorizeAsync(JobId.Parse(jobId), SubResourceId.Parse(stepId), context);
+		}
+
+		async Task<(IJob, IJobStepBatch, IJobStep)> AuthorizeAsync(JobId jobId, SubResourceId stepId, ServerCallContext context)
+		{
+			IJob? job = await _jobCollection.GetAsync(jobId);
+			if (job == null)
+			{
+				throw new StructuredRpcException(StatusCode.NotFound, "Unable to find job {JobId}", jobId);
+			}
+			if (!job.TryGetStep(stepId, out IJobStepBatch? batch, out IJobStep? step))
+			{
+				throw new StructuredRpcException(StatusCode.NotFound, "Unable to find step {JobId}:{StepId}", job.Id, stepId);
+			}
+			if (batch.SessionId == null)
+			{
+				throw new StructuredRpcException(StatusCode.PermissionDenied, "Batch {JobId}:{BatchId} has no session id", job.Id, batch.Id);
+			}
+
+			ClaimsPrincipal principal = context.GetHttpContext().User;
+			if (!principal.HasSessionClaim(batch.SessionId.Value))
+			{
+				throw new StructuredRpcException(StatusCode.PermissionDenied, "Session id {SessionId} not valid for step {JobId}:{BatchId}:{StepId}. Expected {ExpectedSessionId}.", principal.GetSessionClaim() ?? SessionId.Empty, job.Id, batch.Id, step.Id, batch.SessionId.Value);
+			}
+
+			return (job, batch, step);
+		}
+	}
+}
