Horde: Fix missing call to base FinalizeAsync in PerforceExecutor

diff --git a/Engine/Source/Programs/Horde/Horde.Agent/Execution/PerforceExecutor.cs b/Engine/Source/Programs/Horde/Horde.Agent/Execution/PerforceExecutor.cs
--- a/Engine/Source/Programs/Horde/Horde.Agent/Execution/PerforceExecutor.cs
+++ b/Engine/Source/Programs/Horde/Horde.Agent/Execution/PerforceExecutor.cs
@@ -225,6 +225,7 @@ namespace Horde.Agent.Execution
 		public override async Task FinalizeAsync(ILogger logger, CancellationToken cancellationToken)
 		{
 			await _workspace.CleanAsync(cancellationToken);
+			await base.FinalizeAsync(logger, cancellationToken);
 		}
 
 		public static async Task ConformAsync(DirectoryReference rootDir, IList<AgentWorkspace> pendingWorkspaces, bool removeUntrackedFiles, ILogger logger, CancellationToken cancellationToken)
