Horde: Improve XGE metadata metric logging

diff --git a/Engine/Source/Programs/Horde/Horde.Agent/Execution/JobExecutor.cs b/Engine/Source/Programs/Horde/Horde.Agent/Execution/JobExecutor.cs
--- a/Engine/Source/Programs/Horde/Horde.Agent/Execution/JobExecutor.cs
+++ b/Engine/Source/Programs/Horde/Horde.Agent/Execution/JobExecutor.cs
@@ -357,16 +357,28 @@ namespace Horde.Agent.Execution
 			}
 		}
 
-		public virtual Task FinalizeAsync(ILogger logger, CancellationToken cancellationToken)
+		public virtual Task FinalizeAsync(ILogger jobLogger, CancellationToken cancellationToken)
 		{
 			if (_xgeMetadataExtractor != null)
 			{
 				// Collect any new *.ib_mon files and generate a summary for each build invocation, grouped by title.
 				// Title will segregate for example UnrealBuildTool from invocations coming from the engine/cooks.
-				foreach (XgeTaskMetadataSummary s in _xgeMetadataExtractor.GetSummariesForIbMonFiles())
+				List<XgeTaskMetadataSummary> summaries = _xgeMetadataExtractor.GetSummariesForIbMonFiles();
+				if (summaries.Count > 0)
 				{
-					_logger.LogInformation("XGE build invocation {Title} with {LocalTaskCount} local tasks for {LocalTaskDuration} secs total and {RemoteTaskCount} remote tasks for {RemoteTaskDuration} secs total",
-						s.Title, s.LocalTaskCount, s.TotalLocalTaskDuration, s.RemoteTaskCount, s.TotalRemoteTaskDuration);
+					foreach (XgeTaskMetadataSummary s in summaries)
+					{
+						_logger.LogInformation("XGE build invocation {Title} with {LocalTaskCount} local tasks for {LocalTaskDuration} secs total and {RemoteTaskCount} remote tasks for {RemoteTaskDuration} secs total",
+							s.Title, s.LocalTaskCount, s.TotalLocalTaskDuration.TotalSeconds, s.RemoteTaskCount, s.TotalRemoteTaskDuration.TotalSeconds);
+						
+						jobLogger.LogInformation("XGE build invocation {Title} with {LocalTaskCount} local tasks for {LocalTaskDuration} secs total and {RemoteTaskCount} remote tasks for {RemoteTaskDuration} secs total",
+							s.Title, s.LocalTaskCount, s.TotalLocalTaskDuration.TotalSeconds, s.RemoteTaskCount, s.TotalRemoteTaskDuration.TotalSeconds);
+					}
+				}
+				else
+				{
+					_logger.LogInformation("No XGE build invocations found (no *.ib_mon files in History/Local)");
+					jobLogger.LogInformation("No XGE build invocations found (no *.ib_mon files in History/Local)");
 				}
 			}
 			
