Horde: Perform validation of storage config when config is read, rather than when storage service ticks over new settings. Ensures that errors are sent to the correct channel at the correct time.

diff --git a/Engine/Source/Programs/Horde/Horde.Build/ServerSettings.cs b/Engine/Source/Programs/Horde/Horde.Build/ServerSettings.cs
--- a/Engine/Source/Programs/Horde/Horde.Build/ServerSettings.cs
+++ b/Engine/Source/Programs/Horde/Horde.Build/ServerSettings.cs
@@ -70,7 +70,7 @@ namespace Horde.Build
 		public string? AwsBucketPath { get; set; }
 
 		/// <inheritdoc/>
-		public AwsCredentialsType AwsCredentials { get; set; }
+		public AwsCredentialsType? AwsCredentials { get; set; }
 
 		/// <inheritdoc/>
 		public string? AwsRole { get; set; }
diff --git a/Engine/Source/Programs/Horde/Horde.Build/Storage/Backends/AwsStorageBackend.cs b/Engine/Source/Programs/Horde/Horde.Build/Storage/Backends/AwsStorageBackend.cs
--- a/Engine/Source/Programs/Horde/Horde.Build/Storage/Backends/AwsStorageBackend.cs
+++ b/Engine/Source/Programs/Horde/Horde.Build/Storage/Backends/AwsStorageBackend.cs
@@ -72,7 +72,7 @@ namespace Horde.Build.Storage.Backends
 		/// <summary>
 		/// Type of credentials to use
 		/// </summary>
-		public AwsCredentialsType AwsCredentials { get; }
+		public AwsCredentialsType? AwsCredentials { get; }
 
 		/// <summary>
 		/// Name of the bucket to use
@@ -168,7 +168,7 @@ namespace Horde.Build.Storage.Backends
 				awsOptions.Region = RegionEndpoint.GetBySystemName(options.AwsRegion);
 			}
 
-			switch (options.AwsCredentials)
+			switch (options.AwsCredentials ?? AwsCredentialsType.Default)
 			{
 				case AwsCredentialsType.Default:
 					// Using the fallback credentials from the AWS SDK, it will pick up credentials through a number of default mechanisms.
diff --git a/Engine/Source/Programs/Horde/Horde.Build/Storage/StorageConfig.cs b/Engine/Source/Programs/Horde/Horde.Build/Storage/StorageConfig.cs
--- a/Engine/Source/Programs/Horde/Horde.Build/Storage/StorageConfig.cs
+++ b/Engine/Source/Programs/Horde/Horde.Build/Storage/StorageConfig.cs
@@ -1,13 +1,17 @@
 // Copyright Epic Games, Inc. All Rights Reserved.
 
 using System;
+using System.Collections;
 using System.Collections.Generic;
 using System.ComponentModel.DataAnnotations;
 using System.Diagnostics;
 using System.Diagnostics.CodeAnalysis;
+using System.IO;
 using System.Linq;
+using System.Reflection;
 using System.Security.Claims;
 using System.Text.Json.Serialization;
+using Amazon.S3.Model;
 using EpicGames.Horde.Storage;
 using Horde.Build.Acls;
 using Horde.Build.Server;
@@ -65,17 +69,59 @@ namespace Horde.Build.Storage
 		/// </summary>
 		internal void PostLoad(GlobalConfig globalConfig)
 		{
+			// Create a lookup for backend configs
 			foreach (BackendConfig backendConfig in Backends)
 			{
 				_backendLookup.Add(backendConfig.Id, backendConfig);
 			}
+
+			// Fixup all the inherited properties
+			Dictionary<BackendId, BackendConfig> mergedBackendConfigs = new Dictionary<BackendId, BackendConfig>(Backends.Count);
+			foreach (BackendConfig backendConfig in Backends)
+			{
+				MergeBackendConfigs(backendConfig.Id, _backendLookup, mergedBackendConfigs);
+			}
+
+			// Validate the backend config for each namespace
 			foreach (NamespaceConfig namespaceConfig in Namespaces)
 			{
-				namespaceConfig.PostLoad(globalConfig);
+				BackendConfig? backendConfig;
+				if (!_backendLookup.TryGetValue(namespaceConfig.Backend, out backendConfig))
+				{
+					throw new StorageException($"Missing or invalid backend identifier for namespace {namespaceConfig.Id}");
+				}
+
+				namespaceConfig.PostLoad(globalConfig, backendConfig);
 				_namespaceLookup.Add(namespaceConfig.Id, namespaceConfig);
 			}
 		}
 
+		static BackendConfig MergeBackendConfigs(BackendId backendId, Dictionary<BackendId, BackendConfig> baseIdToConfig, Dictionary<BackendId, BackendConfig> mergedIdToConfig)
+		{
+			BackendConfig? config;
+			if (mergedIdToConfig.TryGetValue(backendId, out config))
+			{
+				if (config == null)
+				{
+					throw new InvalidDataException($"Configuration for storage backend '{backendId}' is recursive.");
+				}
+			}
+			else
+			{
+				mergedIdToConfig.Add(backendId, null!);
+
+				config = baseIdToConfig[backendId];
+				if (!config.Base.IsEmpty)
+				{
+					BackendConfig baseConfig = MergeBackendConfigs(config.Base, baseIdToConfig, mergedIdToConfig);
+					config.MergeDefaults(baseConfig);
+				}
+
+				mergedIdToConfig[backendId] = config;
+			}
+			return config;
+		}
+
 		/// <summary>
 		/// Gets a backend with the given id
 		/// </summary>
@@ -123,7 +169,7 @@ namespace Horde.Build.Storage
 		public string? AwsBucketPath { get; set; }
 
 		/// <inheritdoc/>
-		public AwsCredentialsType AwsCredentials { get; set; }
+		public AwsCredentialsType? AwsCredentials { get; set; }
 
 		/// <inheritdoc/>
 		public string? AwsRole { get; set; }
@@ -139,6 +185,21 @@ namespace Horde.Build.Storage
 
 		/// <inheritdoc/>
 		public string? RelayToken { get; set; }
+
+		/// <summary>
+		/// Merge default values from another object
+		/// </summary>
+		/// <param name="other">Object to merge defaults from</param>
+		public void MergeDefaults(BackendConfig other)
+		{
+			foreach (PropertyInfo propertyInfo in GetType().GetProperties(BindingFlags.Instance | BindingFlags.Public))
+			{
+				if (propertyInfo.GetValue(this) == null)
+				{
+					propertyInfo.SetValue(this, propertyInfo.GetValue(other));
+				}
+			}
+		}
 	}
 
 	/// <summary>
@@ -153,6 +214,12 @@ namespace Horde.Build.Storage
 		[JsonIgnore]
 		public GlobalConfig GlobalConfig { get; private set; } = null!;
 
+		/// <summary>
+		/// The referenced backend config
+		/// </summary>
+		[JsonIgnore]
+		public BackendConfig BackendConfig { get; private set; } = null!;
+
 		/// <summary>
 		/// Identifier for this namespace
 		/// </summary>
@@ -194,9 +261,11 @@ namespace Horde.Build.Storage
 		/// Callback once the configuration has been read from disk
 		/// </summary>
 		/// <param name="globalConfig"></param>
-		public void PostLoad(GlobalConfig globalConfig)
+		/// <param name="backendConfig"></param>
+		public void PostLoad(GlobalConfig globalConfig, BackendConfig backendConfig)
 		{
 			GlobalConfig = globalConfig;
+			BackendConfig = backendConfig;
 		}
 
 		/// <summary>
diff --git a/Engine/Source/Programs/Horde/Horde.Build/Storage/StorageService.cs b/Engine/Source/Programs/Horde/Horde.Build/Storage/StorageService.cs
--- a/Engine/Source/Programs/Horde/Horde.Build/Storage/StorageService.cs
+++ b/Engine/Source/Programs/Horde/Horde.Build/Storage/StorageService.cs
@@ -455,8 +455,6 @@ namespace Horde.Build.Storage
 
 		readonly AsyncCachedValue<State> _cachedState;
 
-		static readonly BackendConfig s_defaultBackendConfig = new BackendConfig();
-
 		/// <summary>
 		/// Constructor
 		/// </summary>
@@ -582,34 +580,13 @@ namespace Horde.Build.Storage
 			{
 				StorageConfig storageConfig = globalConfig.Storage;
 
-				// Create a lookup for backend config objects by their id
-				Dictionary<BackendId, BackendConfig> backendIdToConfig = new Dictionary<BackendId, BackendConfig>();
-				foreach (BackendConfig backendConfig in storageConfig.Backends)
-				{
-					backendIdToConfig.Add(backendConfig.Id, backendConfig);
-				}
-
-				// Create a lookup of namespace id to config, to ensure there are no duplicates
-				Dictionary<NamespaceId, NamespaceConfig> namespaceIdToConfig = new Dictionary<NamespaceId, NamespaceConfig>();
-				foreach (NamespaceConfig namespaceConfig in storageConfig.Namespaces)
-				{
-					namespaceIdToConfig.Add(namespaceConfig.Id, namespaceConfig);
-				}
-
 				// Configure the new clients
 				State nextState = new State(storageConfig);
 				try
 				{
-					Dictionary<BackendId, BackendConfig> mergedIdToConfig = new Dictionary<BackendId, BackendConfig>();
 					foreach (NamespaceConfig namespaceConfig in storageConfig.Namespaces)
 					{
-						if (namespaceConfig.Backend.IsEmpty)
-						{
-							throw new StorageException($"No backend configured for namespace {namespaceConfig.Id}");
-						}
-
-						BackendConfig backendConfig = GetBackendConfig(namespaceConfig.Backend, backendIdToConfig, mergedIdToConfig);
-						IStorageBackend backend = CreateStorageBackend(backendConfig);
+						IStorageBackend backend = CreateStorageBackend(namespaceConfig.BackendConfig);
 
 #pragma warning disable CA2000 // Dispose objects before losing scope
 						StorageClient client = new StorageClient(this, namespaceConfig, backend);
@@ -651,93 +628,6 @@ namespace Horde.Build.Storage
 			}
 		}
 
-		/// <summary>
-		/// Gets the configuration for a particular backend
-		/// </summary>
-		/// <param name="backendId">Identifier for the backend</param>
-		/// <param name="baseIdToConfig">Lookup for all backend configuration data</param>
-		/// <param name="mergedIdToConfig">Lookup for computed hierarchical config objects</param>
-		/// <returns>Merged config for the given backend</returns>
-		BackendConfig GetBackendConfig(BackendId backendId, Dictionary<BackendId, BackendConfig> baseIdToConfig, Dictionary<BackendId, BackendConfig> mergedIdToConfig)
-		{
-			BackendConfig? config;
-			if (mergedIdToConfig.TryGetValue(backendId, out config))
-			{
-				if (config == null)
-				{
-					throw new StorageException($"Configuration for storage backend '{backendId}' is recursive.");
-				}
-			}
-			else
-			{
-				mergedIdToConfig.Add(backendId, null!);
-
-				if (!baseIdToConfig.TryGetValue(backendId, out config))
-				{
-					throw new StorageException($"Unable to find storage backend '{backendId}'"); 
-				}
-
-				if (config.Base != BackendId.Empty)
-				{
-					BackendConfig baseConfig = GetBackendConfig(config.Base, baseIdToConfig, mergedIdToConfig);
-					MergeConfigs(baseConfig, config, s_defaultBackendConfig);
-				}
-
-				mergedIdToConfig[backendId] = config;
-			}
-			return config;
-		}
-
-		/// <summary>
-		/// Allows one configuration object to override another
-		/// </summary>
-		static void MergeConfigs<T>(T sourceObject, T targetObject, T defaultObject)
-		{
-			PropertyInfo[] properties = typeof(T).GetProperties(BindingFlags.Public | BindingFlags.Instance);
-			foreach (PropertyInfo property in properties)
-			{
-				object? targetValue = property.GetValue(targetObject);
-				object? defaultValue = property.GetValue(defaultObject);
-
-				if (ValueEquals(targetValue, defaultValue))
-				{
-					object? sourceValue = property.GetValue(sourceObject);
-					property.SetValue(targetObject, sourceValue);
-				}
-			}
-		}
-
-		/// <summary>
-		/// Tests two objects for value equality
-		/// </summary>
-		static bool ValueEquals(object? source, object? target)
-		{
-			if (ReferenceEquals(source, null) || ReferenceEquals(target, null))
-			{
-				return ReferenceEquals(source, target);
-			}
-
-			if (source is IEnumerable _)
-			{
-				IEnumerator sourceEnumerator = ((IEnumerable)source).GetEnumerator();
-				IEnumerator targetEnumerator = ((IEnumerable)target).GetEnumerator();
-
-				for (; ; )
-				{
-					if (!sourceEnumerator.MoveNext())
-					{
-						return !targetEnumerator.MoveNext();
-					}
-					if (!targetEnumerator.MoveNext() || !ValueEquals(sourceEnumerator.Current, targetEnumerator.Current))
-					{
-						return false;
-					}
-				}
-			}
-
-			return source.Equals(target);
-		}
-
 		#endregion
 
 		#region Blobs
