Horde: Fix error updating line count in log files.

diff --git a/Engine/Source/Programs/Horde/Horde.Build/Logs/LogFileCollection.cs b/Engine/Source/Programs/Horde/Horde.Build/Logs/LogFileCollection.cs
--- a/Engine/Source/Programs/Horde/Horde.Build/Logs/LogFileCollection.cs
+++ b/Engine/Source/Programs/Horde/Horde.Build/Logs/LogFileCollection.cs
@@ -134,7 +134,7 @@ namespace Horde.Build.Logs
 		{
 			FilterDefinition<LogFileDocument> filter = Builders<LogFileDocument>.Filter.Eq(x => x.Id, logFileInterface.Id);
 			UpdateDefinition<LogFileDocument> update = Builders<LogFileDocument>.Update.Set(x => x.LineCount, lineCount).Inc(x => x.UpdateIndex, 1);
-			return await _logFiles.FindOneAndUpdateAsync(filter, update, new FindOneAndUpdateOptions<LogFileDocument, ILogFile> { ReturnDocument = ReturnDocument.After }, cancellationToken);
+			return await _logFiles.FindOneAndUpdateAsync(filter, update, new FindOneAndUpdateOptions<LogFileDocument, LogFileDocument> { ReturnDocument = ReturnDocument.After }, cancellationToken);
 		}
 
 		/// <inheritdoc/>
