Horde: Disable the exception handling redirect. The URL specified here is not valid.

diff --git a/Engine/Source/Programs/Horde/Horde.Build/Startup.cs b/Engine/Source/Programs/Horde/Horde.Build/Startup.cs
--- a/Engine/Source/Programs/Horde/Horde.Build/Startup.cs
+++ b/Engine/Source/Programs/Horde/Horde.Build/Startup.cs
@@ -970,7 +970,6 @@ namespace Horde.Build
 		public static void Configure(IApplicationBuilder app, IWebHostEnvironment env, Microsoft.Extensions.Hosting.IHostApplicationLifetime lifetime, IOptions<ServerSettings> settings)
 		{
 			app.UseForwardedHeaders();
-			app.UseExceptionHandler("/api/v1/error");
 
 			// Used for allowing auth cookies in combination with OpenID Connect auth (for example, Google Auth did not work with these unset)
 			app.UseCookiePolicy(new CookiePolicyOptions()
