Horde: Ensure base URLs for storage uploads have a trailing slash.

diff --git a/Engine/Source/Programs/Horde/Horde.Agent/Services/ServerStorageFactory.cs b/Engine/Source/Programs/Horde/Horde.Agent/Services/ServerStorageFactory.cs
--- a/Engine/Source/Programs/Horde/Horde.Agent/Services/ServerStorageFactory.cs
+++ b/Engine/Source/Programs/Horde/Horde.Agent/Services/ServerStorageFactory.cs
@@ -48,6 +48,11 @@ namespace Horde.Agent.Services
 
 		HttpClient CreateHttpClient(ISession session, string baseUrl)
 		{
+			if (!baseUrl.EndsWith("/", StringComparison.Ordinal))
+			{
+				baseUrl += "/";
+			}
+
 			HttpClient httpClient = _httpClientFactory.CreateClient(HttpStorageClient.HttpClientName);
 			httpClient.BaseAddress = new Uri(session.ServerUrl, baseUrl);
 			httpClient.DefaultRequestHeaders.Authorization = new AuthenticationHeaderValue("Bearer", session.Token);
