Horde: Fix default config files not specifying executor as part of job options.

diff --git a/Engine/Source/Programs/Horde/Horde.Build/Defaults/ue5-main.stream.json b/Engine/Source/Programs/Horde/Horde.Build/Defaults/ue5-main.stream.json
--- a/Engine/Source/Programs/Horde/Horde.Build/Defaults/ue5-main.stream.json
+++ b/Engine/Source/Programs/Horde/Horde.Build/Defaults/ue5-main.stream.json
@@ -53,21 +53,25 @@
 		{
 			"id": "test-executor",
 			"name": "Run Test Executor",
-			"executor": "Test",
 			"initialAgentType": "AnyAgent",
 			"arguments": [
 				"-Target=Full Build"
-			]
+			],
+			"jobOptions": {
+				"executor": "Test"
+			}
 		},
 		{
 			"id": "local-executor",
 			"name": "Run Local Executor",
-			"executor": "Local",
 			"initialAgentType": "AnyAgent",
 			"arguments": [
 				"-Script=Engine/Build/Graph/Examples/TagsAndFiles.xml",
 				"-Target=Tag and File Examples"
-			]
+			],
+			"jobOptions": {
+				"executor": "Local"
+			}
 		},
 		{
 			"id": "custom",
