Horde: Add exponential backoff to agent sessions.

diff --git a/Engine/Source/Programs/Horde/Horde.Agent/Services/WorkerService.cs b/Engine/Source/Programs/Horde/Horde.Agent/Services/WorkerService.cs
--- a/Engine/Source/Programs/Horde/Horde.Agent/Services/WorkerService.cs
+++ b/Engine/Source/Programs/Horde/Horde.Agent/Services/WorkerService.cs
@@ -27,6 +27,14 @@ namespace Horde.Agent.Services
 		readonly LeaseHandler[] _leaseHandlers;
 		readonly IServiceProvider _serviceProvider;
 
+		static readonly TimeSpan[] s_sessionBackOffTime =
+		{
+			TimeSpan.FromSeconds(5),
+			TimeSpan.FromSeconds(30),
+			TimeSpan.FromMinutes(1),
+			TimeSpan.FromMinutes(5)
+		};
+
 		/// <summary>
 		/// Constructor. Registers with the server and starts accepting connections.
 		/// </summary>
@@ -70,6 +78,7 @@ namespace Horde.Agent.Services
 			_logger.LogInformation("Arguments: {Arguments}", Environment.CommandLine);
 
 			// Keep trying to start an agent session with the server
+			int failureCount = 0;
 			while (!stoppingToken.IsCancellationRequested)
 			{
 				SessionResult result = SessionResult.Continue;
@@ -87,6 +96,8 @@ namespace Horde.Agent.Services
 							LeaseManager leaseManager = new LeaseManager(session, _capabilitiesService, _leaseHandlers, _settings, _logger);
 							result = await leaseManager.RunAsync(false, stoppingToken);
 						}
+
+						failureCount = 0;
 					}
 					catch (OperationCanceledException) when (stoppingToken.IsCancellationRequested)
 					{
@@ -95,6 +106,14 @@ namespace Horde.Agent.Services
 					catch (Exception ex)
 					{
 						_logger.LogError(ex, "Exception while executing session. Restarting. ({Message})", ex.Message);
+						if (sessionTime.Elapsed < TimeSpan.FromMinutes(5.0))
+						{
+							failureCount++;
+						}
+						else
+						{
+							failureCount = 1;
+						}
 					}
 				}
 
@@ -125,8 +144,14 @@ namespace Horde.Agent.Services
 						}
 					}
 				}
-				
-				if (sessionTime.Elapsed < TimeSpan.FromSeconds(2.0))
+
+				if (failureCount > 0)
+				{
+					TimeSpan backOffTime = s_sessionBackOffTime[Math.Min(failureCount, s_sessionBackOffTime.Length - 1)];
+					_logger.LogInformation("Session failure #{FailureNum}: Waiting {Time} before restarting...", failureCount, backOffTime);
+					await Task.Delay(backOffTime, stoppingToken);
+				}
+				else if (sessionTime.Elapsed < TimeSpan.FromSeconds(2.0))
 				{
 					_logger.LogInformation("Waiting 5 seconds before restarting session...");
 					await Task.Delay(TimeSpan.FromSeconds(5.0), stoppingToken);
