Horde: Include failure count in session failure message.

diff --git a/Engine/Source/Programs/Horde/Horde.Agent/Services/WorkerService.cs b/Engine/Source/Programs/Horde/Horde.Agent/Services/WorkerService.cs
--- a/Engine/Source/Programs/Horde/Horde.Agent/Services/WorkerService.cs
+++ b/Engine/Source/Programs/Horde/Horde.Agent/Services/WorkerService.cs
@@ -105,7 +105,6 @@ namespace Horde.Agent.Services
 					}
 					catch (Exception ex)
 					{
-						_logger.LogError(ex, "Exception while executing session. Restarting. ({Message})", ex.Message);
 						if (sessionTime.Elapsed < TimeSpan.FromMinutes(5.0))
 						{
 							failureCount++;
@@ -114,6 +113,10 @@ namespace Horde.Agent.Services
 						{
 							failureCount = 1;
 						}
+
+						TimeSpan backOffTime = s_sessionBackOffTime[Math.Min(failureCount - 1, s_sessionBackOffTime.Length - 1)];
+						_logger.LogInformation("Session failure #{FailureNum}. Waiting {Time} and restarting. ({Message})", failureCount, backOffTime, ex.Message);
+						await Task.Delay(backOffTime, stoppingToken);
 					}
 				}
 
@@ -145,13 +148,7 @@ namespace Horde.Agent.Services
 					}
 				}
 
-				if (failureCount > 0)
-				{
-					TimeSpan backOffTime = s_sessionBackOffTime[Math.Min(failureCount, s_sessionBackOffTime.Length - 1)];
-					_logger.LogInformation("Session failure #{FailureNum}: Waiting {Time} before restarting...", failureCount, backOffTime);
-					await Task.Delay(backOffTime, stoppingToken);
-				}
-				else if (sessionTime.Elapsed < TimeSpan.FromSeconds(2.0))
+				if (sessionTime.Elapsed < TimeSpan.FromSeconds(2.0))
 				{
 					_logger.LogInformation("Waiting 5 seconds before restarting session...");
 					await Task.Delay(TimeSpan.FromSeconds(5.0), stoppingToken);
