Horde: Issue a custom JWT for each artifact which allows uploading directly to the storage backend for a particular ref and namespace.

diff --git a/Engine/Source/Programs/Horde/Horde.Agent.Tests/WorkerServiceTest.cs b/Engine/Source/Programs/Horde/Horde.Agent.Tests/WorkerServiceTest.cs
--- a/Engine/Source/Programs/Horde/Horde.Agent.Tests/WorkerServiceTest.cs
+++ b/Engine/Source/Programs/Horde/Horde.Agent.Tests/WorkerServiceTest.cs
@@ -61,7 +61,7 @@ namespace Horde.Agent.Tests
 
 		class FakeSessionStorageFactory : IServerStorageFactory
 		{
-			public IStorageClient CreateStorageClient(ISession session, string basePath) => null!;
+			public IStorageClient CreateStorageClient(Uri baseAddress, string token) => null!;
 		}
 
 		internal static IJobExecutor NullExecutor = new SimpleTestExecutor(async (step, logger, cancellationToken) =>
diff --git a/Engine/Source/Programs/Horde/Horde.Agent/Execution/JobExecutor.cs b/Engine/Source/Programs/Horde/Horde.Agent/Execution/JobExecutor.cs
--- a/Engine/Source/Programs/Horde/Horde.Agent/Execution/JobExecutor.cs
+++ b/Engine/Source/Programs/Horde/Horde.Agent/Execution/JobExecutor.cs
@@ -471,6 +471,7 @@ namespace Horde.Agent.Execution
 
 		protected virtual async Task<bool> SetupAsync(BeginStepResponse step, DirectoryReference workspaceDir, DirectoryReference? sharedStorageDir, ILogger logger, CancellationToken cancellationToken)
 		{
+			await TestArtifactAsync(step.StepId, cancellationToken);
 			FileReference definitionFile = FileReference.Combine(workspaceDir, "Engine", "Saved", "Horde", "Exported.json");
 
 			StringBuilder arguments = new StringBuilder($"BuildGraph");
@@ -742,7 +743,7 @@ namespace Horde.Agent.Execution
 			}
 		}
 
-		async Task TestArtifactAsync(string stepId, CancellationToken cancellationToken)
+		protected async Task TestArtifactAsync(string stepId, CancellationToken cancellationToken)
 		{
 			try
 			{
@@ -750,9 +751,7 @@ namespace Horde.Agent.Execution
 				CreateJobArtifactResponse artifact = await jobRpc.Client.CreateArtifactAsync(new CreateJobArtifactRequest { JobId = _jobId, StepId = stepId, Name = "test artifact", Type = JobArtifactType.Saved }, cancellationToken: cancellationToken);
 				_logger.LogInformation("Created artifact {ArtifactId} with ref {RefName} in ns {Namespace}", artifact.Id, artifact.RefName, artifact.NamespaceId);
 
-				using MemoryCache cache = new MemoryCache(new MemoryCacheOptions { });
-				IStorageClient storage = _storageFactory.CreateStorageClient(_session, $"/api/v1/storage/{artifact.NamespaceId}");
-
+				IStorageClient storage = _storageFactory.CreateStorageClient(_session, artifact.NamespaceId, artifact.Token);
 				using TreeWriter writer = new TreeWriter(storage, new RefName(artifact.RefName));
 
 				DirectoryNode dir = new DirectoryNode();
diff --git a/Engine/Source/Programs/Horde/Horde.Agent/Services/ServerStorageFactory.cs b/Engine/Source/Programs/Horde/Horde.Agent/Services/ServerStorageFactory.cs
--- a/Engine/Source/Programs/Horde/Horde.Agent/Services/ServerStorageFactory.cs
+++ b/Engine/Source/Programs/Horde/Horde.Agent/Services/ServerStorageFactory.cs
@@ -5,6 +5,7 @@ using System.Net.Http;
 using System.Net.Http.Headers;
 using EpicGames.Horde.Storage;
 using EpicGames.Horde.Storage.Backends;
+using Microsoft.Extensions.Caching.Memory;
 using Microsoft.Extensions.Logging;
 
 namespace Horde.Agent.Services
@@ -17,10 +18,55 @@ namespace Horde.Agent.Services
 		/// <summary>
 		/// Creates a storage client which authenticates using tokens for the current session
 		/// </summary>
+		/// <param name="baseAddress">Base URL for writing blobs on the server</param>
+		/// <param name="token">Token to use for authentication</param>
+		/// <returns>New logger instance</returns>
+		IStorageClient CreateStorageClient(Uri baseAddress, string token);
+	}
+
+	/// <summary>
+	/// Extension methods for <see cref="IServerStorageFactory"/>
+	/// </summary>
+	static class ServerStorageExtensions
+	{
+		/// <summary>
+		/// Creates a storage client which authenticates using tokens for the current session
+		/// </summary>
+		/// <param name="factory">Factory to create the client</param>
 		/// <param name="session">The current session</param>
 		/// <param name="baseUrl">Base URL for writing blobs on the server</param>
 		/// <returns>New logger instance</returns>
-		IStorageClient CreateStorageClient(ISession session, string baseUrl);
+		public static IStorageClient CreateStorageClient(this IServerStorageFactory factory, ISession session, string baseUrl) => CreateStorageClient(factory, session, baseUrl, session.Token);
+
+		/// <summary>
+		/// Creates a storage client which authenticates using tokens for the current session
+		/// </summary>
+		/// <param name="factory">Factory to create the client</param>
+		/// <param name="session">The current session</param>
+		/// <param name="baseUrl">Base URL for writing blobs on the server</param>
+		/// <param name="token">Bearer token for the server</param>
+		/// <returns>New logger instance</returns>
+		public static IStorageClient CreateStorageClient(this IServerStorageFactory factory, ISession session, string baseUrl, string token)
+		{
+			if (!baseUrl.EndsWith("/", StringComparison.Ordinal))
+			{
+				baseUrl += "/";
+			}
+			return factory.CreateStorageClient(new Uri(session.ServerUrl, baseUrl), token);
+		}
+
+		/// <summary>
+		/// Creates a storage client which authenticates using tokens for the current session
+		/// </summary>
+		/// <param name="factory">Factory to create the client</param>
+		/// <param name="session">The current session</param>
+		/// <param name="namespaceId">Namespace to write to</param>
+		/// <param name="token">Bearer token for the server</param>
+		/// <returns>New logger instance</returns>
+		public static IStorageClient CreateStorageClient(this IServerStorageFactory factory, ISession session, NamespaceId namespaceId, string token)
+		{
+			return factory.CreateStorageClient(new Uri(session.ServerUrl, $"/api/v1/storage/{namespaceId}/"), token);
+		}
 	}
 
 	/// <summary>
@@ -41,21 +87,16 @@ namespace Horde.Agent.Services
 		}
 
 		/// <inheritdoc/>
-		public IStorageClient CreateStorageClient(ISession session, string baseUrl)
+		public IStorageClient CreateStorageClient(Uri baseAddress, string token)
 		{
-			return new HttpStorageClient(() => CreateHttpClient(session, baseUrl), CreateHttpRedirectClient, _logger);
+			return new HttpStorageClient(() => CreateHttpClient(baseAddress, token), CreateHttpRedirectClient, _logger);
 		}
 
-		HttpClient CreateHttpClient(ISession session, string baseUrl)
+		HttpClient CreateHttpClient(Uri baseAddress, string token)
 		{
-			if (!baseUrl.EndsWith("/", StringComparison.Ordinal))
-			{
-				baseUrl += "/";
-			}
-
 			HttpClient httpClient = _httpClientFactory.CreateClient(HttpStorageClient.HttpClientName);
-			httpClient.BaseAddress = new Uri(session.ServerUrl, baseUrl);
-			httpClient.DefaultRequestHeaders.Authorization = new AuthenticationHeaderValue("Bearer", session.Token);
+			httpClient.BaseAddress = baseAddress;
+			httpClient.DefaultRequestHeaders.Authorization = new AuthenticationHeaderValue("Bearer", token);
 			return httpClient;
 		}
 
diff --git a/Engine/Source/Programs/Horde/Horde.Build/Artifacts/ArtifactsController.cs b/Engine/Source/Programs/Horde/Horde.Build/Artifacts/ArtifactsController.cs
--- a/Engine/Source/Programs/Horde/Horde.Build/Artifacts/ArtifactsController.cs
+++ b/Engine/Source/Programs/Horde/Horde.Build/Artifacts/ArtifactsController.cs
@@ -157,31 +157,6 @@ namespace Horde.Build.Artifacts
 			return PropertyFilter.Apply(new GetArtifactResponse(artifact), filter);
 		}
 
-		/// <summary>
-		/// Uploads a blob for an artifact. See /api/v1/storage/XXX/blobs.
-		/// </summary>
-		/// <param name="id">Id of the artifact to write to</param>
-		/// <param name="file">Data for the blob</param>
-		/// <param name="cancellationToken">Cancellation token for the request</param>
-		/// <returns>Information about the requested project</returns>
-		[HttpPost]
-		[Route("/api/v2/artifacts/{id}/blobs")]
-		[ProducesResponseType(typeof(WriteBlobResponse), 200)]
-		public async Task<ActionResult<WriteBlobResponse>> WriteArtifactBlobAsync(ArtifactId id, IFormFile? file, CancellationToken cancellationToken = default)
-		{
-			IArtifact? artifact = await _artifactCollection.GetAsync(id, cancellationToken);
-			if (artifact == null)
-			{
-				return NotFound(id);
-			}
-			if (!_globalConfig.Authorize(artifact.AclScope, AclAction.WriteArtifact, User))
-			{
-				return Forbid(AclAction.WriteArtifact, artifact.AclScope);
-			}
-
-			return await StorageController.WriteBlobAsync(_storageService, Namespace.Artifacts, file, $"{artifact.RefName}", cancellationToken);
-		}
-
 		/// <summary>
 		/// Gets metadata about an artifact object
 		/// </summary>
diff --git a/Engine/Source/Programs/Horde/Horde.Build/Jobs/JobRpcService.cs b/Engine/Source/Programs/Horde/Horde.Build/Jobs/JobRpcService.cs
--- a/Engine/Source/Programs/Horde/Horde.Build/Jobs/JobRpcService.cs
+++ b/Engine/Source/Programs/Horde/Horde.Build/Jobs/JobRpcService.cs
@@ -25,6 +25,7 @@ namespace Horde.Build.Jobs
 	[Authorize]
 	public class JobRpcService : JobRpc.JobRpcBase
 	{
+		readonly AclService _aclService;
 		readonly IJobCollection _jobCollection;
 		readonly IArtifactCollection _artifactCollection;
 		readonly GlobalConfig _globalConfig;
@@ -32,8 +33,9 @@ namespace Horde.Build.Jobs
 		/// <summary>
 		/// Constructor
 		/// </summary>
-		public JobRpcService(IJobCollection jobCollection, IArtifactCollection artifactCollection, IOptionsSnapshot<GlobalConfig> globalConfig)
+		public JobRpcService(AclService aclService, IJobCollection jobCollection, IArtifactCollection artifactCollection, IOptionsSnapshot<GlobalConfig> globalConfig)
 		{
+			_aclService = aclService;
 			_jobCollection = jobCollection;
 			_artifactCollection = artifactCollection;
 			_globalConfig = globalConfig.Value;
@@ -62,25 +64,15 @@ namespace Horde.Build.Jobs
 
 			NamespaceId namespaceId = Namespace.Artifacts;
 			IArtifact artifact = await _artifactCollection.AddAsync(request.Name, type, keys, namespaceId, null, templateConfig.ScopeName, context.CancellationToken);
-			// TODO: Token
-			return new CreateJobArtifactResponse { Id = artifact.Id.ToString(), NamespaceId = artifact.NamespaceId.ToString(), RefName = artifact.RefName.ToString() };
-		}
-		/*
-		/// <inheritdoc/>
-		public override async Task<FinalizeJobArtifactResponse> FinalizeArtifact(FinalizeJobArtifactRequest request, ServerCallContext context)
-		{
-			IArtifact? artifact = await _artifactCollection.GetAsync(ArtifactId.Parse(request.Id), context.CancellationToken);
-			if (artifact == null)
-			{
-				throw new StructuredRpcException(StatusCode.NotFound, "Missing artifact {ArtifactId}", request.Id);
-			}
 
-			IStorageClientImpl storage = await _storageService.GetClientAsync(artifact.NamespaceId, context.CancellationToken);
-			await storage.WriteRefTargetAsync(artifact.RefName, NodeHandle.Parse(request.Locator), cancellationToken: context.CancellationToken);
+			List<AclClaimConfig> claims = new List<AclClaimConfig>();
+			claims.Add(new AclClaimConfig(HordeClaimTypes.WriteNamespace, namespaceId.ToString()));
+			claims.Add(new AclClaimConfig(HordeClaimTypes.WriteRef, artifact.RefName.ToString()));
 
-			return new FinalizeJobArtifactResponse();
+			string token = await _aclService.IssueBearerTokenAsync(claims, TimeSpan.FromHours(8.0));
+			return new CreateJobArtifactResponse { Id = artifact.Id.ToString(), NamespaceId = artifact.NamespaceId.ToString(), RefName = artifact.RefName.ToString(), Token = token };
 		}
-		*/
+
 		Task<(IJob, IJobStepBatch, IJobStep)> AuthorizeAsync(string jobId, string stepId, ServerCallContext context)
 		{
 			return AuthorizeAsync(JobId.Parse(jobId), SubResourceId.Parse(stepId), context);
diff --git a/Engine/Source/Programs/Horde/Horde.Build/Storage/StorageController.cs b/Engine/Source/Programs/Horde/Horde.Build/Storage/StorageController.cs
--- a/Engine/Source/Programs/Horde/Horde.Build/Storage/StorageController.cs
+++ b/Engine/Source/Programs/Horde/Horde.Build/Storage/StorageController.cs
@@ -176,7 +176,7 @@ namespace Horde.Build.Storage
 			{
 				return NotFound(namespaceId);
 			}
-			if (!namespaceConfig.Authorize(AclAction.WriteBlobs, User))
+			if (!namespaceConfig.Authorize(AclAction.WriteBlobs, User) && !User.HasClaim(HordeClaimTypes.WriteNamespace, namespaceId.ToString()))
 			{
 				return Forbid(AclAction.WriteBlobs, namespaceId);
 			}
@@ -339,7 +339,7 @@ namespace Horde.Build.Storage
 			{
 				return NotFound(namespaceId);
 			}
-			if (!namespaceConfig.Authorize(AclAction.WriteRefs, User))
+			if (!namespaceConfig.Authorize(AclAction.WriteRefs, User) && (!User.HasClaim(HordeClaimTypes.WriteNamespace, namespaceConfig.Id.ToString()) || !User.HasClaim(HordeClaimTypes.WriteRef, refName.ToString())))
 			{
 				return Forbid(AclAction.WriteRefs, namespaceId);
 			}
diff --git a/Engine/Source/Programs/Horde/Horde.Build/Utilities/HordeClaimTypes.cs b/Engine/Source/Programs/Horde/Horde.Build/Utilities/HordeClaimTypes.cs
--- a/Engine/Source/Programs/Horde/Horde.Build/Utilities/HordeClaimTypes.cs
+++ b/Engine/Source/Programs/Horde/Horde.Build/Utilities/HordeClaimTypes.cs
@@ -61,6 +61,15 @@ namespace Horde.Build.Utilities
 		/// </summary>
 		public const string ExternalIssueUser = Prefix + "external-issue-user";
 
+		/// <summary>
+		/// Claim allowing a certain namespace to be written to
+		/// </summary>
+		public const string WriteNamespace = Prefix + "write-ns";
+
+		/// <summary>
+		/// Claim allowing a certain ref to be written
+		/// </summary>
+		public const string WriteRef = Prefix + "write-ref";
 	}
 
 	/// <summary>
