Horde: Fix server URL for uploading artifacts.

diff --git a/Engine/Source/Programs/Horde/Horde.Agent/Execution/JobExecutor.cs b/Engine/Source/Programs/Horde/Horde.Agent/Execution/JobExecutor.cs
--- a/Engine/Source/Programs/Horde/Horde.Agent/Execution/JobExecutor.cs
+++ b/Engine/Source/Programs/Horde/Horde.Agent/Execution/JobExecutor.cs
@@ -751,7 +751,7 @@ namespace Horde.Agent.Execution
 				CreateJobArtifactResponse artifact = await jobRpc.Client.CreateArtifactAsync(new CreateJobArtifactRequest { JobId = _jobId, StepId = stepId, Name = "test artifact", Type = JobArtifactType.Saved }, cancellationToken: cancellationToken);
 				_logger.LogInformation("Created artifact {ArtifactId} with ref {RefName} in ns {Namespace}", artifact.Id, artifact.RefName, artifact.NamespaceId);
 
-				IStorageClient storage = _storageFactory.CreateStorageClient(_session, artifact.NamespaceId, artifact.Token);
+				IStorageClient storage = _storageFactory.CreateStorageClient(_session, new NamespaceId(artifact.NamespaceId), artifact.Token);
 				using TreeWriter writer = new TreeWriter(storage, new RefName(artifact.RefName));
 
 				DirectoryNode dir = new DirectoryNode();
