Horde: Fix log tailing not getting enabled correctly on first query from dashboard.

diff --git a/Engine/Source/Programs/Horde/Horde.Build/Logs/LogFileService.cs b/Engine/Source/Programs/Horde/Horde.Build/Logs/LogFileService.cs
--- a/Engine/Source/Programs/Horde/Horde.Build/Logs/LogFileService.cs
+++ b/Engine/Source/Programs/Horde/Horde.Build/Logs/LogFileService.cs
@@ -925,13 +925,7 @@ namespace Horde.Build.Logs
 			LogMetadata metadata = new LogMetadata();
 			if (logFile.UseNewStorageBackend)
 			{
-				metadata.MaxLineIndex = logFile.LineCount;
-
-				int nextIdx = await _logTailService.GetTailNextAsync(logFile.Id);
-				if (nextIdx != -1)
-				{
-					metadata.MaxLineIndex = nextIdx;
-				}
+				metadata.MaxLineIndex = await _logTailService.GetFullLineCount(logFile.Id, logFile.LineCount);
 			}
 			else
 			{
diff --git a/Engine/Source/Programs/Horde/Horde.Build/Logs/LogTailService.cs b/Engine/Source/Programs/Horde/Horde.Build/Logs/LogTailService.cs
--- a/Engine/Source/Programs/Horde/Horde.Build/Logs/LogTailService.cs
+++ b/Engine/Source/Programs/Horde/Horde.Build/Logs/LogTailService.cs
@@ -267,13 +267,22 @@ namespace Horde.Build.Logs
 		}
 
 		/// <summary>
-		/// Return the desired next line to append for the given log file, or -1 if it is not being tailed.
+		/// Gets the total number of lines in a log
 		/// </summary>
 		/// <param name="logId">Log to query</param>
-		/// <returns>Number of lines for the log file</returns>
-		public async Task<int> GetTailNextAsync(LogId logId)
+		/// <param name="flushedLineCount">Number of flushed lines</param>
+		/// <returns>Total number of lines in the file</returns>
+		public async Task<int> GetFullLineCount(LogId logId, int flushedLineCount)
 		{
-			return await _redisService.GetDatabase().StringGetAsync(TailNextKey(logId), -1);
+			RedisStringKey<int> tailNextKey = TailNextKey(logId);
+
+			int lineCount = await _redisService.GetDatabase().StringGetAsync(tailNextKey, -1);
+			if (lineCount == -1)
+			{
+				await EnableTailingAsync(logId, flushedLineCount);
+			}
+
+			return Math.Max(lineCount, flushedLineCount);
 		}
 
 		/// <summary>
