Horde: Add more logging for tailing service.

diff --git a/Engine/Source/Programs/Horde/Horde.Build/Logs/LogTailService.cs b/Engine/Source/Programs/Horde/Horde.Build/Logs/LogTailService.cs
--- a/Engine/Source/Programs/Horde/Horde.Build/Logs/LogTailService.cs
+++ b/Engine/Source/Programs/Horde/Horde.Build/Logs/LogTailService.cs
@@ -127,6 +127,8 @@ namespace Horde.Build.Logs
 
 				int minLineIndex = (lineCount & ~(ChunkLineCount - 1)) - 1;
 				_ = await _redisService.GetDatabase().SortedSetRemoveRangeByScoreAsync(TailDataKey(logId), Double.NegativeInfinity, minLineIndex, Exclude.None, CommandFlags.FireAndForget);
+
+				_logger.LogDebug("Removed tail data for log {LogId} at line {LineCount}", logId, lineCount);
 			}
 
 			SortedSetEntry<LogId>[] expireEntries = await _redisService.GetDatabase().SortedSetRangeByScoreWithScoresAsync(_expireQueue, stop: maxScore);
@@ -139,6 +141,14 @@ namespace Horde.Build.Logs
 				_ = transaction.KeyDeleteAsync(TailNextKey(logId), CommandFlags.FireAndForget);
 				_ = transaction.KeyDeleteAsync(TailDataKey(logId), CommandFlags.FireAndForget);
 				_ = transaction.ExecuteAsync(CommandFlags.FireAndForget);
+
+				_logger.LogDebug("Expired tailing request for log {LogId}", logId);
+			}
+
+			long tailCount = await _redisService.GetDatabase().SortedSetLengthAsync(_expireQueue);
+			if (tailCount > 0)
+			{
+				_logger.LogDebug("Currently tailing {NumLogs} logs", tailCount);
 			}
 		}
 
@@ -191,7 +201,7 @@ namespace Horde.Build.Logs
 				_ = transaction.StringSetAsync(redisTailNext, newLineCount, flags: CommandFlags.FireAndForget);
 
 				bool result = await transaction.ExecuteAsync();
-				_logger.LogTrace("Tail data for {LogId}, line {TailNext} -> {NewTailNext}, result: {Result}", logId, tailNext, newLineCount, result);
+				_logger.LogDebug("Tail data for {LogId}, line {TailNext} -> {NewTailNext}, result: {Result}", logId, tailNext, newLineCount, result);
 			}
 		}
 
