Horde: Fix log tail entries not being expired from Redis correctly.

diff --git a/Engine/Source/Programs/Horde/Horde.Build/Logs/LogTailService.cs b/Engine/Source/Programs/Horde/Horde.Build/Logs/LogTailService.cs
--- a/Engine/Source/Programs/Horde/Horde.Build/Logs/LogTailService.cs
+++ b/Engine/Source/Programs/Horde/Horde.Build/Logs/LogTailService.cs
@@ -130,6 +130,7 @@ namespace Horde.Build.Logs
 
 				_logger.LogDebug("Removed tail data for log {LogId} at line {LineCount}", logId, lineCount);
 			}
+			_ = await _redisService.GetDatabase().SortedSetRemoveRangeByScoreAsync(_trimQueue, Double.NegativeInfinity, maxScore, flags: CommandFlags.FireAndForget);
 
 			SortedSetEntry<LogId>[] expireEntries = await _redisService.GetDatabase().SortedSetRangeByScoreWithScoresAsync(_expireQueue, stop: maxScore);
 			foreach (SortedSetEntry<LogId> expireEntry in expireEntries)
@@ -144,6 +145,7 @@ namespace Horde.Build.Logs
 
 				_logger.LogDebug("Expired tailing request for log {LogId}", logId);
 			}
+			_ = await _redisService.GetDatabase().SortedSetRemoveRangeByScoreAsync(_expireQueue, Double.NegativeInfinity, maxScore, flags: CommandFlags.FireAndForget);
 
 			long tailCount = await _redisService.GetDatabase().SortedSetLengthAsync(_expireQueue);
 			if (tailCount > 0)
