Horde: Print tailing queue length whenever we expire entries.

diff --git a/Engine/Source/Programs/Horde/Horde.Build/Logs/LogTailService.cs b/Engine/Source/Programs/Horde/Horde.Build/Logs/LogTailService.cs
--- a/Engine/Source/Programs/Horde/Horde.Build/Logs/LogTailService.cs
+++ b/Engine/Source/Programs/Horde/Horde.Build/Logs/LogTailService.cs
@@ -118,38 +118,41 @@ namespace Horde.Build.Logs
 			double maxScore = (_clock.UtcNow - DateTime.UnixEpoch).TotalSeconds;
 
 			SortedSetEntry<string>[] trimEntries = await _redisService.GetDatabase().SortedSetRangeByScoreWithScoresAsync(_trimQueue, stop: maxScore);
-			foreach (SortedSetEntry<string> trimEntry in trimEntries)
+			if (trimEntries.Length > 0)
 			{
-				string[] values = trimEntry.Element.Split('=');
+				foreach (SortedSetEntry<string> trimEntry in trimEntries)
+				{
+					string[] values = trimEntry.Element.Split('=');
 
-				LogId logId = LogId.Parse(values[0]);
-				int lineCount = Int32.Parse(values[1], NumberStyles.None, CultureInfo.InvariantCulture);
+					LogId logId = LogId.Parse(values[0]);
+					int lineCount = Int32.Parse(values[1], NumberStyles.None, CultureInfo.InvariantCulture);
 
-				int minLineIndex = (lineCount & ~(ChunkLineCount - 1)) - 1;
-				_ = await _redisService.GetDatabase().SortedSetRemoveRangeByScoreAsync(TailDataKey(logId), Double.NegativeInfinity, minLineIndex, Exclude.None, CommandFlags.FireAndForget);
+					int minLineIndex = (lineCount & ~(ChunkLineCount - 1)) - 1;
+					_ = await _redisService.GetDatabase().SortedSetRemoveRangeByScoreAsync(TailDataKey(logId), Double.NegativeInfinity, minLineIndex, Exclude.None, CommandFlags.FireAndForget);
 
-				_logger.LogDebug("Removed tail data for log {LogId} at line {LineCount}", logId, lineCount);
+					_logger.LogDebug("Removed tail data for log {LogId} at line {LineCount}", logId, lineCount);
+				}
+				_ = await _redisService.GetDatabase().SortedSetRemoveRangeByScoreAsync(_trimQueue, Double.NegativeInfinity, maxScore, flags: CommandFlags.FireAndForget);
 			}
-			_ = await _redisService.GetDatabase().SortedSetRemoveRangeByScoreAsync(_trimQueue, Double.NegativeInfinity, maxScore, flags: CommandFlags.FireAndForget);
 
 			SortedSetEntry<LogId>[] expireEntries = await _redisService.GetDatabase().SortedSetRangeByScoreWithScoresAsync(_expireQueue, stop: maxScore);
-			foreach (SortedSetEntry<LogId> expireEntry in expireEntries)
+			if (expireEntries.Length > 0)
 			{
-				LogId logId = expireEntry.Element;
+				foreach (SortedSetEntry<LogId> expireEntry in expireEntries)
+				{
+					LogId logId = expireEntry.Element;
 
-				ITransaction transaction = _redisService.GetDatabase().CreateTransaction();
-				transaction.AddCondition(Condition.SortedSetEqual(_expireQueue.Inner, expireEntry.ElementValue, expireEntry.Score));
-				_ = transaction.KeyDeleteAsync(TailNextKey(logId), CommandFlags.FireAndForget);
-				_ = transaction.KeyDeleteAsync(TailDataKey(logId), CommandFlags.FireAndForget);
-				_ = transaction.ExecuteAsync(CommandFlags.FireAndForget);
+					ITransaction transaction = _redisService.GetDatabase().CreateTransaction();
+					transaction.AddCondition(Condition.SortedSetEqual(_expireQueue.Inner, expireEntry.ElementValue, expireEntry.Score));
+					_ = transaction.KeyDeleteAsync(TailNextKey(logId), CommandFlags.FireAndForget);
+					_ = transaction.KeyDeleteAsync(TailDataKey(logId), CommandFlags.FireAndForget);
+					await transaction.ExecuteAsync();
 
-				_logger.LogDebug("Expired tailing request for log {LogId}", logId);
-			}
-			_ = await _redisService.GetDatabase().SortedSetRemoveRangeByScoreAsync(_expireQueue, Double.NegativeInfinity, maxScore, flags: CommandFlags.FireAndForget);
+					_logger.LogDebug("Expired tailing request for log {LogId}", logId);
+				}
+				await _redisService.GetDatabase().SortedSetRemoveRangeByScoreAsync(_expireQueue, Double.NegativeInfinity, maxScore);
 
-			long tailCount = await _redisService.GetDatabase().SortedSetLengthAsync(_expireQueue);
-			if (tailCount > 0)
-			{
+				long tailCount = await _redisService.GetDatabase().SortedSetLengthAsync(_expireQueue);
 				_logger.LogDebug("Currently tailing {NumLogs} logs", tailCount);
 			}
 		}
