Horde: Avoid querying log tails once a log file is marked as complete.

diff --git a/Engine/Source/Programs/Horde/Horde.Build/Logs/ILogFile.cs b/Engine/Source/Programs/Horde/Horde.Build/Logs/ILogFile.cs
--- a/Engine/Source/Programs/Horde/Horde.Build/Logs/ILogFile.cs
+++ b/Engine/Source/Programs/Horde/Horde.Build/Logs/ILogFile.cs
@@ -91,6 +91,11 @@ namespace Horde.Build.Logs
 		/// Number of lines (V2 storage backend)
 		/// </summary>
 		public int LineCount { get; }
+
+		/// <summary>
+		/// Whether the log is complete (V2 storage backend)
+		/// </summary>
+		public bool Complete { get; }
 	}
 
 	/// <summary>
diff --git a/Engine/Source/Programs/Horde/Horde.Build/Logs/ILogFileCollection.cs b/Engine/Source/Programs/Horde/Horde.Build/Logs/ILogFileCollection.cs
--- a/Engine/Source/Programs/Horde/Horde.Build/Logs/ILogFileCollection.cs
+++ b/Engine/Source/Programs/Horde/Horde.Build/Logs/ILogFileCollection.cs
@@ -65,9 +65,10 @@ namespace Horde.Build.Logs
 		/// </summary>
 		/// <param name="logFileInterface">The current log file</param>
 		/// <param name="lineCount">New line count for the log file</param>
+		/// <param name="complete">Flag indicating whether the log is complete, or can still be tailed for additional data</param>
 		/// <param name="cancellationToken">Cancellation token for the call</param>
 		/// <returns>The updated log file document</returns>
-		Task<ILogFile> UpdateLineCountAsync(ILogFile logFileInterface, int lineCount, CancellationToken cancellationToken);
+		Task<ILogFile> UpdateLineCountAsync(ILogFile logFileInterface, int lineCount, bool complete, CancellationToken cancellationToken);
 
 		/// <summary>
 		/// Adds a new chunk
diff --git a/Engine/Source/Programs/Horde/Horde.Build/Logs/LogFileCollection.cs b/Engine/Source/Programs/Horde/Horde.Build/Logs/LogFileCollection.cs
--- a/Engine/Source/Programs/Horde/Horde.Build/Logs/LogFileCollection.cs
+++ b/Engine/Source/Programs/Horde/Horde.Build/Logs/LogFileCollection.cs
@@ -72,6 +72,9 @@ namespace Horde.Build.Logs
 			public int LineCount { get; set; }
 			public RefName RefName { get; set; }
 
+			[BsonIgnoreIfDefault]
+			public bool Complete { get; set; }
+
 			[BsonRequired]
 			public int UpdateIndex { get; set; }
 
@@ -130,10 +133,10 @@ namespace Horde.Build.Logs
 		}
 
 		/// <inheritdoc/>
-		public async Task<ILogFile> UpdateLineCountAsync(ILogFile logFileInterface, int lineCount, CancellationToken cancellationToken)
+		public async Task<ILogFile> UpdateLineCountAsync(ILogFile logFileInterface, int lineCount, bool complete, CancellationToken cancellationToken)
 		{
 			FilterDefinition<LogFileDocument> filter = Builders<LogFileDocument>.Filter.Eq(x => x.Id, logFileInterface.Id);
-			UpdateDefinition<LogFileDocument> update = Builders<LogFileDocument>.Update.Set(x => x.LineCount, lineCount).Inc(x => x.UpdateIndex, 1);
+			UpdateDefinition<LogFileDocument> update = Builders<LogFileDocument>.Update.Set(x => x.LineCount, lineCount).Set(x => x.Complete, complete).Inc(x => x.UpdateIndex, 1);
 			return await _logFiles.FindOneAndUpdateAsync(filter, update, new FindOneAndUpdateOptions<LogFileDocument, LogFileDocument> { ReturnDocument = ReturnDocument.After }, cancellationToken);
 		}
 
diff --git a/Engine/Source/Programs/Horde/Horde.Build/Logs/LogFileService.cs b/Engine/Source/Programs/Horde/Horde.Build/Logs/LogFileService.cs
--- a/Engine/Source/Programs/Horde/Horde.Build/Logs/LogFileService.cs
+++ b/Engine/Source/Programs/Horde/Horde.Build/Logs/LogFileService.cs
@@ -673,7 +673,7 @@ namespace Horde.Build.Logs
 					}
 				}
 
-				if (root == null || !root.Complete)
+				if (!logFile.Complete)
 				{
 					await _logTailService.EnableTailingAsync(logFile.Id, root?.LineCount ?? 0);
 					if (index < maxIndex)
@@ -925,7 +925,14 @@ namespace Horde.Build.Logs
 			LogMetadata metadata = new LogMetadata();
 			if (logFile.UseNewStorageBackend)
 			{
-				metadata.MaxLineIndex = await _logTailService.GetFullLineCount(logFile.Id, logFile.LineCount);
+				if (logFile.Complete)
+				{
+					metadata.MaxLineIndex = logFile.LineCount;
+				}
+				else
+				{
+					metadata.MaxLineIndex = await _logTailService.GetFullLineCount(logFile.Id, logFile.LineCount);
+				}
 			}
 			else
 			{
@@ -1899,23 +1906,26 @@ namespace Horde.Build.Logs
 			}
 
 			// Search any tail data we have
-			for (; ; )
+			if (!logFile.Complete)
 			{
-				Utf8String[] lines = await ReadTailAsync(logFile, firstLine, cancellationToken);
-				if (lines.Length == 0)
+				for (; ; )
 				{
-					break;
-				}
+					Utf8String[] lines = await ReadTailAsync(logFile, firstLine, cancellationToken);
+					if (lines.Length == 0)
+					{
+						break;
+					}
 
-				for (int idx = 0; idx < lines.Length; idx++)
-				{
-					if (SearchTerm.FindNextOcurrence(lines[idx].Span, 0, searchText) != -1)
+					for (int idx = 0; idx < lines.Length; idx++)
 					{
-						yield return firstLine + idx;
+						if (SearchTerm.FindNextOcurrence(lines[idx].Span, 0, searchText) != -1)
+						{
+							yield return firstLine + idx;
+						}
 					}
-				}
 
-				firstLine += lines.Length;
+					firstLine += lines.Length;
+				}
 			}
 		}
 
@@ -1923,6 +1933,11 @@ namespace Horde.Build.Logs
 		{
 			const int BatchSize = 128;
 
+			if (logFile.Complete)
+			{
+				return Array.Empty<Utf8String>();
+			}
+
 			string cacheKey = $"{logFile.Id}@{index}";
 			if (_logFileCache.TryGetValue(cacheKey, out Utf8String[]? lines))
 			{
diff --git a/Engine/Source/Programs/Horde/Horde.Build/Logs/LogRpcService.cs b/Engine/Source/Programs/Horde/Horde.Build/Logs/LogRpcService.cs
--- a/Engine/Source/Programs/Horde/Horde.Build/Logs/LogRpcService.cs
+++ b/Engine/Source/Programs/Horde/Horde.Build/Logs/LogRpcService.cs
@@ -50,10 +50,10 @@ namespace Horde.Build.Logs
 			}
 
 			IStorageClient store = await _storageService.GetClientAsync(Namespace.Logs, context.CancellationToken);
-			_logger.LogInformation("Updating {LogId} to node {RefTarget}", request.LogId, request.Target);
+			_logger.LogInformation("Updating {LogId} to node {RefTarget} (lines: {LineCount}, complete: {Complete})", request.LogId, request.Target, request.LineCount, request.Complete);
 			await store.WriteRefTargetAsync(new RefName(request.LogId), NodeHandle.Parse(request.Target));
 
-			await _logFileCollection.UpdateLineCountAsync(logFile, request.LineCount, CancellationToken.None);
+			await _logFileCollection.UpdateLineCountAsync(logFile, request.LineCount, request.Complete, CancellationToken.None);
 
 			await _logTailService.FlushAsync(logFile.Id, request.LineCount);
 
