Horde: Set the complete flag on log updates once the last packet is sent.

diff --git a/Engine/Source/Programs/Horde/Horde.Agent/Utility/JsonRpcLogSink.cs b/Engine/Source/Programs/Horde/Horde.Agent/Utility/JsonRpcLogSink.cs
--- a/Engine/Source/Programs/Horde/Horde.Agent/Utility/JsonRpcLogSink.cs
+++ b/Engine/Source/Programs/Horde/Horde.Agent/Utility/JsonRpcLogSink.cs
@@ -200,7 +200,7 @@ namespace Horde.Agent.Utility
 			if (request.Flush || _bufferLength > FlushLength)
 			{
 				NodeHandle target = await _builder.FlushAsync(_writer, request.Flush, cancellationToken);
-				await UpdateLogAsync(target, _builder.LineCount, cancellationToken);
+				await UpdateLogAsync(target, _builder.LineCount, request.Flush, cancellationToken);
 				_bufferLength = 0;
 			}
 
@@ -209,7 +209,7 @@ namespace Horde.Agent.Utility
 
 		#region RPC calls
 
-		protected virtual async Task UpdateLogAsync(NodeHandle target, int lineCount, CancellationToken cancellationToken)
+		protected virtual async Task UpdateLogAsync(NodeHandle target, int lineCount, bool complete, CancellationToken cancellationToken)
 		{
 			_logger.LogInformation("Updating log {LogId} to line {LineCount}, target {Locator}", _logId, lineCount, target);
 
@@ -217,6 +217,7 @@ namespace Horde.Agent.Utility
 			request.LogId = _logId;
 			request.LineCount = lineCount;
 			request.Target = target.ToString();
+			request.Complete = complete;
 			await _connection.InvokeAsync((LogRpcClient client) => client.UpdateLogAsync(request, cancellationToken: cancellationToken), cancellationToken);
 		}
 
