Horde: Fix tests.

diff --git a/Engine/Source/Programs/Horde/Horde.Agent.Tests/JsonRpcLoggerTests.cs b/Engine/Source/Programs/Horde/Horde.Agent.Tests/JsonRpcLoggerTests.cs
--- a/Engine/Source/Programs/Horde/Horde.Agent.Tests/JsonRpcLoggerTests.cs
+++ b/Engine/Source/Programs/Horde/Horde.Agent.Tests/JsonRpcLoggerTests.cs
@@ -114,7 +114,7 @@ namespace Horde.Agent.Tests
 			{
 			}
 
-			protected override Task UpdateLogAsync(NodeHandle target, int lineCount, CancellationToken cancellationToken)
+			protected override Task UpdateLogAsync(NodeHandle target, int lineCount, bool complete, CancellationToken cancellationToken)
 			{
 				Target = target;
 				return Task.CompletedTask;
