Horde: Skip adding logs to the trim queue if there is no tail data for them.

diff --git a/Engine/Source/Programs/Horde/Horde.Build/Logs/LogTailService.cs b/Engine/Source/Programs/Horde/Horde.Build/Logs/LogTailService.cs
--- a/Engine/Source/Programs/Horde/Horde.Build/Logs/LogTailService.cs
+++ b/Engine/Source/Programs/Horde/Horde.Build/Logs/LogTailService.cs
@@ -278,7 +278,11 @@ namespace Horde.Build.Logs
 		{
 			string trimEntry = $"{logId}={lineCount}";
 			double trimTime = (_clock.UtcNow + TrimAfter - DateTime.UnixEpoch).TotalSeconds;
-			_ = await _redisService.GetDatabase().SortedSetAddAsync(_trimQueue, trimEntry, trimTime, flags: CommandFlags.FireAndForget);
+
+			ITransaction transaction = _redisService.GetDatabase().CreateTransaction();
+			transaction.AddCondition(Condition.KeyExists(TailDataKey(logId).Inner));
+			_ = transaction.SortedSetAddAsync(_trimQueue, trimEntry, trimTime, flags: CommandFlags.FireAndForget);
+			await transaction.ExecuteAsync(CommandFlags.FireAndForget);
 		}
 
 		/// <summary>
