Horde: Additional logging for tracing when log tailing is enabled.

diff --git a/Engine/Source/Programs/Horde/Horde.Build/Logs/LogRpcService.cs b/Engine/Source/Programs/Horde/Horde.Build/Logs/LogRpcService.cs
--- a/Engine/Source/Programs/Horde/Horde.Build/Logs/LogRpcService.cs
+++ b/Engine/Source/Programs/Horde/Horde.Build/Logs/LogRpcService.cs
@@ -79,14 +79,18 @@ namespace Horde.Build.Logs
 
 				moveNextTask = requestStream.MoveNext();
 
-				using (CancellationTokenSource cancellationSource = new CancellationTokenSource())
+				response.TailNext = await _logTailService.GetTailNextAsync(logId, context.CancellationToken);
+				if (response.TailNext == -1)
 				{
-					Task<int> waitTask = _logTailService.WaitForTailNextAsync(logId, cancellationSource.Token);
+					using (CancellationTokenSource cancellationSource = new CancellationTokenSource())
+					{
+						Task<int> waitTask = _logTailService.WaitForTailNextAsync(logId, cancellationSource.Token);
 
-					await Task.WhenAny(waitTask, moveNextTask);
-					cancellationSource.Cancel();
+						await Task.WhenAny(waitTask, moveNextTask);
+						cancellationSource.Cancel();
 
-					response.TailNext = await waitTask;
+						response.TailNext = await waitTask;
+					}
 				}
 
 				await responseStream.WriteAsync(response);
diff --git a/Engine/Source/Programs/Horde/Horde.Build/Logs/LogTailService.cs b/Engine/Source/Programs/Horde/Horde.Build/Logs/LogTailService.cs
--- a/Engine/Source/Programs/Horde/Horde.Build/Logs/LogTailService.cs
+++ b/Engine/Source/Programs/Horde/Horde.Build/Logs/LogTailService.cs
@@ -304,6 +304,19 @@ namespace Horde.Build.Logs
 			return Math.Max(lineCount, flushedLineCount);
 		}
 
+		/// <summary>
+		/// Gets the next requested tail line number, returning immediately if it is set
+		/// </summary>
+		/// <param name="logId">The log file to query</param>
+		/// <param name="cancellationToken">Cancellation token for the operation</param>
+		/// <returns>Next requested line number, or -1 if tailing is not enabled.</returns>
+		public async Task<int> GetTailNextAsync(LogId logId, CancellationToken cancellationToken)
+		{
+			_ = cancellationToken;
+			RedisStringKey<int> tailNextKey = TailNextKey(logId);
+			return await _redisService.GetDatabase().StringGetAsync(tailNextKey, -1);
+		}
+
 		/// <summary>
 		/// Waits for a request to read tail data
 		/// </summary>
@@ -324,6 +337,7 @@ namespace Horde.Build.Logs
 						int tailNext = await _redisService.GetDatabase().StringGetAsync(TailNextKey(logId), -1);
 						if (tailNext != -1)
 						{
+							_logger.LogDebug("Released wait for tail on log {LogId} from line {Line}", logId, tailNext);
 							return tailNext;
 						}
 						else if (cancellationToken.IsCancellationRequested)
