Horde: Fix log tailing taking a long time to start; tail service not running on server runmodes, causing notifications not to reach pods communicating directly with agents.

diff --git a/Engine/Source/Programs/Horde/Horde.Build/Logs/LogTailService.cs b/Engine/Source/Programs/Horde/Horde.Build/Logs/LogTailService.cs
--- a/Engine/Source/Programs/Horde/Horde.Build/Logs/LogTailService.cs
+++ b/Engine/Source/Programs/Horde/Horde.Build/Logs/LogTailService.cs
@@ -14,6 +14,7 @@ using Microsoft.Extensions.Hosting;
 using HordeCommon;
 using Microsoft.Extensions.Logging;
 using System.Globalization;
+using Microsoft.Extensions.Options;
 
 namespace Horde.Build.Logs
 {
@@ -57,6 +58,7 @@ namespace Horde.Build.Logs
 		readonly ITicker _expireTailsTicker;
 
 		readonly ILogger _logger;
+		readonly ServerSettings _settings;
 
 		IAsyncDisposable? _tailStartSubscription;
 
@@ -79,7 +81,7 @@ namespace Horde.Build.Logs
 		/// <summary>
 		/// Constructor
 		/// </summary>
-		public LogTailService(RedisService redisService, IClock clock, int chunkLineCount, ILogger<LogTailService> logger)
+		public LogTailService(RedisService redisService, IClock clock, int chunkLineCount, IOptions<ServerSettings> settings, ILogger<LogTailService> logger)
 		{
 			if ((chunkLineCount & (chunkLineCount - 1)) != 0)
 			{
@@ -93,6 +95,7 @@ namespace Horde.Build.Logs
 
 			_clock = clock;
 			_expireTailsTicker = clock.AddSharedTicker<LogTailService>(TimeSpan.FromSeconds(30.0), ExpireTailsAsync, logger);
+			_settings = settings.Value;
 			_logger = logger;
 		}
 
@@ -100,13 +103,19 @@ namespace Horde.Build.Logs
 		public async Task StartAsync(CancellationToken cancellationToken)
 		{
 			_tailStartSubscription = await _redisService.SubscribeAsync(_tailStartChannel, OnTailStart);
-			await _expireTailsTicker.StartAsync();
+			if (_settings.IsRunModeActive(RunMode.Worker))
+			{
+				await _expireTailsTicker.StartAsync();
+			}
 		}
 
 		/// <inheritdoc/>
 		public async Task StopAsync(CancellationToken cancellationToken)
 		{
-			await _expireTailsTicker.StopAsync();
+			if (_settings.IsRunModeActive(RunMode.Worker))
+			{
+				await _expireTailsTicker.StopAsync();
+			}
 			if (_tailStartSubscription != null)
 			{
 				await _tailStartSubscription.DisposeAsync();
@@ -161,6 +170,7 @@ namespace Horde.Build.Logs
 		{
 			if (_notifyLogEvents.TryGetValue(logId, out AsyncEvent? notifyEvent))
 			{
+				_logger.LogDebug("Latched notification for tail start on log {LogId}", logId);
 				notifyEvent.Latch();
 			}
 		}
diff --git a/Engine/Source/Programs/Horde/Horde.Build/Startup.cs b/Engine/Source/Programs/Horde/Horde.Build/Startup.cs
--- a/Engine/Source/Programs/Horde/Horde.Build/Startup.cs
+++ b/Engine/Source/Programs/Horde/Horde.Build/Startup.cs
@@ -419,6 +419,7 @@ namespace Horde.Build
 			services.AddSingleton<LifetimeService>();
 			services.AddSingleton<ILogFileService, LogFileService>();
 			services.AddSingleton<LogTailService>();
+			services.AddHostedService(provider => provider.GetRequiredService<LogTailService>()); // NB: Runs even on workers, to receive tail notifications.
 			services.AddSingleton<INotificationService, NotificationService>();
 
 			if (settings.Commits.ReplicateMetadata)
@@ -654,7 +655,6 @@ namespace Horde.Build
 				services.AddHostedService(provider => provider.GetRequiredService<IssueService>());
 				services.AddHostedService<IssueReportService>();
 				services.AddHostedService(provider => (LogFileService)provider.GetRequiredService<ILogFileService>());
-				services.AddHostedService(provider => provider.GetRequiredService<LogTailService>());
 				services.AddHostedService(provider => (NotificationService)provider.GetRequiredService<INotificationService>());
 				services.AddHostedService(provider => provider.GetRequiredService<ReplicationService>());
 				if (!settings.DisableSchedules)
