Horde: Build fix.

diff --git a/Engine/Source/Programs/Horde/Horde.Build.Tests/LogTailServiceTests.cs b/Engine/Source/Programs/Horde/Horde.Build.Tests/LogTailServiceTests.cs
--- a/Engine/Source/Programs/Horde/Horde.Build.Tests/LogTailServiceTests.cs
+++ b/Engine/Source/Programs/Horde/Horde.Build.Tests/LogTailServiceTests.cs
@@ -13,6 +13,7 @@ using Horde.Build.Utilities;
 using HordeCommon;
 using Microsoft.Extensions.DependencyInjection;
 using Microsoft.Extensions.Logging;
+using Microsoft.Extensions.Options;
 using Microsoft.VisualStudio.TestTools.UnitTesting;
 
 namespace Horde.Build.Tests
@@ -73,7 +74,7 @@ namespace Horde.Build.Tests
 		{
 			RedisService redisService = ServiceProvider.GetRequiredService<RedisService>();
 
-			LogTailService tailService = new LogTailService(redisService, ServiceProvider.GetRequiredService<IClock>(), 4, ServiceProvider.GetRequiredService<ILogger<LogTailService>>());
+			LogTailService tailService = new LogTailService(redisService, ServiceProvider.GetRequiredService<IClock>(), 4, ServiceProvider.GetRequiredService<IOptions<ServerSettings>>(), ServiceProvider.GetRequiredService<ILogger<LogTailService>>());
 			await tailService.StartAsync(CancellationToken.None);
 
 			await tailService.EnableTailingAsync(_logId, 30);
@@ -173,7 +174,7 @@ namespace Horde.Build.Tests
 			LogTailService tailService = ServiceProvider.GetRequiredService<LogTailService>();
 			await tailService.StartAsync(CancellationToken.None);
 
-			LogTailService tailService2 = new LogTailService(ServiceProvider.GetRequiredService<RedisService>(), ServiceProvider.GetRequiredService<IClock>(), 4, ServiceProvider.GetRequiredService<ILogger<LogTailService>>());
+			LogTailService tailService2 = new LogTailService(ServiceProvider.GetRequiredService<RedisService>(), ServiceProvider.GetRequiredService<IClock>(), 4, ServiceProvider.GetRequiredService<IOptions<ServerSettings>>(), ServiceProvider.GetRequiredService<ILogger<LogTailService>>());
 
 			Task<int> task = tailService.WaitForTailNextAsync(_logId, CancellationToken.None);
 			Assert.IsTrue(!task.IsCompleted);
diff --git a/Engine/Source/Programs/Horde/Horde.Build/Logs/LogTailService.cs b/Engine/Source/Programs/Horde/Horde.Build/Logs/LogTailService.cs
--- a/Engine/Source/Programs/Horde/Horde.Build/Logs/LogTailService.cs
+++ b/Engine/Source/Programs/Horde/Horde.Build/Logs/LogTailService.cs
@@ -73,8 +73,8 @@ namespace Horde.Build.Logs
 		/// <summary>
 		/// Constructor
 		/// </summary>
-		public LogTailService(RedisService redisService, IClock clock, ILogger<LogTailService> logger)
-			: this(redisService, clock, 32, logger)
+		public LogTailService(RedisService redisService, IClock clock, IOptions<ServerSettings> settings, ILogger<LogTailService> logger)
+			: this(redisService, clock, 32, settings, logger)
 		{
 		}
 
