Horde: Upload artifacts to new storage backend when UseNewTempStorage is set in job options.

diff --git a/Engine/Source/Programs/Horde/Horde.Agent/Execution/JobExecutor.cs b/Engine/Source/Programs/Horde/Horde.Agent/Execution/JobExecutor.cs
--- a/Engine/Source/Programs/Horde/Horde.Agent/Execution/JobExecutor.cs
+++ b/Engine/Source/Programs/Horde/Horde.Agent/Execution/JobExecutor.cs
@@ -503,7 +503,7 @@ namespace Horde.Agent.Execution
 				arguments.Append($" CopyUAT -WithLauncher -TargetDir=\"{buildDir}\"");
 			}
 
-			int result = await ExecuteAutomationToolAsync(step, workspaceDir, arguments.ToString(), null, logger, cancellationToken);
+			int result = await ExecuteAutomationToolAsync(step, workspaceDir, arguments.ToString(), logger, cancellationToken);
 			if (result != 0)
 			{
 				return false;
@@ -739,7 +739,7 @@ namespace Horde.Agent.Execution
 				{
 					arguments.AppendArgument("-SharedStorageDir=", sharedStorageDir.FullName);
 				}
-				return await ExecuteAutomationToolAsync(step, workspaceDir, arguments.ToString(), null, logger, cancellationToken) == 0;
+				return await ExecuteAutomationToolAsync(step, workspaceDir, arguments.ToString(), logger, cancellationToken) == 0;
 			}
 		}
 
@@ -827,7 +827,7 @@ namespace Horde.Agent.Execution
 			}
 
 			// Run UAT
-			if (await ExecuteAutomationToolAsync(step, workspaceDir, arguments, storage, logger, cancellationToken) != 0)
+			if (await ExecuteAutomationToolAsync(step, workspaceDir, arguments, logger, cancellationToken) != 0)
 			{
 				return false;
 			}
@@ -978,7 +978,7 @@ namespace Horde.Agent.Execution
 			return true;
 		}
 
-		protected async Task<int> ExecuteAutomationToolAsync(BeginStepResponse step, DirectoryReference workspaceDir, string arguments, IStorageClient? store, ILogger logger, CancellationToken cancellationToken)
+		protected async Task<int> ExecuteAutomationToolAsync(BeginStepResponse step, DirectoryReference workspaceDir, string arguments, ILogger logger, CancellationToken cancellationToken)
 		{
 			int result;
 			using (IScope scope = GlobalTracer.Instance.BuildSpan("BuildGraph").StartActive())
@@ -990,15 +990,15 @@ namespace Horde.Agent.Execution
 
 				if (RuntimeInformation.IsOSPlatform(OSPlatform.Windows))
 				{
-					result = await ExecuteCommandAsync(step, workspaceDir, Environment.GetEnvironmentVariable("COMSPEC") ?? "cmd.exe", $"/C \"\"{workspaceDir}\\Engine\\Build\\BatchFiles\\RunUAT.bat\" {arguments}\"", store, logger, cancellationToken);
+					result = await ExecuteCommandAsync(step, workspaceDir, Environment.GetEnvironmentVariable("COMSPEC") ?? "cmd.exe", $"/C \"\"{workspaceDir}\\Engine\\Build\\BatchFiles\\RunUAT.bat\" {arguments}\"", logger, cancellationToken);
 				}
 				else if (RuntimeInformation.IsOSPlatform(OSPlatform.Linux))
 				{
-					result = await ExecuteCommandAsync(step, workspaceDir, "/bin/bash", $"\"{workspaceDir}/Engine/Build/BatchFiles/RunUAT.sh\" {arguments}", store, logger, cancellationToken);
+					result = await ExecuteCommandAsync(step, workspaceDir, "/bin/bash", $"\"{workspaceDir}/Engine/Build/BatchFiles/RunUAT.sh\" {arguments}", logger, cancellationToken);
 				}
 				else if (RuntimeInformation.IsOSPlatform(OSPlatform.OSX))
 				{
-					result = await ExecuteCommandAsync(step, workspaceDir, "/bin/sh", $"\"{workspaceDir}/Engine/Build/BatchFiles/RunUAT.sh\" {arguments}", store, logger, cancellationToken);
+					result = await ExecuteCommandAsync(step, workspaceDir, "/bin/sh", $"\"{workspaceDir}/Engine/Build/BatchFiles/RunUAT.sh\" {arguments}", logger, cancellationToken);
 				}
 				else
 				{
@@ -1156,7 +1156,7 @@ namespace Horde.Agent.Execution
 		
 		string RefPrefix => $"{_job.StreamId}/{_job.Change}-{_jobId}";
 
-		async Task<int> ExecuteCommandAsync(BeginStepResponse step, DirectoryReference workspaceDir, string fileName, string arguments, IStorageClient? store, ILogger jobLogger, CancellationToken cancellationToken)
+		async Task<int> ExecuteCommandAsync(BeginStepResponse step, DirectoryReference workspaceDir, string fileName, string arguments, ILogger jobLogger, CancellationToken cancellationToken)
 		{
 			// Combine all the supplied environment variables together
 			Dictionary<string, string> newEnvVars = new Dictionary<string, string>(_envVars, StringComparer.Ordinal);
@@ -1374,20 +1374,24 @@ namespace Horde.Agent.Execution
 				jobLogger.LogInformation("Found {NumResults} test results", combinedTestData.Count);
 				await UploadTestDataAsync(step.StepId, combinedTestData);
 			}
-					
-			if (store != null)
+
+			if (_jobOptions.UseNewTempStorage ?? false)
 			{
-				RefName refName = new RefName(RefName.Sanitize($"{RefPrefix}/artifacts"));
+				using IRpcClientRef<JobRpc.JobRpcClient> jobRpc = await RpcConnection.GetClientRefAsync<JobRpc.JobRpcClient>(cancellationToken);
+				CreateJobArtifactResponse artifact = await jobRpc.Client.CreateArtifactAsync(new CreateJobArtifactRequest { JobId = _jobId, StepId = step.StepId, Name = "Agent Logs", Type = JobArtifactType.Saved }, cancellationToken: cancellationToken);
+				_logger.LogInformation("Created logs artifact {ArtifactId} with ref {RefName} in ns {Namespace}", artifact.Id, artifact.RefName, artifact.NamespaceId);
+
+				IStorageClient storage = _storageFactory.CreateStorageClient(_session, new NamespaceId(artifact.NamespaceId), artifact.Token);
 
 				TreeOptions treeOptions = new TreeOptions();
-				using TreeWriter treeWriter = new TreeWriter(store, treeOptions, refName.Text);
+				using TreeWriter treeWriter = new TreeWriter(storage, treeOptions, artifact.RefName);
 
 				DirectoryNode directoryNode = new DirectoryNode(DirectoryFlags.None);
 
 				ChunkingOptions options = new ChunkingOptions();
 				await directoryNode.CopyFromDirectoryAsync(logDir.ToDirectoryInfo(), options, treeWriter, cancellationToken);
 
-				await treeWriter.WriteAsync(refName, directoryNode, cancellationToken: cancellationToken);
+				await treeWriter.WriteAsync(new RefName(artifact.RefName), directoryNode, cancellationToken: cancellationToken);
 			}
 			else
 			{
