Horde: Output artifact upload messages to the job output log.

diff --git a/Engine/Source/Programs/Horde/Horde.Agent/Execution/JobExecutor.cs b/Engine/Source/Programs/Horde/Horde.Agent/Execution/JobExecutor.cs
--- a/Engine/Source/Programs/Horde/Horde.Agent/Execution/JobExecutor.cs
+++ b/Engine/Source/Programs/Horde/Horde.Agent/Execution/JobExecutor.cs
@@ -1378,8 +1378,8 @@ namespace Horde.Agent.Execution
 			if (_jobOptions.UseNewTempStorage ?? false)
 			{
 				using IRpcClientRef<JobRpc.JobRpcClient> jobRpc = await RpcConnection.GetClientRefAsync<JobRpc.JobRpcClient>(cancellationToken);
-				CreateJobArtifactResponse artifact = await jobRpc.Client.CreateArtifactAsync(new CreateJobArtifactRequest { JobId = _jobId, StepId = step.StepId, Name = "Agent Logs", Type = JobArtifactType.Saved }, cancellationToken: cancellationToken);
-				_logger.LogInformation("Created logs artifact {ArtifactId} with ref {RefName} in ns {Namespace}", artifact.Id, artifact.RefName, artifact.NamespaceId);
+				CreateJobArtifactResponse artifact = await jobRpc.Client.CreateArtifactAsync(new CreateJobArtifactRequest { JobId = _jobId, StepId = step.StepId, Name = "Agent State", Type = JobArtifactType.Saved }, cancellationToken: cancellationToken);
+				jobLogger.LogInformation("Created saved artifact {ArtifactId} with ref {RefName} in ns {Namespace}", artifact.Id, artifact.RefName, artifact.NamespaceId);
 
 				IStorageClient storage = _storageFactory.CreateStorageClient(_session, new NamespaceId(artifact.NamespaceId), artifact.Token);
 
