Horde: Add a delay when enabling log tailing until we have some data come through.

diff --git a/Engine/Source/Programs/Horde/Horde.Build/Logs/LogTailService.cs b/Engine/Source/Programs/Horde/Horde.Build/Logs/LogTailService.cs
--- a/Engine/Source/Programs/Horde/Horde.Build/Logs/LogTailService.cs
+++ b/Engine/Source/Programs/Horde/Horde.Build/Logs/LogTailService.cs
@@ -182,15 +182,34 @@ namespace Horde.Build.Logs
 		/// <param name="lineCount">The current flushed line count for the log file</param>
 		public async ValueTask EnableTailingAsync(LogId logId, int lineCount)
 		{
+			bool waitForTailData = false;
 			if (await _redisService.GetDatabase().StringSetAsync(TailNextKey(logId), lineCount, when: When.NotExists))
 			{
 				_logger.LogDebug("Enabled tailing for log {LogId}", logId);
+				waitForTailData = true;
 			}
 
 			await _redisService.GetDatabase().PublishAsync(_tailStartChannel, logId);
 
 			double score = (_clock.UtcNow + ExpireAfter - DateTime.UnixEpoch).TotalSeconds;
 			_ = _redisService.GetDatabase().SortedSetAddAsync(_expireQueue, logId, score, flags: CommandFlags.FireAndForget);
+
+			if (waitForTailData)
+			{
+				// Bit hacky; poll for a while until we get some tail data through
+				RedisSortedSetKey<ReadOnlyMemory<byte>> tailDataKey = TailDataKey(logId);
+				for (int idx = 0; idx < 20; idx++)
+				{
+					if (await _redisService.GetDatabase().SortedSetLengthAsync(tailDataKey) > 0)
+					{
+						break;
+					}
+					else
+					{
+						await Task.Delay(100);
+					}
+				}
+			}
 		}
 
 		/// <summary>
