Horde: Upload .ib_mon files from XGE as artifacts

diff --git a/Engine/Source/Programs/Horde/Horde.Agent/Execution/JobExecutor.cs b/Engine/Source/Programs/Horde/Horde.Agent/Execution/JobExecutor.cs
--- a/Engine/Source/Programs/Horde/Horde.Agent/Execution/JobExecutor.cs
+++ b/Engine/Source/Programs/Horde/Horde.Agent/Execution/JobExecutor.cs
@@ -292,16 +292,12 @@ namespace Horde.Agent.Execution
 				_additionalArguments.Add($"-set:PreflightChange={_job.PreflightChange}");
 			}
 
-			// Only run XGE metadata extraction on Windows
 			if (RuntimeInformation.IsOSPlatform(OSPlatform.Windows))
 			{
 				DirectoryReference? xgeDir = XgeMetadataExtractor.FindXgeDir();
 				if (xgeDir != null)
 				{
 					_xgeMetadataExtractor = new XgeMetadataExtractor(xgeDir);
-					
-					// Clear any *.ib_mon files from previous builds
-					_xgeMetadataExtractor.ClearLocalIbMonFiles();
 				}
 				else
 				{
@@ -359,29 +355,6 @@ namespace Horde.Agent.Execution
 
 		public virtual Task FinalizeAsync(ILogger jobLogger, CancellationToken cancellationToken)
 		{
-			if (_xgeMetadataExtractor != null)
-			{
-				// Collect any new *.ib_mon files and generate a summary for each build invocation, grouped by title.
-				// Title will segregate for example UnrealBuildTool from invocations coming from the engine/cooks.
-				List<XgeTaskMetadataSummary> summaries = _xgeMetadataExtractor.GetSummariesForIbMonFiles();
-				if (summaries.Count > 0)
-				{
-					foreach (XgeTaskMetadataSummary s in summaries)
-					{
-						_logger.LogInformation("XGE build invocation {Title} with {LocalTaskCount} local tasks for {LocalTaskDuration} secs total and {RemoteTaskCount} remote tasks for {RemoteTaskDuration} secs total",
-							s.Title, s.LocalTaskCount, s.TotalLocalTaskDuration.TotalSeconds, s.RemoteTaskCount, s.TotalRemoteTaskDuration.TotalSeconds);
-						
-						jobLogger.LogInformation("XGE build invocation {Title} with {LocalTaskCount} local tasks for {LocalTaskDuration} secs total and {RemoteTaskCount} remote tasks for {RemoteTaskDuration} secs total",
-							s.Title, s.LocalTaskCount, s.TotalLocalTaskDuration.TotalSeconds, s.RemoteTaskCount, s.TotalRemoteTaskDuration.TotalSeconds);
-					}
-				}
-				else
-				{
-					_logger.LogInformation("No XGE build invocations found (no *.ib_mon files in History/Local)");
-					jobLogger.LogInformation("No XGE build invocations found (no *.ib_mon files in History/Local)");
-				}
-			}
-			
 			return Task.CompletedTask;
 		}
 
@@ -688,6 +661,27 @@ namespace Horde.Agent.Execution
 
 			return true;
 		}
+		
+		private async Task UploadXgeMonitorFilesAsync(BeginStepResponse step, ILogger logger, CancellationToken cancellationToken)
+		{
+			if (_xgeMetadataExtractor != null)
+			{
+				using IScope scope = GlobalTracer.Instance.BuildSpan("XgeMonitorFilesUpload").StartActive();
+				List<FileReference> ibMonFiles = _xgeMetadataExtractor.GetLocalIbMonFilePaths();
+
+				scope.Span.SetTag("NumIbMonFiles", ibMonFiles.Count);
+				ParallelOptions options = new () { MaxDegreeOfParallelism = 5, CancellationToken = cancellationToken };
+				await Parallel.ForEachAsync(ibMonFiles, options, async (file, innerCt) =>
+				{
+					string artifactName = "Xge/" + file.GetFileName()
+						.Replace("{", "", StringComparison.Ordinal)
+						.Replace("}", "", StringComparison.Ordinal);
+					
+					await ArtifactUploader.UploadAsync(RpcConnection, _jobId, _batchId, step.StepId,
+						artifactName, file, logger, innerCt);
+				});
+			}
+		}
 
 		private static void GetRecursiveDependencies(string name, Dictionary<string, ExportedNode> nameToNode, HashSet<string> dependencies)
 		{
@@ -729,9 +723,12 @@ namespace Horde.Agent.Execution
 				}
 			}
 
+			_xgeMetadataExtractor?.ClearLocalIbMonFiles();
 			if (_jobOptions.UseNewTempStorage ?? false)
 			{
-				return await ExecuteWithTempStorageAsync(step, workspaceDir, arguments.ToString(), logger, cancellationToken);
+				bool result = await ExecuteWithTempStorageAsync(step, workspaceDir, arguments.ToString(), logger, cancellationToken);
+				await UploadXgeMonitorFilesAsync(step, logger, cancellationToken);
+				return result;
 			}
 			else
 			{
@@ -739,7 +736,10 @@ namespace Horde.Agent.Execution
 				{
 					arguments.AppendArgument("-SharedStorageDir=", sharedStorageDir.FullName);
 				}
-				return await ExecuteAutomationToolAsync(step, workspaceDir, arguments.ToString(), logger, cancellationToken) == 0;
+				
+				bool result = await ExecuteAutomationToolAsync(step, workspaceDir, arguments.ToString(), logger, cancellationToken) == 0;
+				await UploadXgeMonitorFilesAsync(step, logger, cancellationToken);
+				return result;
 			}
 		}
 
@@ -1119,40 +1119,6 @@ namespace Horde.Agent.Execution
 				}
 			}
 		}
-
-		private class InterceptingLogger : ILogger
-		{
-			private readonly ILogger _logger;
-			private readonly int[] _includeEventIds;
-			private readonly Action<LogLevel, EventId, object?, Exception?> _callback = null!;
-
-			public InterceptingLogger(ILogger logger, int[] includeEventIds, Action<LogLevel, EventId, object?, Exception?> callback)
-			{
-				_logger = logger;
-				_includeEventIds = includeEventIds;
-				_callback = callback;
-			}
-
-			public void Log<TState>(LogLevel logLevel, EventId eventId, TState state, Exception? exception, Func<TState, Exception?, string> formatter)
-			{
-				if (_includeEventIds.Any(id => id == eventId.Id))
-				{
-					_callback(logLevel, eventId, state, exception);	
-				}
-
-				_logger.Log(logLevel, eventId, state, exception, formatter);
-			}
-
-			public bool IsEnabled(LogLevel logLevel)
-			{
-				return _logger.IsEnabled(logLevel);
-			}
-
-			public IDisposable BeginScope<TState>(TState state)
-			{
-				return _logger.BeginScope(state);
-			}
-		}
 		
 		string RefPrefix => $"{_job.StreamId}/{_job.Change}-{_jobId}";
 
