Horde: Only raise class specific exception in XgeMetadataExtractor

Avoid surfacing internal exceptions such as FileNotFound or JsonException.

diff --git a/Engine/Source/Programs/Horde/Horde.Agent/Utility/XgeMetadataExtractor.cs b/Engine/Source/Programs/Horde/Horde.Agent/Utility/XgeMetadataExtractor.cs
--- a/Engine/Source/Programs/Horde/Horde.Agent/Utility/XgeMetadataExtractor.cs
+++ b/Engine/Source/Programs/Horde/Horde.Agent/Utility/XgeMetadataExtractor.cs
@@ -82,6 +82,18 @@ public class XgeTaskMetadataSummary
 	}
 }
 
+/// <summary>
+/// Exception raised by XGE metadata extraction
+/// </summary>
+public class XgeMetadataExtractorException : Exception
+{
+	/// <inheritdoc/>
+	public XgeMetadataExtractorException(string? message) : base(message) { }
+
+	/// <inheritdoc/>
+	public XgeMetadataExtractorException(string? message, Exception? innerException) : base(message, innerException) { }
+}
+
 /// <summary>
 /// Extracts metadata about executed XGE / IncrediBuild tasks by locating local *.ib_mon files
 /// </summary>
@@ -102,18 +114,18 @@ public class XgeMetadataExtractor
 	{
 		if (!OperatingSystem.IsWindows())
 		{
-			throw new Exception("Only Windows is supported for XGE metadata extraction");
+			throw new XgeMetadataExtractorException("Only Windows is supported for XGE metadata extraction");
 		}
 
 		if (!IsXgeDirectory(xgeDir.FullName))
 		{
-			throw new ArgumentException("Path is not XGE base directory: " + xgeDir.FullName);
+			throw new XgeMetadataExtractorException("Path is not XGE base directory: " + xgeDir.FullName);
 		}
 		
 		string buildMonitorExePath = Path.Join(xgeDir.FullName, "BuildMonitor.exe");
 		if (!File.Exists(buildMonitorExePath))
 		{
-			throw new FileNotFoundException("BuildMonitor.exe not found in XGE base dir " + xgeDir.FullName);
+			throw new XgeMetadataExtractorException("BuildMonitor.exe not found in XGE base dir " + xgeDir.FullName);
 		}
 		
 		_xgeDir = xgeDir;
@@ -133,7 +145,7 @@ public class XgeMetadataExtractor
 
 		if (GetLocalIbMonFilePaths().Count != 0)
 		{
-			throw new Exception("Unable to remove all *.ib_mon files from history");
+			throw new XgeMetadataExtractorException("Unable to remove all *.ib_mon files from history");
 		}
 	}
 	
@@ -162,13 +174,12 @@ public class XgeMetadataExtractor
 	/// </summary>
 	/// <param name="ibMonFile">Path to .ib_mon file</param>
 	/// <returns>A summary of task count and durations</returns>
-	/// <exception cref="FileNotFoundException">If input wasn't found</exception>
-	/// <exception cref="Exception">If JSON file conversion wasn't successful</exception>
-	public XgeTaskMetadataSummary? GetSummaryFromIbMonFile(FileReference ibMonFile)
+	/// <exception cref="XgeMetadataExtractorException">If file conversion wasn't successful</exception>
+	public XgeTaskMetadataSummary GetSummaryFromIbMonFile(FileReference ibMonFile)
 	{
 		if (!FileReference.Exists(ibMonFile))
 		{
-			throw new FileNotFoundException("Path to ib_mon file not found " + ibMonFile.FullName);
+			throw new XgeMetadataExtractorException("Path to ib_mon file not found " + ibMonFile.FullName);
 		}
 		
 		string tempFile = Path.GetTempFileName();
@@ -181,13 +192,13 @@ public class XgeMetadataExtractor
 
 		if (process.ExitCode != 0)
 		{
-			throw new Exception("XGE build monitor executable finished with non-zero exit code: " + process.ExitCode);
+			throw new XgeMetadataExtractorException("XGE build monitor executable finished with non-zero exit code: " + process.ExitCode);
 		}
 
 		FileInfo tempFileInfo = new (tempFile);
 		if (!tempFileInfo.Exists || tempFileInfo.Length == 0)
 		{
-			throw new Exception("Generated JSON file does not exist or is empty");
+			throw new XgeMetadataExtractorException("Generated JSON file does not exist or is empty");
 		}
 
 		byte[] jsonData = File.ReadAllBytes(tempFile);
@@ -205,19 +216,28 @@ public class XgeMetadataExtractor
 	/// </summary>
 	/// <param name="data">Path to .ib_mon file</param>
 	/// <returns>A C# object representation</returns>
-	/// <exception cref="Exception">If JSON file conversion wasn't successful</exception>
-	public static XgeTaskMetadataSummary? GetSummaryFromJson(byte[] data)
+	/// <exception cref="XgeMetadataExtractorException">If JSON file conversion wasn't successful</exception>
+	public static XgeTaskMetadataSummary GetSummaryFromJson(byte[] data)
 	{
 		JsonSerializerOptions jsonOptions = new () { Converters = { new XgeMetadataBoolConverter() } };
-		BuildMonitorData? buildMonitorData = JsonSerializer.Deserialize<BuildMonitorData>(data, jsonOptions);
-		if (buildMonitorData == null)
+		BuildMonitorData? buildMonitorData;
+
+		try
+		{
+			buildMonitorData = JsonSerializer.Deserialize<BuildMonitorData>(data, jsonOptions);
+			if (buildMonitorData == null)
+			{
+				throw new JsonException("No result for deserialization");
+			}
+		}
+		catch (Exception e)
 		{
-			throw new Exception("Failed reading JSON file converted from .ib_mon file");
+			throw new XgeMetadataExtractorException("Failed reading JSON file converted from .ib_mon file", e);
 		}
 
 		if (buildMonitorData.Tasks.Count == 0)
 		{
-			return null;
+			throw new XgeMetadataExtractorException("Deserialized .ib_mon files does not contain any tasks");
 		}
 
 		int localTaskCount = 0;
@@ -255,7 +275,30 @@ public class XgeMetadataExtractor
 		Dictionary<string, XgeTaskMetadataSummary> result = new();
 		foreach (XgeTaskMetadataSummary? s in GetLocalIbMonFilePaths().Select(GetSummaryFromIbMonFile))
 		{
-			if (s == null || (s.LocalTaskCount == 0 && s.RemoteTaskCount == 0)) continue;
+			if (s.LocalTaskCount == 0 && s.RemoteTaskCount == 0) continue;
+			
+			if (!result.TryGetValue(s.Title, out XgeTaskMetadataSummary? aggregatedSummary))
+			{
+				aggregatedSummary = new XgeTaskMetadataSummary(s.Title, 0, 0, TimeSpan.Zero, TimeSpan.Zero);
+			}
+
+			result[s.Title] = aggregatedSummary.Add(s);
+		}
+
+		return result.Values.ToList();
+	}
+
+	/// <summary>
+	/// Aggregate all task counts and durations by XGE title
+	/// </summary>
+	/// <param name="summaries">A list of summaries</param>
+	/// <returns>A list of summaries grouped by title</returns>
+	public static List<XgeTaskMetadataSummary> GroupSummariesByTitle(List<XgeTaskMetadataSummary> summaries)
+	{
+		Dictionary<string, XgeTaskMetadataSummary> result = new();
+		foreach (XgeTaskMetadataSummary? s in summaries)
+		{
+			if (s.LocalTaskCount == 0 && s.RemoteTaskCount == 0) continue;
 			
 			if (!result.TryGetValue(s.Title, out XgeTaskMetadataSummary? aggregatedSummary))
 			{
