Horde: Add flag for using Wine in JobOptions

When set to true, jobs on Linux will change to a Wine-specific UAT entrypoint.

diff --git a/Engine/Source/Programs/Horde/Horde.Agent/Execution/JobExecutor.cs b/Engine/Source/Programs/Horde/Horde.Agent/Execution/JobExecutor.cs
--- a/Engine/Source/Programs/Horde/Horde.Agent/Execution/JobExecutor.cs
+++ b/Engine/Source/Programs/Horde/Horde.Agent/Execution/JobExecutor.cs
@@ -981,32 +981,38 @@ namespace Horde.Agent.Execution
 		protected async Task<int> ExecuteAutomationToolAsync(BeginStepResponse step, DirectoryReference workspaceDir, string arguments, ILogger logger, CancellationToken cancellationToken)
 		{
 			int result;
-			using (IScope scope = GlobalTracer.Instance.BuildSpan("BuildGraph").StartActive())
+			using IScope scope = GlobalTracer.Instance.BuildSpan("BuildGraph").StartActive();
+			
+			if (!_compileAutomationTool)
 			{
-				if (!_compileAutomationTool)
-				{
-					arguments += " -NoCompile";
-				}
+				arguments += " -NoCompile";
+			}
 
-				if (RuntimeInformation.IsOSPlatform(OSPlatform.Windows))
-				{
-					result = await ExecuteCommandAsync(step, workspaceDir, Environment.GetEnvironmentVariable("COMSPEC") ?? "cmd.exe", $"/C \"\"{workspaceDir}\\Engine\\Build\\BatchFiles\\RunUAT.bat\" {arguments}\"", logger, cancellationToken);
-				}
-				else if (RuntimeInformation.IsOSPlatform(OSPlatform.Linux))
-				{
-					result = await ExecuteCommandAsync(step, workspaceDir, "/bin/bash", $"\"{workspaceDir}/Engine/Build/BatchFiles/RunUAT.sh\" {arguments}", logger, cancellationToken);
-				}
-				else if (RuntimeInformation.IsOSPlatform(OSPlatform.OSX))
-				{
-					result = await ExecuteCommandAsync(step, workspaceDir, "/bin/sh", $"\"{workspaceDir}/Engine/Build/BatchFiles/RunUAT.sh\" {arguments}", logger, cancellationToken);
-				}
-				else
+			if (RuntimeInformation.IsOSPlatform(OSPlatform.Windows))
+			{
+				result = await ExecuteCommandAsync(step, workspaceDir, Environment.GetEnvironmentVariable("COMSPEC") ?? "cmd.exe", $"/C \"\"{workspaceDir}\\Engine\\Build\\BatchFiles\\RunUAT.bat\" {arguments}\"", logger, cancellationToken);
+			}
+			else if (RuntimeInformation.IsOSPlatform(OSPlatform.Linux))
+			{
+				string args = $"\"{workspaceDir}/Engine/Build/BatchFiles/RunUAT.sh\" {arguments}";
+				
+				if (_jobOptions.UseWine is true)
 				{
-					throw new Exception("Unsupported platform");
+					args = $"\"{workspaceDir}/Engine/Build/BatchFiles/RunWineUAT.sh\" {arguments}";
 				}
-
-				_compileAutomationTool = false;
+				
+				result = await ExecuteCommandAsync(step, workspaceDir, "/bin/bash", args, logger, cancellationToken);
 			}
+			else if (RuntimeInformation.IsOSPlatform(OSPlatform.OSX))
+			{
+				result = await ExecuteCommandAsync(step, workspaceDir, "/bin/sh", $"\"{workspaceDir}/Engine/Build/BatchFiles/RunUAT.sh\" {arguments}", logger, cancellationToken);
+			}
+			else
+			{
+				throw new Exception("Unsupported platform");
+			}
+
+			_compileAutomationTool = false;
 			return result;
 		}
 
