Horde: Allow storage backends to be reused between config refreshes.

diff --git a/Engine/Source/Programs/Horde/Horde.Build/Storage/StorageService.cs b/Engine/Source/Programs/Horde/Horde.Build/Storage/StorageService.cs
--- a/Engine/Source/Programs/Horde/Horde.Build/Storage/StorageService.cs
+++ b/Engine/Source/Programs/Horde/Horde.Build/Storage/StorageService.cs
@@ -12,6 +12,7 @@ using System.Reflection;
 using System.Runtime.CompilerServices;
 using System.Runtime.InteropServices;
 using System.Security.Claims;
+using System.Text.Json;
 using System.Threading;
 using System.Threading.Tasks;
 using EpicGames.Core;
@@ -92,7 +93,7 @@ namespace Horde.Build.Storage
 	/// </summary>
 	public sealed class StorageService : IHostedService, IDisposable, IStorageClientFactory
 	{
-		sealed class StorageClient : StorageClientBase, IStorageClientImpl, IDisposable
+		sealed class StorageClient : StorageClientBase, IStorageClientImpl
 		{
 			readonly StorageService _outer;
 			readonly string _prefix;
@@ -119,14 +120,6 @@ namespace Horde.Build.Storage
 				}
 			}
 
-			public void Dispose()
-			{
-				if (Backend is IDisposable disposable)
-				{
-					disposable.Dispose();
-				}
-			}
-
 			string GetBlobPath(BlobId blobId) => $"{_prefix}{blobId}.blob";
 
 			#region Blobs
@@ -266,18 +259,43 @@ namespace Horde.Build.Storage
 			#endregion
 		}
 
+
+		class BackendInfo
+		{
+			public IStorageBackend Backend { get; }
+
+			int _refCount;
+
+			public BackendInfo(IStorageBackend backend)
+			{
+				Backend = backend;
+			}
+
+			public void AddRef() => Interlocked.Increment(ref _refCount);
+
+			public void Release()
+			{
+				if (Interlocked.Decrement(ref _refCount) == 0 && Backend is IDisposable disposable)
+				{
+					disposable.Dispose();
+				}
+			}
+		}
+
 		class NamespaceInfo : IDisposable
 		{
 			public NamespaceConfig Config { get; }
 			public StorageClient Client { get; }
+			public BackendInfo Backend { get; }
 
-			public NamespaceInfo(NamespaceConfig config, StorageClient client)
+			public NamespaceInfo(NamespaceConfig config, StorageClient client, BackendInfo backend)
 			{
 				Config = config;
 				Client = client;
+				Backend = backend;
 			}
 
-			public void Dispose() => (Client as IDisposable)?.Dispose();
+			public void Dispose() => Backend.Release();
 		}
 
 		class State : IDisposable
@@ -452,6 +470,7 @@ namespace Horde.Build.Storage
 
 		State? _lastState;
 		string? _lastConfigRevision;
+		readonly Dictionary<IoHash, BackendInfo> _backends = new Dictionary<IoHash, BackendInfo>();
 
 		readonly AsyncCachedValue<State> _cachedState;
 
@@ -586,12 +605,9 @@ namespace Horde.Build.Storage
 				{
 					foreach (NamespaceConfig namespaceConfig in storageConfig.Namespaces)
 					{
-						IStorageBackend backend = CreateStorageBackend(namespaceConfig.BackendConfig);
-
-#pragma warning disable CA2000 // Dispose objects before losing scope
-						StorageClient client = new StorageClient(this, namespaceConfig, backend);
-						nextState.Namespaces.Add(namespaceConfig.Id, new NamespaceInfo(namespaceConfig, client));
-#pragma warning restore CA2000 // Dispose objects before losing scope
+						BackendInfo backend = FindOrAddStorageBackend(namespaceConfig.BackendConfig);
+						StorageClient client = new StorageClient(this, namespaceConfig, backend.Backend);
+						nextState.Namespaces.Add(namespaceConfig.Id, new NamespaceInfo(namespaceConfig, client, backend));
 					}
 				}
 				catch
@@ -608,6 +624,38 @@ namespace Horde.Build.Storage
 			return _lastState;
 		}
 
+		/// <summary>
+		/// Creates a storage backend with the given configuration
+		/// </summary>
+		/// <param name="config">Configuration for the backend</param>
+		/// <returns>New storage backend instance</returns>
+		BackendInfo FindOrAddStorageBackend(BackendConfig config)
+		{
+			// Compute the new hash of the configuration data
+			IoHash hash;
+			using (MemoryStream stream = new MemoryStream())
+			{
+				JsonSerializerOptions options = new JsonSerializerOptions();
+				Startup.ConfigureJsonSerializer(options);
+
+				JsonSerializer.Serialize(stream, config, options: options);
+
+				hash = IoHash.Compute(stream.ToArray());
+			}
+
+			// See if we've got an existing backend we can use
+			if (_backends.TryGetValue(hash, out BackendInfo? existingBackendInfo))
+			{
+				existingBackendInfo.AddRef();
+				return existingBackendInfo;
+			}
+
+			// Create a new backend
+			BackendInfo backendInfo = new BackendInfo(CreateStorageBackend(config));
+			_backends.Add(hash, backendInfo);
+			return backendInfo;
+		}
+
 		/// <summary>
 		/// Creates a storage backend with the given configuration
 		/// </summary>
