Horde: Upload temp storage output to a new artifact.

diff --git a/Engine/Source/Programs/Horde/Horde.Agent/Execution/JobExecutor.cs b/Engine/Source/Programs/Horde/Horde.Agent/Execution/JobExecutor.cs
--- a/Engine/Source/Programs/Horde/Horde.Agent/Execution/JobExecutor.cs
+++ b/Engine/Source/Programs/Horde/Horde.Agent/Execution/JobExecutor.cs
@@ -770,15 +770,15 @@ namespace Horde.Agent.Execution
 			await TestArtifactAsync(step.StepId, cancellationToken);
 
 			using MemoryCache cache = new MemoryCache(new MemoryCacheOptions { });
-			IStorageClient storage = _storageFactory.CreateStorageClient(_session, $"/api/v1/jobs/{_jobId}/batches/{_batchId}/steps/{step.StepId}/temp");
-			TreeReader reader = new TreeReader(storage, cache, _logger);
+//			IStorageClient storage = _storageFactory.CreateStorageClient(_session, $"/api/v1/jobs/{_jobId}/batches/{_batchId}/steps/{step.StepId}/temp");
+//			TreeReader reader = new TreeReader(storage, cache, _logger);
 
 			// Create the mapping of tag names to file sets
 			Dictionary<string, HashSet<FileReference>> tagNameToFileSet = new Dictionary<string, HashSet<FileReference>>();
 
 			// Read all the input tags for this node, and build a list of referenced input storage blocks
 			HashSet<TempStorageBlockRef> inputStorageBlocks = new HashSet<TempStorageBlockRef>();
-			foreach (string input in step.Inputs)
+/*			foreach (string input in step.Inputs)
 			{
 				int slashIdx = input.IndexOf('/', StringComparison.Ordinal);
 				if (slashIdx == -1)
@@ -794,10 +794,10 @@ namespace Horde.Agent.Execution
 				tagNameToFileSet[tagName] = fileList.ToFileSet(workspaceDir);
 				inputStorageBlocks.UnionWith(fileList.Blocks);
 			}
-
+*/
 			// Read the manifests for all the input storage blocks
 			Dictionary<TempStorageBlockRef, TempStorageBlockManifest> inputManifests = new Dictionary<TempStorageBlockRef, TempStorageBlockManifest>();
-			using (IScope scope = GlobalTracer.Instance.BuildSpan("TempStorage").WithTag("resource", "read").StartActive())
+/*			using (IScope scope = GlobalTracer.Instance.BuildSpan("TempStorage").WithTag("resource", "read").StartActive())
 			{
 				Stopwatch timer = Stopwatch.StartNew();
 				scope.Span.SetTag("blocks", inputStorageBlocks.Count);
@@ -809,10 +809,10 @@ namespace Horde.Agent.Execution
 				scope.Span.SetTag("size", inputManifests.Sum(x => x.Value.GetTotalSize()));
 				logger.LogInformation("Download took {Time:n1}s", timer.Elapsed.TotalSeconds);
 			}
-
+*/
 			// Read all the input storage blocks, keeping track of which block each file came from
 			Dictionary<FileReference, TempStorageBlockRef> fileToStorageBlock = new Dictionary<FileReference, TempStorageBlockRef>();
-			foreach (KeyValuePair<TempStorageBlockRef, TempStorageBlockManifest> pair in inputManifests)
+/*			foreach (KeyValuePair<TempStorageBlockRef, TempStorageBlockManifest> pair in inputManifests)
 			{
 				TempStorageBlockRef inputStorageBlock = pair.Key;
 				foreach (FileReference file in pair.Value.Files.Select(x => x.ToFileReference(workspaceDir)))
@@ -825,7 +825,7 @@ namespace Horde.Agent.Execution
 					fileToStorageBlock[file] = inputStorageBlock;
 				}
 			}
-
+*/
 			// Run UAT
 			if (await ExecuteAutomationToolAsync(step, workspaceDir, arguments, logger, cancellationToken) != 0)
 			{
@@ -833,6 +833,7 @@ namespace Horde.Agent.Execution
 			}
 
 			// Read all the output manifests
+			HashSet<FileReference> outputFiles = new HashSet<FileReference>();
 			foreach (string tagName in step.OutputNames)
 			{
 				FileReference tagFileListLocation = TempStorage.GetTagManifestLocation(manifestDir, step.Name, tagName);
@@ -840,6 +841,8 @@ namespace Horde.Agent.Execution
 
 				TempStorageTagManifest fileList = TempStorageTagManifest.Load(tagFileListLocation);
 				tagNameToFileSet[tagName] = fileList.ToFileSet(workspaceDir);
+
+				outputFiles.Add(tagFileListLocation);
 			}
 
 			// Check that none of the inputs have been clobbered
@@ -870,108 +873,31 @@ namespace Horde.Agent.Execution
 			}
 
 			// Determine all the output files which are required to be copied to temp storage (because they're referenced by nodes in another agent)
-			HashSet<FileReference> referencedOutputFiles = new HashSet<FileReference>();
 			foreach (int publishOutput in step.PublishOutputs)
 			{
 				string tagName = step.OutputNames[publishOutput];
-				referencedOutputFiles.UnionWith(tagNameToFileSet[tagName]);
-			}
-
-			// Find a block name for all new outputs
-			Dictionary<FileReference, string> fileToBlockName = new Dictionary<FileReference, string>();
-			for (int idx = 0; idx < step.OutputNames.Count; idx++)
-			{
-				string tagName = step.OutputNames[idx];
-
-				string outputNameWithoutHash = tagName.TrimStart('#');
-				bool isDefaultOutput = outputNameWithoutHash.Equals(step.Name, StringComparison.OrdinalIgnoreCase);
-
-				HashSet<FileReference> files = tagNameToFileSet[tagName];
-				foreach (FileReference file in files)
-				{
-					if (!fileToStorageBlock.ContainsKey(file) && file.IsUnderDirectory(workspaceDir))
-					{
-						if (isDefaultOutput)
-						{
-							if (!fileToBlockName.ContainsKey(file))
-							{
-								fileToBlockName[file] = "";
-							}
-						}
-						else
-						{
-							string? blockName;
-							if (fileToBlockName.TryGetValue(file, out blockName) && blockName.Length > 0)
-							{
-								fileToBlockName[file] = $"{blockName}+{outputNameWithoutHash}";
-							}
-							else
-							{
-								fileToBlockName[file] = outputNameWithoutHash;
-							}
-						}
-					}
-				}
-			}
-
-			// Invert the dictionary to make a mapping of storage block to the files each contains
-			Dictionary<string, HashSet<FileReference>> outputStorageBlockToFiles = new Dictionary<string, HashSet<FileReference>>();
-			foreach (KeyValuePair<FileReference, string> pair in fileToBlockName)
-			{
-				HashSet<FileReference>? files;
-				if (!outputStorageBlockToFiles.TryGetValue(pair.Value, out files))
-				{
-					files = new HashSet<FileReference>();
-					outputStorageBlockToFiles.Add(pair.Value, files);
-				}
-				files.Add(pair.Key);
+				outputFiles.UnionWith(tagNameToFileSet[tagName]);
 			}
 
 			// Write all the storage blocks, and update the mapping from file to storage block
 			using (GlobalTracer.Instance.BuildSpan("TempStorage").WithTag("resource", "Write").StartActive())
 			{
 				Stopwatch timer = Stopwatch.StartNew();
-				RefName refName = TempStorage.GetRefNameForNode(RefPrefix, step.Name);
-
-				TreeOptions treeOptions = new TreeOptions();
-				using TreeWriter treeWriter = new TreeWriter(storage, treeOptions, refName.Text);
-
-				TempStorageNode outputNode = new TempStorageNode();
 
-				// Create all the output blocks
-				foreach (KeyValuePair<string, HashSet<FileReference>> pair in outputStorageBlockToFiles)
-				{
-					TempStorageBlockRef outputBlock = new TempStorageBlockRef(step.Name, pair.Key);
-					foreach (FileReference file in pair.Value)
-					{
-						fileToStorageBlock.Add(file, outputBlock);
-					}
-					if (pair.Value.Any(x => referencedOutputFiles.Contains(x)))
-					{
-						outputNode.Blocks[pair.Key] = await TempStorage.ArchiveBlockAsync(treeWriter, pair.Key, workspaceDir, pair.Value.ToArray(), logger, cancellationToken);
-					}
-				}
+				using IRpcClientRef<JobRpc.JobRpcClient> jobRpc = await RpcConnection.GetClientRefAsync<JobRpc.JobRpcClient>(cancellationToken);
+				CreateJobArtifactResponse artifact = await jobRpc.Client.CreateArtifactAsync(new CreateJobArtifactRequest { JobId = _jobId, StepId = step.StepId, Name = "Output", Type = JobArtifactType.Output }, cancellationToken: cancellationToken);
+				_logger.LogInformation("Created artifact {ArtifactId} with ref {RefName} in ns {Namespace}", artifact.Id, artifact.RefName, artifact.NamespaceId);
 
-				// Create all the output tags
-				foreach (string outputName in step.OutputNames)
-				{
-					HashSet<FileReference> files = tagNameToFileSet[outputName];
+				IStorageClient storage = _storageFactory.CreateStorageClient(_session, new NamespaceId(artifact.NamespaceId), artifact.Token);
+				using TreeWriter writer = new TreeWriter(storage, new RefName(artifact.RefName));
 
-					HashSet<TempStorageBlockRef> storageBlocks = new HashSet<TempStorageBlockRef>();
-					foreach (FileReference file in files)
-					{
-						TempStorageBlockRef? storageBlock;
-						if (fileToStorageBlock.TryGetValue(file, out storageBlock))
-						{
-							storageBlocks.Add(storageBlock);
-						}
-					}
+				ChunkingOptions chunkingOptions = new ChunkingOptions();
 
-					outputNode.Tags[outputName] = await TempStorage.ArchiveTagAsync(treeWriter, outputName, workspaceDir, files, storageBlocks.ToArray(), logger, cancellationToken);
-				}
+				DirectoryNode root = new DirectoryNode();
+				await root.CopyFilesAsync(workspaceDir, outputFiles, chunkingOptions, writer, cancellationToken);
 
 				// Write the final node
-				await treeWriter.WriteAsync(refName, outputNode, new RefOptions { Lifetime = TimeSpan.FromDays(3.0) }, cancellationToken);
+				await writer.WriteAsync(artifact.RefName, root, new RefOptions { Lifetime = TimeSpan.FromDays(3.0) }, cancellationToken);
 				logger.LogInformation("Upload took {Time:n1}s", timer.Elapsed.TotalSeconds);
 			}
 
diff --git a/Engine/Source/Programs/Shared/EpicGames.Horde/Storage/Nodes/DirectoryNode.cs b/Engine/Source/Programs/Shared/EpicGames.Horde/Storage/Nodes/DirectoryNode.cs
--- a/Engine/Source/Programs/Shared/EpicGames.Horde/Storage/Nodes/DirectoryNode.cs
+++ b/Engine/Source/Programs/Shared/EpicGames.Horde/Storage/Nodes/DirectoryNode.cs
@@ -643,6 +643,52 @@ namespace EpicGames.Horde.Storage.Nodes
 
 		#endregion
 
+		/// <summary>
+		/// Adds files from a flat list of paths
+		/// </summary>
+		/// <param name="baseDir">Base directory to base paths relative to</param>
+		/// <param name="files">Files to add</param>
+		/// <param name="options">Options for chunking file content</param>
+		/// <param name="writer">Writer for new node data</param>
+		/// <param name="cancellationToken">Cancellation token for the operation</param>
+		public async Task CopyFilesAsync(DirectoryReference baseDir, IEnumerable<FileReference> files, ChunkingOptions options, TreeWriter writer, CancellationToken cancellationToken)
+		{
+			Dictionary<DirectoryReference, DirectoryNode> dirToNode = new Dictionary<DirectoryReference, DirectoryNode>();
+			dirToNode.Add(baseDir, this);
+
+			List<(DirectoryNode, FileInfo)> groupedFiles = new List<(DirectoryNode, FileInfo)>();
+			foreach (FileReference file in files)
+			{
+				DirectoryNode? node = FindOrAddDirectory(baseDir, file.Directory, dirToNode);
+				if (node == null)
+				{
+					throw new InvalidOperationException($"File {file} is not under base directory {baseDir}");
+				}
+				groupedFiles.Add((node, file.ToFileInfo()));
+			}
+
+			await CopyFromDirectoryAsync(groupedFiles, options, writer, cancellationToken);
+		}
+
+		DirectoryNode? FindOrAddDirectory(DirectoryReference baseDir, DirectoryReference dir, Dictionary<DirectoryReference, DirectoryNode> dirToNode)
+		{
+			DirectoryNode? node;
+			if (!dirToNode.TryGetValue(dir, out node))
+			{
+				DirectoryReference? parentDir = dir.ParentDirectory;
+				if (parentDir != null)
+				{
+					DirectoryNode? parentNode = FindOrAddDirectory(baseDir, parentDir, dirToNode);
+					if (parentNode != null)
+					{
+						node = parentNode.AddDirectory(dir.GetDirectoryName());
+						dirToNode.Add(dir, node);
+					}
+				}
+			}
+			return node;
+		}
+
 		/// <summary>
 		/// Adds files from a directory on disk
 		/// </summary>
