Horde: Give each agent a token for communication with the server regarding that specific lease.

diff --git a/Engine/Source/Programs/Horde/Horde.Build/Jobs/JobTaskSource.cs b/Engine/Source/Programs/Horde/Horde.Build/Jobs/JobTaskSource.cs
--- a/Engine/Source/Programs/Horde/Horde.Build/Jobs/JobTaskSource.cs
+++ b/Engine/Source/Programs/Horde/Horde.Build/Jobs/JobTaskSource.cs
@@ -13,6 +13,7 @@ using EpicGames.Core;
 using EpicGames.Horde.Common;
 using Google.Protobuf;
 using Google.Protobuf.WellKnownTypes;
+using Horde.Build.Acls;
 using Horde.Build.Agents;
 using Horde.Build.Agents.Leases;
 using Horde.Build.Agents.Pools;
@@ -20,6 +21,7 @@ using Horde.Build.Jobs.Graphs;
 using Horde.Build.Logs;
 using Horde.Build.Perforce;
 using Horde.Build.Server;
+using Horde.Build.Storage;
 using Horde.Build.Streams;
 using Horde.Build.Tasks;
 using Horde.Build.Ugs;
@@ -173,7 +175,7 @@ namespace Horde.Build.Jobs
 			}
 		}
 
-		readonly GlobalsService _globalsService;
+		readonly AclService _aclService;
 		readonly IStreamCollection _streamCollection;
 		readonly ILogFileService _logFileService;
 		readonly IAgentCollection _agentsCollection;
@@ -220,9 +222,9 @@ namespace Horde.Build.Jobs
 		/// <summary>
 		/// Constructor
 		/// </summary>
-		public JobTaskSource(GlobalsService globalsService, IAgentCollection agents, IJobCollection jobs, IJobStepRefCollection jobStepRefs, IGraphCollection graphs, IPoolCollection pools, IUgsMetadataCollection ugsMetadataCollection, IStreamCollection streamCollection, ILogFileService logFileService, PerforceLoadBalancer perforceLoadBalancer, IClock clock, IOptionsMonitor<ServerSettings> settings, IOptionsMonitor<GlobalConfig> globalConfig, ILogger<JobTaskSource> logger)
+		public JobTaskSource(AclService aclService, IAgentCollection agents, IJobCollection jobs, IJobStepRefCollection jobStepRefs, IGraphCollection graphs, IPoolCollection pools, IUgsMetadataCollection ugsMetadataCollection, IStreamCollection streamCollection, ILogFileService logFileService, PerforceLoadBalancer perforceLoadBalancer, IClock clock, IOptionsMonitor<ServerSettings> settings, IOptionsMonitor<GlobalConfig> globalConfig, ILogger<JobTaskSource> logger)
 		{
-			_globalsService = globalsService;
+			_aclService = aclService;
 			_agentsCollection = agents;
 			_jobs = jobs;
 			_jobStepRefs = jobStepRefs;
@@ -760,11 +762,8 @@ namespace Horde.Build.Jobs
 				}
 				leaseName.Append(CultureInfo.InvariantCulture, $" - {job.Name}");
 
-				// Get the global settings
-				IGlobals globals = await _globalsService.GetAsync();
-
 				// Encode the payload
-				ExecuteJobTask? task = await CreateExecuteJobTaskAsync(streamConfig, job, batch, agent, item._workspace, item._useAutoSdk, logId);
+				ExecuteJobTask? task = await CreateExecuteJobTaskAsync(leaseId, streamConfig, job, batch, agent, item._workspace, item._useAutoSdk, logId);
 				if (task != null)
 				{
 					byte[] payload = Any.Pack(task).ToByteArray();
@@ -814,7 +813,7 @@ namespace Horde.Build.Jobs
 			return null;
 		}
 
-		async Task<ExecuteJobTask?> CreateExecuteJobTaskAsync(StreamConfig streamConfig, IJob job, IJobStepBatch batch, IAgent agent, AgentWorkspace workspace, bool useAutoSdk, LogId logId)
+		async Task<ExecuteJobTask?> CreateExecuteJobTaskAsync(LeaseId leaseId, StreamConfig streamConfig, IJob job, IJobStepBatch batch, IAgent agent, AgentWorkspace workspace, bool useAutoSdk, LogId logId)
 		{
 			// Get the lease name
 			StringBuilder leaseName = new StringBuilder($"{streamConfig.Name} - ");
@@ -831,6 +830,11 @@ namespace Horde.Build.Jobs
 			// Get the global settings
 			GlobalConfig globalConfig = _globalConfig.CurrentValue;
 
+			// Create a bearer token for the job executor
+			List<AclClaimConfig> claims = new List<AclClaimConfig>();
+			claims.Add(new AclClaimConfig(HordeClaimTypes.Lease, leaseId.ToString()));
+			claims.Add(new AclClaimConfig(HordeClaimTypes.WriteNamespace, Namespace.Artifacts.ToString()));
+
 			// Encode the payload
 			ExecuteJobTask task = new ExecuteJobTask();
 			task.JobId = job.Id.ToString();
@@ -838,6 +842,7 @@ namespace Horde.Build.Jobs
 			task.LogId = logId.ToString();
 			task.JobName = leaseName.ToString();
 			task.JobOptions = job.JobOptions;
+			task.Token = await _aclService.IssueBearerTokenAsync(claims, null);
 
 			List<HordeCommon.Rpc.Messages.AgentWorkspace> workspaces = new ();
 
diff --git a/Engine/Source/Programs/Horde/Horde.Build/Storage/StorageController.cs b/Engine/Source/Programs/Horde/Horde.Build/Storage/StorageController.cs
--- a/Engine/Source/Programs/Horde/Horde.Build/Storage/StorageController.cs
+++ b/Engine/Source/Programs/Horde/Horde.Build/Storage/StorageController.cs
@@ -255,7 +255,7 @@ namespace Horde.Build.Storage
 			{
 				return NotFound(namespaceId);
 			}
-			if (!namespaceConfig.Authorize(AclAction.ReadBlobs, User))
+			if (!namespaceConfig.Authorize(AclAction.ReadBlobs, User) && !User.HasClaim(HordeClaimTypes.ReadNamespace, namespaceId.ToString()))
 			{
 				return Forbid(AclAction.ReadBlobs, namespaceId);
 			}
@@ -366,7 +366,7 @@ namespace Horde.Build.Storage
 			{
 				return NotFound(namespaceId);
 			}
-			if (!namespaceConfig.Authorize(AclAction.ReadRefs, User))
+			if (!namespaceConfig.Authorize(AclAction.ReadRefs, User) && !User.HasClaim(HordeClaimTypes.ReadNamespace, namespaceId.ToString()))
 			{
 				return Forbid(AclAction.ReadRefs, namespaceId);
 			}
diff --git a/Engine/Source/Programs/Horde/Horde.Build/Utilities/HordeClaimTypes.cs b/Engine/Source/Programs/Horde/Horde.Build/Utilities/HordeClaimTypes.cs
--- a/Engine/Source/Programs/Horde/Horde.Build/Utilities/HordeClaimTypes.cs
+++ b/Engine/Source/Programs/Horde/Horde.Build/Utilities/HordeClaimTypes.cs
@@ -61,6 +61,16 @@ namespace Horde.Build.Utilities
 		/// </summary>
 		public const string ExternalIssueUser = Prefix + "external-issue-user";
 
+		/// <summary>
+		/// Claim identifying an agent as executing a particular lease
+		/// </summary>
+		public const string Lease = Prefix + "lease";
+
+		/// <summary>
+		/// Claim allowing a certain namespace to be read from
+		/// </summary>
+		public const string ReadNamespace = Prefix + "read-ns";
+
 		/// <summary>
 		/// Claim allowing a certain namespace to be written to
 		/// </summary>
