Horde: Fix 415 errors on querying artifacts.

diff --git a/Engine/Source/Programs/Horde/Horde.Build/Artifacts/ArtifactId.cs b/Engine/Source/Programs/Horde/Horde.Build/Artifacts/ArtifactId.cs
--- a/Engine/Source/Programs/Horde/Horde.Build/Artifacts/ArtifactId.cs
+++ b/Engine/Source/Programs/Horde/Horde.Build/Artifacts/ArtifactId.cs
@@ -1,5 +1,6 @@
 // Copyright Epic Games, Inc. All Rights Reserved.
 
+using System.ComponentModel;
 using Horde.Build.Utilities;
 using MongoDB.Bson;
 
@@ -9,6 +10,7 @@ namespace Horde.Build.Artifacts
 	/// Unique id for an artifact
 	/// </summary>
 	/// <param name="Id">Identifier for the artifact</param>
+	[TypeConverter(typeof(ObjectIdTypeConverter<ArtifactId, ArtifactIdConverter>))]
 	[ObjectIdConverter(typeof(ArtifactIdConverter))]
 	public record struct ArtifactId(ObjectId Id)
 	{
diff --git a/Engine/Source/Programs/Horde/Horde.Build/Artifacts/ArtifactsController.cs b/Engine/Source/Programs/Horde/Horde.Build/Artifacts/ArtifactsController.cs
--- a/Engine/Source/Programs/Horde/Horde.Build/Artifacts/ArtifactsController.cs
+++ b/Engine/Source/Programs/Horde/Horde.Build/Artifacts/ArtifactsController.cs
@@ -115,6 +115,7 @@ namespace Horde.Build.Artifacts
 	/// <summary>
 	/// Public interface for artifacts
 	/// </summary>
+	[Authorize]
 	[ApiController]
 	public class ArtifactsController : HordeControllerBase
 	{
@@ -139,7 +140,6 @@ namespace Horde.Build.Artifacts
 		/// <param name="filter">Filter for returned properties</param>
 		/// <returns>Information about all the artifacts</returns>
 		[HttpGet]
-		[Authorize]
 		[Route("/api/v2/artifacts/{id}")]
 		[ProducesResponseType(typeof(FindArtifactsResponse), 200)]
 		public async Task<ActionResult<object>> GetArtifactAsync(ArtifactId id, [FromQuery] PropertyFilter? filter = null)
@@ -166,7 +166,6 @@ namespace Horde.Build.Artifacts
 		/// <param name="cancellationToken">Cancellation token for the operation</param>
 		/// <returns>Information about all the artifacts</returns>
 		[HttpGet]
-		[Authorize]
 		[Route("/api/v2/artifacts/{id}/browse")]
 		[ProducesResponseType(typeof(GetArtifactDirectoryResponse), 200)]
 		public async Task<ActionResult<object>> BrowseArtifactAsync(ArtifactId id, [FromQuery] string? path = null, [FromQuery] PropertyFilter? filter = null, CancellationToken cancellationToken = default)
@@ -209,12 +208,11 @@ namespace Horde.Build.Artifacts
 		/// <param name="filter">Filter for returned values</param>
 		/// <returns>Information about all the artifacts</returns>
 		[HttpGet]
-		[Authorize]
 		[Route("/api/v2/artifacts")]
 		[ProducesResponseType(typeof(FindArtifactsResponse), 200)]
 		public async Task<ActionResult<object>> FindArtifactsAsync([FromQuery(Name = "id")] IEnumerable<ArtifactId>? ids = null, [FromQuery(Name = "key")] IEnumerable<string>? keys = null, [FromQuery] PropertyFilter? filter = null)
 		{
-			if (ids == null && keys == null)
+			if ((ids == null || !ids.Any()) && (keys == null || !keys.Any()))
 			{
 				return BadRequest("At least one search term must be specified");
 			}
