Horde: Fix exception due to disposed semaphore (BackendInfo object starting with ref count of zero instead of one).

diff --git a/Engine/Source/Programs/Horde/Horde.Build/Storage/StorageService.cs b/Engine/Source/Programs/Horde/Horde.Build/Storage/StorageService.cs
--- a/Engine/Source/Programs/Horde/Horde.Build/Storage/StorageService.cs
+++ b/Engine/Source/Programs/Horde/Horde.Build/Storage/StorageService.cs
@@ -264,7 +264,7 @@ namespace Horde.Build.Storage
 		{
 			public IStorageBackend Backend { get; }
 
-			int _refCount;
+			int _refCount = 1;
 
 			public BackendInfo(IStorageBackend backend)
 			{
