Horde: Adding job detail timelines for jobs, labels, and steps with dependencies
#rnx
#skipci

diff --git a/Engine/Source/Programs/Horde/HordeDashboard/src/components/jobDetailsV2/JobDetailStepTrends.tsx b/Engine/Source/Programs/Horde/HordeDashboard/src/components/jobDetailsV2/JobDetailStepTrends.tsx
--- a/Engine/Source/Programs/Horde/HordeDashboard/src/components/jobDetailsV2/JobDetailStepTrends.tsx
+++ b/Engine/Source/Programs/Horde/HordeDashboard/src/components/jobDetailsV2/JobDetailStepTrends.tsx
@@ -284,7 +284,7 @@ class TrendsDataView extends JobDataView {
       });
    };
 
-   order = 7;
+   order = 8;
 }
 
 JobDetailsV2.registerDataView("TrendsDataView", (details: JobDetailsV2) => new TrendsDataView(details));
diff --git a/Engine/Source/Programs/Horde/HordeDashboard/src/components/jobDetailsV2/JobDetailTimeline.tsx b/Engine/Source/Programs/Horde/HordeDashboard/src/components/jobDetailsV2/JobDetailTimeline.tsx
new file mode 100644
--- /dev/null
+++ b/Engine/Source/Programs/Horde/HordeDashboard/src/components/jobDetailsV2/JobDetailTimeline.tsx
@@ -0,0 +1,831 @@
+// Copyright Epic Games, Inc. All Rights Reserved.
+
+import { Slider, Spinner, SpinnerSize, Stack, Text } from "@fluentui/react";
+import * as d3 from "d3";
+import { observer } from "mobx-react-lite";
+import { useEffect, useState } from "react";
+import { GetBatchResponse, GetStepResponse, JobStepOutcome, StepData } from "../../backend/Api";
+import { ISideRailLink } from "../../base/components/SideRail";
+import { hordeClasses, modeColors } from "../../styles/Styles";
+import { JobDataView, JobDetailsV2 } from "./JobDetailsViewCommon";
+
+import moment from "moment";
+import { NavigateFunction, useNavigate } from "react-router-dom";
+import dashboard, { StatusColor } from "../../backend/Dashboard";
+import { msecToElapsed } from "../../base/utilities/timeUtils";
+
+type SelectionType = d3.Selection<SVGGElement, unknown, null, undefined>;
+
+type Batch = GetBatchResponse & {
+   utcStart: Date;
+   utcFinish: Date;
+   utcWaitStart: Date;
+   utcWaitFinish: Date;
+   utcInitStart: Date;
+   utcInitFinish: Date;
+}
+
+type Lane = {
+   batches: Batch[];
+}
+
+enum SpanType {
+   Wait,
+   Init,
+   Step
+}
+
+type Span = {
+
+   lane: number;
+
+   batch: Batch;
+
+   type: SpanType;
+   utcStart: Date;
+   utcFinish: Date;
+
+   cost?: number;
+
+   step?: GetStepResponse;
+   stepName?: string;
+
+   filtered?: boolean;
+}
+
+const sideRail: ISideRailLink = { text: "Timeline", url: "rail_timeline" };
+
+class TimelineDataView extends JobDataView {
+
+   filterUpdated() {
+      this.dirty = true;
+   }
+
+   set() {
+
+      const details = this.details;
+
+      if (!details) {
+         return;
+      }
+
+      this.initialize([sideRail]);
+
+   }
+
+   clear() {
+      this.allSpans = undefined;
+      this.minTime = undefined;
+      this.maxTime = undefined;
+      this.anchorMS = undefined;
+      this.maxSpanMS = undefined;
+      this.filterTime = 0;
+      this.filterCost = 0;
+
+      super.clear();
+   }
+
+   detailsUpdated() {
+      if (this.details?.timing) {
+         this.generateSpans();
+         this.updateReady();
+      }
+   }
+
+   generateSpans() {
+
+      if (!this.details?.jobData?.batches) {
+         return;
+      }
+
+      const jobDetails = this.details;
+
+      const job = jobDetails.jobData!;
+
+      let bresponses = job.batches?.filter(b => !!b.agentId && (!!b.startTime && !!b.finishTime) && b.steps?.length)!.map(b => { return { ...b } as GetBatchResponse })!;
+
+      if (jobDetails.filter) {
+         const label = jobDetails.filter.label;
+         if (label) {
+
+            bresponses?.forEach(b => {
+               b.steps = b.steps?.filter(s => {
+                  const node = jobDetails.nodeByStepId(s.id);
+                  return label.includedNodes.indexOf(node?.name ?? "") !== -1;
+               });
+            });
+
+            bresponses = bresponses.filter(b => !!b.steps?.length);
+         }
+      }
+
+      let filterStep: StepData | undefined;
+
+      if (this.filterStepId) {
+         filterStep = jobDetails.stepById(this.filterStepId);
+         const steps = new Set(jobDetails.getStepDependencies(this.filterStepId).map(s => s.id));
+
+         bresponses?.forEach(b => {
+            b.steps = b.steps?.filter(s => {
+               return steps.has(s.id);
+            });
+         });
+
+         bresponses = bresponses.filter(b => !!b.steps?.length);
+
+      }
+
+      const timing = jobDetails.timing!;
+
+      let maxTime = 0;
+
+      if (filterStep && filterStep.finishTime) {
+         maxTime = new Date(filterStep.finishTime).getTime();
+      }
+
+      let batches: Batch[] = [];
+
+      bresponses.forEach(b => {
+
+         const utcStart = new Date(b.startTime as string);
+         const utcFinish = new Date(b.finishTime as string);
+
+         if (!filterStep || !maxTime) {
+            if (utcFinish.getTime() > maxTime) {
+               maxTime = utcFinish.getTime();
+            }
+         }
+
+         let utcWaitStart = utcStart;
+         let utcInitStart = utcStart;
+         let utcInitFinish = utcFinish;
+         let utcWaitFinish = utcFinish;
+
+         const stepTiming = timing.steps[b.steps[0].id];
+         const anchor = b.steps[0]?.startTime ? new Date(b.steps[0]?.startTime) : utcStart;
+
+         /*
+         // verify batch timing is consistent, lots of redundant data from step timing
+         console.assert(stepTiming);
+         b.steps.forEach(s => {
+            const otime = timing.steps[s.id];
+            console.assert(stepTiming.totalInitTime == otime.totalInitTime);
+            console.assert(stepTiming.totalWaitTime == otime.totalWaitTime);
+         });
+         */
+
+         //console.assert(stepTiming && stepTiming.totalWaitTime && stepTiming.totalInitTime)
+
+         if (stepTiming && stepTiming.totalWaitTime && stepTiming.totalInitTime) {
+
+            utcWaitStart = new Date(anchor.getTime() - (stepTiming.totalWaitTime /*- stepTiming.totalInitTime*/) * 1000);
+            utcInitStart = new Date(anchor.getTime() - stepTiming.totalInitTime * 1000);
+            utcWaitFinish = utcInitStart;
+            utcInitFinish = anchor;
+         }
+
+         batches.push({ ...b, utcStart: utcStart, utcFinish: utcFinish, utcInitStart: utcInitStart, utcInitFinish: utcInitFinish, utcWaitStart: utcWaitStart, utcWaitFinish: utcWaitFinish });
+
+      });
+
+      batches = batches.sort((a, b) => {
+         return a.utcWaitStart.getTime() - b.utcWaitStart.getTime();
+      });
+
+      this.minTime = batches[0].utcWaitStart;
+
+      this.maxTime = new Date(maxTime);
+
+      const bestLane = (batch: Batch, lanes: Lane[]) => {
+
+         return lanes.find(lane => {
+            return batch.utcWaitStart.getTime() > lane.batches[lane.batches.length - 1].utcFinish.getTime();
+         });
+      }
+
+      const processLanes = (batch: Batch, lanes: Lane[]) => {
+
+         const waitStart = batch.utcWaitStart.getTime();
+
+         const clanes = lanes.filter(lane => {
+            return lane.batches[lane.batches.length - 1].utcFinish.getTime() < waitStart;
+         });
+
+         if (!clanes.length) {
+            // need a new lane
+            lanes.push({ batches: [batch] });
+         } else {
+            const lane = bestLane(batch, clanes);
+            console.assert(lane);
+            lane!.batches.push(batch);
+         }
+
+      }
+
+      const allLanes: Lane[] = [];
+
+      this.anchorMS = batches[0].utcWaitStart.getTime();
+
+      batches.forEach(b => {
+         processLanes(b, allLanes);
+      });
+
+      const spans: Span[] = [];
+
+      allLanes.forEach((lane, index) => {
+         lane.batches.forEach(batch => {
+
+            spans.push({
+               lane: index,
+               batch: batch,
+               type: SpanType.Wait,
+               utcStart: batch.utcWaitStart,
+               utcFinish: batch.utcWaitFinish,
+            });
+
+            let initPrice: number | undefined;
+
+            if (batch.agentRate) {
+               const start = moment(batch.utcInitStart);
+               const end = moment(batch.utcInitFinish);
+               const hours = moment.duration(end.diff(start)).asHours();
+               initPrice = hours * batch.agentRate;
+            }
+
+            spans.push({
+               lane: index,
+               batch: batch,
+               type: SpanType.Init,
+               utcStart: batch.utcInitStart,
+               utcFinish: batch.utcInitFinish,
+               cost: initPrice
+            });
+
+
+            batch.steps.forEach(step => {
+
+               if (!step.startTime || !step.finishTime || step.outcome === JobStepOutcome.Unspecified) {
+                  return;
+               }
+
+               let stepPrice: number | undefined;
+
+               if (batch.agentRate) {
+                  const start = moment(step.startTime);
+                  const end = moment(step.finishTime);
+                  const hours = moment.duration(end.diff(start)).asHours();
+                  stepPrice = hours * batch.agentRate;
+               }
+
+               spans.push({
+                  lane: index,
+                  batch: batch,
+                  type: SpanType.Step,
+                  utcStart: new Date(step.startTime!),
+                  utcFinish: new Date(step.finishTime!),
+                  step: step,
+                  stepName: jobDetails.getStepName(step.id),
+                  cost: stepPrice
+               });
+            })
+         });
+      });
+
+      this.maxSpanMS = 0;
+      this.maxSpanCost = 0;
+
+      spans.forEach(s => {
+
+         if (s.cost && s.cost > this.maxSpanCost!) {
+            this.maxSpanCost = s.cost;
+         }
+
+         const t = s.utcFinish.getTime() - s.utcStart.getTime();
+         this.maxSpanMS = Math.max(t, this.maxSpanMS!);
+      });
+
+      this.allSpans = spans;
+   }
+
+   get spans(): Span[] {
+
+      if (this.dirty) {
+         this.generateSpans();
+         this.dirty = false;
+      }
+      
+      if (!this.allSpans) {
+         return [];
+      }
+
+      let fcost = this.filterCost;
+      let ftime = this.filterTime;
+
+      this.allSpans.forEach(s => s.filtered = false);
+
+      if (!this.filterCost && !this.filterTime) {
+         return this.allSpans
+      }
+
+      this.allSpans.forEach(s => {
+         s.filtered = false;
+
+         if (fcost) {
+            if ((s.type === 0 || !s.cost)) {
+               s.filtered = true;
+            } else {
+               if (s.cost < fcost) {
+                  s.filtered = true;
+               }
+            }
+         }
+
+         if (ftime) {
+            const t = s.utcFinish.getTime() - s.utcStart.getTime();
+            if (t < ftime) {
+               s.filtered = true;
+            }
+         }
+      });
+
+      return this.allSpans;
+
+   }
+
+   private allSpans?: Span[];
+   minTime?: Date;
+   maxTime?: Date;
+   maxSpanMS?: number;
+   maxSpanCost?: number;
+   anchorMS?: number;
+
+   order: number = 7
+
+   filterTime: number = 0;
+   filterCost: number = 0;
+   filterStepId?: string;
+
+   dirty = false; 
+
+}
+
+class TimelineRenderer {
+
+   // find max active and factor in wait time
+
+   constructor(dataView: TimelineDataView) {
+
+      this.dataView = dataView;
+      this.margin = { top: 0, right: 32, bottom: 0, left: 32 };
+
+   }
+
+   render(container: HTMLDivElement, navigate: NavigateFunction) {
+
+      if (this.hasRendered && !this.forceRender) {
+         //return;
+      }
+
+      const dataView = this.dataView;
+
+      if (!dataView.spans.length) {
+         return;
+      }
+
+      this.hasRendered = true;
+
+      const width = 1800;
+      const margin = this.margin;
+
+      const spans = dataView.spans;
+
+      const scolors = dashboard.getStatusColors();
+      const colors: Record<string, string> = {
+         "Success": scolors.get(StatusColor.Success)!,
+         "Failure": scolors.get(StatusColor.Failure)!,
+         "Warnings": scolors.get(StatusColor.Warnings)!,
+         "Unspecified": scolors.get(StatusColor.Skipped)!,
+      };
+
+
+      const X = d3.map(spans, (s) => s);
+      const Y = d3.map(spans, (s) => s.lane);
+
+      let yDomain: any = Y;
+      yDomain = new d3.InternSet(yDomain);
+
+      const I = d3.range(X.length);
+
+      const yPadding = 1;
+      const height = Math.ceil((yDomain.size + yPadding) * 16) + margin.top + margin.bottom;
+
+      let yRange = [margin.top, height - this.margin.bottom];
+
+      const xScale = d3.scaleLinear()
+         .domain([dataView.minTime!, dataView.maxTime!].map(d => d.getTime() / 1000))
+         .range([margin.left, width - margin.right])
+
+      const yScale = d3.scalePoint(yDomain, yRange).round(true).padding(yPadding);
+
+      let svg = this.svg;
+
+      if (!svg) {
+         svg = d3.select(container)
+            .append("svg") as any as SelectionType
+
+         this.svg = svg;
+      } else {
+         // remove tooltip
+         d3.select(container).selectAll('div').remove();
+         svg.selectAll("*").remove();
+      }
+
+      svg.attr("viewBox", [0, 0, width, height + 24] as any);
+
+      const g = svg.append("g")
+         .selectAll()
+         .data(d3.group(I, i => Y[i]))
+         .join("g")
+         .attr("transform", ([y]) => `translate(0,${(yScale(y) as any) + 16})`);
+
+      g.selectAll("line")
+         .data(([, I]: any) => I)
+         .join("line")
+         .attr("x1", i => xScale(X[i as any].utcStart.getTime() / 1000))
+         .attr("x2", i => xScale(X[i as any].utcFinish.getTime() / 1000))
+         .attr("stroke-width", i => X[i as any].filtered ? 0 : 5)
+         .attr("stroke-linecap", 0)
+         .attr("stroke", i => {
+            const span = spans[i as any];
+            if (span.type === 0) {
+               return dashboard.darktheme ? "#5B6367" : "#D3D2D1";
+            }
+            if (span.type === 1) {
+               return dashboard.darktheme ? "#506F7C" : "#BAD1DB";               
+            }
+
+            return colors[span.step?.outcome ?? "Unspecified"];
+
+         })
+
+      const radius = 2.5
+      g.selectAll("circle")
+         .data(([, I]: any) => I)
+         .join("circle")
+         .attr("cx", i => xScale(X[i as any].utcStart.getTime() / 1000))
+         .attr("fill", i => "#2D3F5F")
+         .attr("r", i => X[i as any].filtered ? 0 : radius);
+
+
+      const xAxis = (g: SelectionType) => {
+
+         g.attr("transform", `translate(0,16)`)
+            .style("font-family", "Horde Open Sans Regular")
+            .style("font-size", "9px")
+            .call(d3.axisTop(xScale)
+               .tickFormat(d => {
+                  return msecToElapsed(((d as number) * 1000) - dataView.anchorMS!, true, false);
+               })
+               .tickSizeOuter(0))
+            .call(g => g.select(".domain").remove())
+            .call(g => g.selectAll(".tick line").attr("stroke-opacity", dashboard.darktheme ? 0.35 : 0.25)
+               .attr("stroke", dashboard.darktheme ? "#6D6C6B" : "#4D4C4B")
+               .attr("y2", height - this.margin.bottom))
+
+      }
+
+      // top axis
+      svg.append("g").attr("class", "x-axis").call(xAxis)
+
+      function zoomed(event: any) {
+         xScale.range([margin.left, width - margin.right].map(d => event.transform.applyX(d)));
+         svg!.selectAll("circle")
+            .attr("cx", i => xScale(X[i as any].utcStart.getTime() / 1000));
+         svg!.selectAll("line")
+            .attr("x1", i => { if (i as number >= X.length) return 0; return xScale(X[i as any].utcStart.getTime() / 1000) })
+            .attr("x2", i => { if (i as number >= X.length) return 0; return xScale(X[i as any].utcFinish.getTime() / 1000) })
+
+         svg!.selectAll(".x-axis").call(xAxis as any);
+
+      }
+
+      const zoom = d3.zoom()
+         .scaleExtent([1, 48])
+         .extent([[margin.left, 0], [width - margin.right, height]])
+         .translateExtent([[margin.left, -Infinity], [width - margin.right, Infinity]])
+         .on("zoom", zoomed)
+
+
+
+      svg.call(zoom as any);
+
+      const tooltip = d3.select(container)
+         .append("div")
+         .attr("id", "tooltip")
+         .style("display", "none")
+         .style("background-color", modeColors.background)
+         .style("border", "solid")
+         .style("border-width", "1px")
+         .style("border-radius", "3px")
+         .style("border-color", dashboard.darktheme ? "#413F3D" : "#2D3F5F")
+         .style("padding", "8px")
+         .style("position", "absolute")
+         .style("pointer-events", "none");
+
+
+
+      const closestData = (x: number, y: number): Span | undefined => {
+
+         y -= 16;
+
+         let first = spans.find(s => !s.filtered);
+
+         if (!first) {
+            return undefined;
+         }
+
+         let closest = spans.reduce<Span>((best: Span, span: Span, i: number, spans: Span[]) => {
+
+            if (span.filtered) {
+               return best;
+            }
+
+            const by = Math.abs(yScale(best.lane)! - y);
+            const cy = Math.abs(yScale(span.lane)! - y);
+
+            if (cy > by) {
+               return best;
+            }
+
+            const ss = xScale(span.utcStart.getTime() / 1000);
+            const sf = xScale(span.utcFinish.getTime() / 1000);
+
+            if (x >= ss && x <= sf) {
+               return span;
+            }
+
+            return best;
+
+         }, first);
+
+         const ss = xScale(closest.utcStart.getTime() / 1000);
+         const sf = xScale(closest.utcFinish.getTime() / 1000);
+
+         if (x < ss || x > sf) {
+            return undefined;
+         }
+
+         return closest;
+
+      }
+
+      let mousedown = false;
+
+      const handleMouseMove = (event: any) => {
+
+         let mouseX = d3.pointer(event)[0];
+         let mouseY = d3.pointer(event)[1];
+
+         const span = closestData(mouseX, mouseY);
+
+         if (!span) {
+            tooltip.style("display", "none");
+            return;
+         }
+
+         if (span) {
+
+            const details = dataView.details!;
+
+            const group = details.groups[span.batch.groupIdx];
+            const pool = details.stream?.agentTypes[group?.agentType!];
+
+            let html = "";
+
+            let agentHtml = "";
+            agentHtml += `Type: ${group?.agentType.toUpperCase() ?? "Unknown Agent Type"}<br/>`;
+            agentHtml += `Pool: ${[pool?.pool?.toUpperCase() ?? "Unknown Pool"]}<br/>`;
+            agentHtml += `Agent: ${[span.batch.agentId ?? "Unknown Pool"]}<br/>`;
+
+            if (span.cost) {
+               agentHtml += `Estimated cost: $${span.cost.toFixed(2)}`;
+            }
+
+            const elapsed = msecToElapsed(span.utcFinish.getTime() - span.utcStart.getTime(), true, false);
+
+            if (span.type === 0) {
+
+               html += `Wait: ${elapsed}<br/>`;
+               html += agentHtml;
+            }
+            if (span.type === 1) {
+               html += `Init: <b>${elapsed}</b><br/>`;
+               html += agentHtml;
+            }
+            if (span.type === 2) {
+               html = `${span.stepName ?? "Unknown Step"}<br/>`;
+               html += `Time: ${elapsed}<br/>`;
+               html += agentHtml;
+            }
+
+            let topY = mouseY - 78;
+            if (topY < 0) {
+               topY = 0;
+            }
+
+            let tipX = mouseX ;
+            let translateX = "0%";
+
+            if (tipX > 1500) {               
+               translateX = "-180%";
+            }
+
+            tooltip
+               .style("display", mousedown ? "none" : "block")
+               .html(html ?? "")
+               .style("position", `absolute`)
+               .style("width", `max-content`)
+               .style("top", (topY) + "px")
+               .style("left", `${tipX}px`)
+               .style("transform", `translateX(${translateX})`)
+               .style("font-family", "Horde Open Sans Semibold")
+               .style("font-size", "10px")
+               .style("line-height", "16px")
+               .style("shapeRendering", "crispEdges")
+               .style("stroke", "none")
+         }
+      }
+
+      const handleMouseLeave = (event: any) => {
+         mousedown = false;
+         tooltip.style("display", "none");
+      }
+
+      const handleMouseClick = (event: any) => {
+
+         let mouseX = d3.pointer(event)[0];
+         let mouseY = d3.pointer(event)[1];
+
+         const span = closestData(mouseX, mouseY);
+
+         if (span && span.step) {
+            navigate(`/job/${dataView.details!.jobId}?step=${span.step.id}`);
+         }
+      }
+
+
+      svg.on("click", (event) => handleMouseClick(event))
+
+      svg.on("pointerdown", (event) => { mousedown = true; tooltip.style("display", "none"); });
+      svg.on("pointerup", (event) => { mousedown = false; tooltip.style("display", "block"); });
+
+      svg.on("mousemove", (event) => handleMouseMove(event));
+      svg.on("mouseleave", (event) => { handleMouseLeave(event); })
+      svg.on("wheel", (event) => { event.preventDefault(); })
+
+   }
+
+   margin: { top: number, right: number, bottom: number, left: number }
+
+   dataView: TimelineDataView;
+
+   svg?: SelectionType;
+
+   hasRendered = false;
+   forceRender = false;
+}
+
+
+const TimelineGraph: React.FC<{ dataView: TimelineDataView, filterTime: number, filterCost: number, stepId?: string }> = ({ dataView, filterTime, filterCost, stepId }) => {
+
+   const graph_container_id = `timeline_graph_container`;
+
+   const [container, setContainer] = useState<HTMLDivElement | null>(null);
+   const [state, setState] = useState<{ graph?: TimelineRenderer }>({});
+   const navigate = useNavigate();
+
+   if (!dataView.details?.timing) {
+      return null;
+   }
+
+   // todo: we can probably cache
+   dataView.filterTime = filterTime;
+   dataView.filterCost = filterCost;
+   if (dataView.filterStepId !==  stepId) {
+      dataView.filterStepId = stepId;
+      dataView.dirty = true;
+   }
+   
+
+   if (!state.graph) {
+      setState({ ...state, graph: new TimelineRenderer(dataView) })
+      return null;
+   }
+
+   if (container) {
+      try {
+         state.graph?.render(container, navigate);
+
+      } catch (err) {
+         console.error(err);
+      }
+
+   }
+
+   return <Stack className={hordeClasses.horde}>
+      <Stack style={{ paddingLeft: 8, paddingTop: 8 }}>
+         <div id={graph_container_id} className="horde-no-darktheme" style={{ shapeRendering: "crispEdges", userSelect: "none", position: "relative" }} ref={(ref: HTMLDivElement) => setContainer(ref)} onMouseEnter={() => { }} onMouseLeave={() => { }} />
+      </Stack>
+   </Stack>;
+
+}
+
+
+JobDetailsV2.registerDataView("TimelineDataView", (details: JobDetailsV2) => new TimelineDataView(details));
+
+export const TimelinePanel: React.FC<{ jobDetails: JobDetailsV2, stepId?: string }> = observer(({ jobDetails, stepId }) => {
+
+   const [state, setState] = useState({ time: 0, cost: 0 });
+
+   if (jobDetails.updated) { }
+
+   const dataView = jobDetails.getDataView<TimelineDataView>("TimelineDataView");
+
+   useEffect(() => {
+      return () => {
+         dataView?.clear();
+      };
+   }, [dataView]);
+
+   dataView.subscribe();
+
+   if (!jobDetails.jobData) {
+      return null;
+   }
+
+   dataView.set();
+
+   const job = jobDetails.jobData;
+   const batches = job.batches?.filter(b => !!b.agentId && (!!b.startTime && !!b.finishTime));
+
+   if (!batches?.length) {
+      return null;
+   }
+
+   return (<Stack id={sideRail.url} styles={{ root: { paddingTop: 18, paddingRight: 12 } }}>
+      <Stack className={hordeClasses.raised}>
+         <Stack tokens={{ childrenGap: 18 }}>
+            <Stack horizontal>
+               <Stack>
+                  <Text variant="mediumPlus" styles={{ root: { fontFamily: "Horde Open Sans SemiBold" } }}>Timeline</Text>
+               </Stack>
+               <Stack grow />
+               {!!dataView.minTime && <Stack horizontal tokens={{ childrenGap: 18 }}>
+                  <Stack>
+                     <Slider
+                        styles={{ root: { width: 240 } }}
+                        label="Min Time"
+                        min={0}
+                        max={dataView.maxSpanMS!}
+                        value={state.time}
+                        step={10}
+                        showValue
+                        valueFormat={(value) => {
+                           if (!value) return ""
+                           return msecToElapsed(value, true, false)
+                        }}
+                        onChange={(time) => {
+                           setState({ ...state, time: time })
+                        }}
+                     />
+                  </Stack>
+                  {!!dataView.maxSpanCost && <Stack>
+                     <Slider
+                        styles={{ root: { width: 240 } }}
+                        label="Min Cost"
+                        min={0}
+                        step={dataView.maxSpanCost! / 10}
+                        max={dataView.maxSpanCost!}
+                        value={state.cost}
+                        showValue
+                        valueFormat={(value) => {
+                           if (!value) return ""
+                           return `$${value.toFixed(2)}`
+                        }}
+                        onChange={(cost) => {
+                           setState({ ...state, cost: cost })
+                        }}
+                     />
+                  </Stack>}
+               </Stack>}
+
+            </Stack>
+            <Stack>
+               {!!dataView.minTime && <TimelineGraph dataView={dataView} stepId={stepId} filterTime={state.time} filterCost={state.cost} />}
+               {!dataView.minTime && <Spinner size={SpinnerSize.large} />}
+            </Stack>
+         </Stack>
+      </Stack>
+   </Stack>);
+
+})
diff --git a/Engine/Source/Programs/Horde/HordeDashboard/src/components/jobDetailsV2/JobDetailViewV2.tsx b/Engine/Source/Programs/Horde/HordeDashboard/src/components/jobDetailsV2/JobDetailViewV2.tsx
--- a/Engine/Source/Programs/Horde/HordeDashboard/src/components/jobDetailsV2/JobDetailViewV2.tsx
+++ b/Engine/Source/Programs/Horde/HordeDashboard/src/components/jobDetailsV2/JobDetailViewV2.tsx
@@ -1,7 +1,7 @@
-import { mergeStyleSets, Pivot, PivotItem, ScrollablePane, ScrollbarVisibility, Stack } from "@fluentui/react";
+import { mergeStyleSets, Pivot, PivotItem, Stack } from "@fluentui/react";
 import { observer } from "mobx-react-lite";
-import { useEffect, useState } from "react";
-import { useNavigate, useParams } from "react-router-dom";
+import { useEffect, useRef, useState } from "react";
+import { useLocation, useNavigate, useParams } from "react-router-dom";
 import { GetJobsTabResponse } from "../../backend/Api";
 import { useWindowSize } from "../../base/utilities/hooks";
 import { hordeClasses, modeColors } from "../../styles/Styles";
@@ -11,6 +11,7 @@ import { TopNav } from "../TopNav";
 import { HealthPanel } from "./JobDetailHealthV2";
 import { StepTrendsPanel } from "./JobDetailStepTrends";
 import { JobDataView, JobDetailsV2 } from "./JobDetailsViewCommon";
+import { TimelinePanel } from "./JobDetailTimeline";
 import { StepsPanelV2 } from "./JobDetailViewSteps";
 import { SummaryPanel } from "./JobDetailViewSummary";
 import { JobOperations } from "./JobOperationsBar";
@@ -296,25 +297,49 @@ const DetailsViewOverview: React.FC<{ jobDetails: JobDetailsV2 }> = ({ jobDetail
       <SummaryPanel jobDetails={details} />
       <StepsPanelV2 jobDetails={details} />
       <HealthPanel jobDetails={details} />
+      <TimelinePanel jobDetails={details} />
       <StepTrendsPanel jobDetails={details} />
+
    </Stack>
 };
 
 const rootWidth = 1776;
 
+const ScrollRestore: React.FC<{ jobDetails: JobDetailsV2, scrollRef: React.RefObject<HTMLDivElement> }> = ({ jobDetails, scrollRef }) => {
+
+   const location = useLocation();
+
+   if (!scrollRef?.current) {
+      return null;
+   }
+
+   try {
+      if (!location.hash) {
+         scrollRef.current.scrollTop = 0;
+      }
+   } catch (message) {
+      console.error(message);
+   }
+   
+   return null;
+}
+
+
 const DetailsView: React.FC<{ jobDetails: JobDetailsV2 }> = ({ jobDetails }) => {
 
    const windowSize = useWindowSize();
+   const scrollRef = useRef<HTMLDivElement>(null);
 
    const vw = Math.max(document.documentElement.clientWidth, window.innerWidth || 0);
 
    const details = jobDetails;
 
    const centerAlign = vw / 2 - 890;
-   const key = `windowsize_jobdetail_view_${windowSize.width}_${windowSize.height}`;   
+   const key = `windowsize_jobdetail_view_${windowSize.width}_${windowSize.height}`;
 
    return <Stack id="rail_job_overview">
       <JobBreadCrumbs jobDetails={details} />
+      <ScrollRestore jobDetails={details} scrollRef={scrollRef} />
       <Stack horizontal styles={{ root: { backgroundColor: modeColors.background } }}>
          <Stack styles={{ root: { width: "100%" } }}>
             <Stack horizontal>
@@ -332,7 +357,7 @@ const DetailsView: React.FC<{ jobDetails: JobDetailsV2 }> = ({ jobDetails }) =>
             <Stack style={{ width: "100%", backgroundColor: modeColors.background }}>
                <Stack horizontal>
                   <Stack /*className={classNames.pointerSuppress} */ style={{ position: "relative", width: "100%", height: 'calc(100vh - 228px)' }}>
-                     <ScrollablePane id="hordeContentArea" scrollbarVisibility={ScrollbarVisibility.always}>
+                     <div id="hordeContentArea" ref={scrollRef} style={{ overflowX: "hidden", overflowY: "visible" }}>
                         <Stack horizontal>
                            <Stack key={`${key}_2`} style={{ paddingLeft: centerAlign }} />
                            <Stack style={{ width: rootWidth }}>
@@ -340,7 +365,7 @@ const DetailsView: React.FC<{ jobDetails: JobDetailsV2 }> = ({ jobDetails }) =>
                               <DetailsViewStep jobDetails={details} />
                            </Stack>
                         </Stack>
-                     </ScrollablePane>
+                     </div>
                   </Stack>
                </Stack>
             </Stack>
diff --git a/Engine/Source/Programs/Horde/HordeDashboard/src/components/jobDetailsV2/JobDetailsViewCommon.tsx b/Engine/Source/Programs/Horde/HordeDashboard/src/components/jobDetailsV2/JobDetailsViewCommon.tsx
--- a/Engine/Source/Programs/Horde/HordeDashboard/src/components/jobDetailsV2/JobDetailsViewCommon.tsx
+++ b/Engine/Source/Programs/Horde/HordeDashboard/src/components/jobDetailsV2/JobDetailsViewCommon.tsx
@@ -3,7 +3,7 @@ import { action, makeObservable, observable } from "mobx";
 import { observer } from "mobx-react-lite";
 import moment from "moment";
 import React, { useState } from "react";
-import { useNavigate, useLocation } from "react-router-dom";
+import { useLocation, useNavigate } from "react-router-dom";
 import backend from "../../backend";
 import { BatchData, GetJobTimingResponse, GetLabelResponse, GetTemplateRefResponse, GroupData, JobData, JobState, LabelState, NodeData, ReportPlacement, StepData, StreamData } from "../../backend/Api";
 import { JobLabel } from "../../backend/JobDetails";
@@ -16,7 +16,7 @@ import { useQuery } from "../JobDetailCommon";
 import { LabelStatusIcon, StepStatusIcon } from "../StatusIcon";
 
 export abstract class JobDataView {
-   
+
    constructor(details: JobDetailsV2) {
       makeObservable(this);
       this.details = details;
@@ -122,7 +122,7 @@ export abstract class JobDataView {
 const defaultUpdateMS = 5000;
 export class JobDetailsV2 extends PollBase {
 
-   constructor(jobId: string) {      
+   constructor(jobId: string) {
       super(defaultUpdateMS)
       makeObservable(this);
       this.jobId = jobId;
@@ -722,6 +722,45 @@ export class JobDetailsV2 extends PollBase {
       this.views = this.views.filter(v => v !== view);
    }
 
+   getStepDependencies(stepId: string): StepData[] {
+
+      const step = this.stepById(stepId);
+
+      if (!step) {
+         return [];
+      }
+
+      const steps: StepData[] = [];
+
+      const nodes: NodeData[] = [];
+
+      const getStepsRecursive = (stepId: string) => {
+
+         const stepNode = this.nodeByStepId(stepId);
+
+         if (!stepNode || nodes.find(n => stepNode === n)) {
+            return;
+         }
+
+         nodes.push(stepNode);
+
+         [stepNode.inputDependencies, stepNode.orderDependencies].flat().forEach(name => {
+            const s = this.stepByName(name);
+            if (s) {
+               steps.push(s);
+               getStepsRecursive(s.id);
+            }
+         });
+      };
+
+      getStepsRecursive(stepId);
+
+      steps.push(step);
+      
+      return steps;
+
+   }
+
    getDataView<T extends JobDataView>(name: string): T {
 
       let view = this.views.find(v => v.name === name);
@@ -1261,7 +1300,7 @@ export const JobFilterBar: React.FC<{ jobDetails: JobDetailsV2 }> = observer(({
                   if (idx >= 0) {
                      navigate(location.pathname + `?label=${idx}`);
                   } else {
-                     navigate(location.pathname, {replace: true});
+                     navigate(location.pathname, { replace: true });
                   }
                }
 
diff --git a/Engine/Source/Programs/Horde/HordeDashboard/src/components/jobDetailsV2/StepDetailView.tsx b/Engine/Source/Programs/Horde/HordeDashboard/src/components/jobDetailsV2/StepDetailView.tsx
--- a/Engine/Source/Programs/Horde/HordeDashboard/src/components/jobDetailsV2/StepDetailView.tsx
+++ b/Engine/Source/Programs/Horde/HordeDashboard/src/components/jobDetailsV2/StepDetailView.tsx
@@ -12,6 +12,7 @@ import { HealthPanel } from "./JobDetailHealthV2";
 import { StepHistoryPanel } from "./JobDetailStepHistory";
 import { StepTrendsPanel } from "./JobDetailStepTrends";
 import { JobDataView, JobDetailsV2 } from "./JobDetailsViewCommon";
+import { TimelinePanel } from "./JobDetailTimeline";
 import { StepsPanelV2 } from "./JobDetailViewSteps";
 import { StepSummaryPanel } from "./StepDetailSummary";
 import { StepTestReportPanel } from "./StepDetailTestPanel";
@@ -88,6 +89,9 @@ const StepDetailViewInner: React.FC<{ jobDetails: JobDetailsV2, stepId: string }
       <Stack>
          <JobDetailArtifactsV2 jobDetails={jobDetails} stepId={stepId} />
       </Stack>
+      {!!step && <Stack>
+         <TimelinePanel jobDetails={jobDetails} stepId={stepId} />
+      </Stack>}
       {!!step && <Stack>
          <StepTrendsPanel jobDetails={jobDetails} />
       </Stack>}
