Horde: Don't virtualize steps list, less janky on resizes (though, may have performance issues)
#skipci
#rnx

diff --git a/Engine/Source/Programs/Horde/HordeDashboard/src/components/jobDetailsV2/JobDetailViewSteps.tsx b/Engine/Source/Programs/Horde/HordeDashboard/src/components/jobDetailsV2/JobDetailViewSteps.tsx
--- a/Engine/Source/Programs/Horde/HordeDashboard/src/components/jobDetailsV2/JobDetailViewSteps.tsx
+++ b/Engine/Source/Programs/Horde/HordeDashboard/src/components/jobDetailsV2/JobDetailViewSteps.tsx
@@ -614,7 +614,7 @@ export const StepsPanelInner: React.FC<{ jobDetails: JobDetailsV2, depStepId?: s
          onRenderDetailsHeader={onRenderDetailsHeader}
          onRenderItemColumn={onRenderItemColumn}
          onRenderRow={onRenderRow}
-         onShouldVirtualize={() => { return true; }} // <--- significant speedup if enabled
+         onShouldVirtualize={() => { return false; }} // <--- previous versions of Fluent were very slow if not virtualized, new version seems faster and not virtualizing makes the size estimations less janky, if there are performance issues, look here
       />
    </Stack>;
 
