Horde: Add back functionality to generate separate blocks/tags for temp storage while still also creating artifacts.

diff --git a/Engine/Source/Programs/Horde/Horde.Agent.Tests/Execution/WorkspaceExecutorTest.cs b/Engine/Source/Programs/Horde/Horde.Agent.Tests/Execution/WorkspaceExecutorTest.cs
--- a/Engine/Source/Programs/Horde/Horde.Agent.Tests/Execution/WorkspaceExecutorTest.cs
+++ b/Engine/Source/Programs/Horde/Horde.Agent.Tests/Execution/WorkspaceExecutorTest.cs
@@ -49,7 +49,7 @@ public sealed class WorkspaceExecutorTest : IDisposable
 		_workspace.SetFile(1, "main.cpp", "main");
 		_workspace.SetFile(1, "foo/bar/baz.h", "baz");
 
-		JobExecutorOptions executorOptions = new JobExecutorOptions(_session, null!, JobId, "batch1", AgentType, new JobOptions());
+		JobExecutorOptions executorOptions = new JobExecutorOptions(_session, null!, JobId, "batch1", AgentType, default, null!, new JobOptions());
 		_executor = new (executorOptions, null, _workspace, NullLogger.Instance);
 	}
 
@@ -71,7 +71,7 @@ public sealed class WorkspaceExecutorTest : IDisposable
 	[TestMethod]
 	public async Task RegularAndAutoSdkWorkspace()
 	{
-		JobExecutorOptions executorOptions = new JobExecutorOptions(_session, null!, JobId, "batch1", AgentType, new JobOptions());
+		JobExecutorOptions executorOptions = new JobExecutorOptions(_session, null!, JobId, "batch1", AgentType, default, null!, new JobOptions());
 		WorkspaceExecutor executor = new (executorOptions, _autoSdkWorkspace, _workspace, NullLogger.Instance);
 
 		await executor.InitializeAsync(_logger, CancellationToken.None);
@@ -84,7 +84,7 @@ public sealed class WorkspaceExecutorTest : IDisposable
 	[TestMethod]
 	public async Task EnvVars()
 	{
-		JobExecutorOptions executorOptions = new JobExecutorOptions(_session, null!, JobId, "batch1", AgentType, new JobOptions());
+		JobExecutorOptions executorOptions = new JobExecutorOptions(_session, null!, JobId, "batch1", AgentType, default, null!, new JobOptions());
 		WorkspaceExecutor executor = new (executorOptions, _autoSdkWorkspace, _workspace, NullLogger.Instance);
 		await executor.InitializeAsync(_logger, CancellationToken.None);
 
@@ -109,7 +109,7 @@ public sealed class WorkspaceExecutorTest : IDisposable
 		_server.AddJob("jobPreflight", StreamId, 1, 1000);		
 		_workspace.SetFile(1000, "New/Feature/Foo.cs", "foo");
 
-		JobExecutorOptions executorOptions = new JobExecutorOptions(_session, null!, "jobPreflight", "batch1", AgentType, new JobOptions());
+		JobExecutorOptions executorOptions = new JobExecutorOptions(_session, null!, "jobPreflight", "batch1", AgentType, default, null!, new JobOptions());
 		WorkspaceExecutor executor = new (executorOptions, null, _workspace, NullLogger.Instance);
 		
 		await executor.InitializeAsync(_logger, CancellationToken.None);
@@ -123,7 +123,7 @@ public sealed class WorkspaceExecutorTest : IDisposable
 	public async Task JobWithNoChange()
 	{
 		_server.AddJob("jobNoChange", StreamId, 0, 0);		
-		JobExecutorOptions executorOptions = new JobExecutorOptions(_session, null!, "jobNoChange", "batch1", AgentType, new JobOptions());
+		JobExecutorOptions executorOptions = new JobExecutorOptions(_session, null!, "jobNoChange", "batch1", AgentType, default, null!, new JobOptions());
 		WorkspaceExecutor executor = new (executorOptions, null, _workspace, NullLogger.Instance);
 		await Assert.ThrowsExceptionAsync<WorkspaceMaterializationException>(() => executor.InitializeAsync(_logger, CancellationToken.None));
 	}
diff --git a/Engine/Source/Programs/Horde/Horde.Agent/Execution/JobExecutor.cs b/Engine/Source/Programs/Horde/Horde.Agent/Execution/JobExecutor.cs
--- a/Engine/Source/Programs/Horde/Horde.Agent/Execution/JobExecutor.cs
+++ b/Engine/Source/Programs/Horde/Horde.Agent/Execution/JobExecutor.cs
@@ -50,15 +50,19 @@ namespace Horde.Agent.Execution
 		public string JobId { get; }
 		public string BatchId { get; }
 		public string AgentType { get; }
+		public NamespaceId NamespaceId { get; }
+		public string Token { get; }
 		public JobOptions JobOptions { get; }
 
-		public JobExecutorOptions(ISession session, IServerStorageFactory storageFactory, string jobId, string batchId, string agentType, JobOptions jobOptions)
+		public JobExecutorOptions(ISession session, IServerStorageFactory storageFactory, string jobId, string batchId, string agentType, NamespaceId namespaceId, string token, JobOptions jobOptions)
 		{
 			Session = session;
 			StorageFactory = storageFactory;
 			JobId = jobId;
 			BatchId = batchId;
 			AgentType = agentType;
+			NamespaceId = namespaceId;
+			Token = token;
 			JobOptions = jobOptions;
 		}
 	}
@@ -193,6 +197,8 @@ namespace Horde.Agent.Execution
 
 		protected readonly ISession _session;
 		protected readonly IServerStorageFactory _storageFactory;
+		private readonly NamespaceId _namespaceId;
+		private readonly string _token;
 		protected readonly JobOptions _jobOptions;
 
 		protected IRpcConnection RpcConnection => _session.RpcConnection;
@@ -211,6 +217,8 @@ namespace Horde.Agent.Execution
 			_batchId = options.BatchId;
 			_agentTypeName = options.AgentType;
 
+			_namespaceId = options.NamespaceId;
+			_token = options.Token;
 			_jobOptions = options.JobOptions;
 
 			_logger = logger;
@@ -767,18 +775,16 @@ namespace Horde.Agent.Execution
 		{
 			DirectoryReference manifestDir = DirectoryReference.Combine(workspaceDir, "Engine", "Saved", "BuildGraph");
 
-			await TestArtifactAsync(step.StepId, cancellationToken);
-
 			using MemoryCache cache = new MemoryCache(new MemoryCacheOptions { });
-//			IStorageClient storage = _storageFactory.CreateStorageClient(_session, $"/api/v1/jobs/{_jobId}/batches/{_batchId}/steps/{step.StepId}/temp");
-//			TreeReader reader = new TreeReader(storage, cache, _logger);
+			IStorageClient storage = _storageFactory.CreateStorageClient(_session, _namespaceId, _token);
+			TreeReader reader = new TreeReader(storage, cache, _logger);
 
 			// Create the mapping of tag names to file sets
 			Dictionary<string, HashSet<FileReference>> tagNameToFileSet = new Dictionary<string, HashSet<FileReference>>();
 
 			// Read all the input tags for this node, and build a list of referenced input storage blocks
 			HashSet<TempStorageBlockRef> inputStorageBlocks = new HashSet<TempStorageBlockRef>();
-/*			foreach (string input in step.Inputs)
+			foreach (string input in step.Inputs)
 			{
 				int slashIdx = input.IndexOf('/', StringComparison.Ordinal);
 				if (slashIdx == -1)
@@ -794,10 +800,10 @@ namespace Horde.Agent.Execution
 				tagNameToFileSet[tagName] = fileList.ToFileSet(workspaceDir);
 				inputStorageBlocks.UnionWith(fileList.Blocks);
 			}
-*/
+
 			// Read the manifests for all the input storage blocks
 			Dictionary<TempStorageBlockRef, TempStorageBlockManifest> inputManifests = new Dictionary<TempStorageBlockRef, TempStorageBlockManifest>();
-/*			using (IScope scope = GlobalTracer.Instance.BuildSpan("TempStorage").WithTag("resource", "read").StartActive())
+			using (IScope scope = GlobalTracer.Instance.BuildSpan("TempStorage").WithTag("resource", "read").StartActive())
 			{
 				Stopwatch timer = Stopwatch.StartNew();
 				scope.Span.SetTag("blocks", inputStorageBlocks.Count);
@@ -809,10 +815,10 @@ namespace Horde.Agent.Execution
 				scope.Span.SetTag("size", inputManifests.Sum(x => x.Value.GetTotalSize()));
 				logger.LogInformation("Download took {Time:n1}s", timer.Elapsed.TotalSeconds);
 			}
-*/
+
 			// Read all the input storage blocks, keeping track of which block each file came from
 			Dictionary<FileReference, TempStorageBlockRef> fileToStorageBlock = new Dictionary<FileReference, TempStorageBlockRef>();
-/*			foreach (KeyValuePair<TempStorageBlockRef, TempStorageBlockManifest> pair in inputManifests)
+			foreach (KeyValuePair<TempStorageBlockRef, TempStorageBlockManifest> pair in inputManifests)
 			{
 				TempStorageBlockRef inputStorageBlock = pair.Key;
 				foreach (FileReference file in pair.Value.Files.Select(x => x.ToFileReference(workspaceDir)))
@@ -825,7 +831,7 @@ namespace Horde.Agent.Execution
 					fileToStorageBlock[file] = inputStorageBlock;
 				}
 			}
-*/
+
 			// Run UAT
 			if (await ExecuteAutomationToolAsync(step, workspaceDir, arguments, logger, cancellationToken) != 0)
 			{
@@ -833,7 +839,6 @@ namespace Horde.Agent.Execution
 			}
 
 			// Read all the output manifests
-			HashSet<FileReference> outputFiles = new HashSet<FileReference>();
 			foreach (string tagName in step.OutputNames)
 			{
 				FileReference tagFileListLocation = TempStorage.GetTagManifestLocation(manifestDir, step.Name, tagName);
@@ -841,8 +846,6 @@ namespace Horde.Agent.Execution
 
 				TempStorageTagManifest fileList = TempStorageTagManifest.Load(tagFileListLocation);
 				tagNameToFileSet[tagName] = fileList.ToFileSet(workspaceDir);
-
-				outputFiles.Add(tagFileListLocation);
 			}
 
 			// Check that none of the inputs have been clobbered
@@ -873,32 +876,123 @@ namespace Horde.Agent.Execution
 			}
 
 			// Determine all the output files which are required to be copied to temp storage (because they're referenced by nodes in another agent)
+			HashSet<FileReference> referencedOutputFiles = new HashSet<FileReference>();
 			foreach (int publishOutput in step.PublishOutputs)
 			{
 				string tagName = step.OutputNames[publishOutput];
-				outputFiles.UnionWith(tagNameToFileSet[tagName]);
+				referencedOutputFiles.UnionWith(tagNameToFileSet[tagName]);
+			}
+
+			// Find a block name for all new outputs
+			Dictionary<FileReference, string> fileToBlockName = new Dictionary<FileReference, string>();
+			for (int idx = 0; idx < step.OutputNames.Count; idx++)
+			{
+				string tagName = step.OutputNames[idx];
+
+				string outputNameWithoutHash = tagName.TrimStart('#');
+				bool isDefaultOutput = outputNameWithoutHash.Equals(step.Name, StringComparison.OrdinalIgnoreCase);
+
+				HashSet<FileReference> files = tagNameToFileSet[tagName];
+				foreach (FileReference file in files)
+				{
+					if (!fileToStorageBlock.ContainsKey(file) && file.IsUnderDirectory(workspaceDir))
+					{
+						if (isDefaultOutput)
+						{
+							if (!fileToBlockName.ContainsKey(file))
+							{
+								fileToBlockName[file] = "";
+							}
+						}
+						else
+						{
+							string? blockName;
+							if (fileToBlockName.TryGetValue(file, out blockName) && blockName.Length > 0)
+							{
+								fileToBlockName[file] = $"{blockName}+{outputNameWithoutHash}";
+							}
+							else
+							{
+								fileToBlockName[file] = outputNameWithoutHash;
+							}
+						}
+					}
+				}
+			}
+
+			// Invert the dictionary to make a mapping of storage block to the files each contains
+			Dictionary<string, HashSet<FileReference>> outputStorageBlockToFiles = new Dictionary<string, HashSet<FileReference>>();
+			foreach (KeyValuePair<FileReference, string> pair in fileToBlockName)
+			{
+				HashSet<FileReference>? files;
+				if (!outputStorageBlockToFiles.TryGetValue(pair.Value, out files))
+				{
+					files = new HashSet<FileReference>();
+					outputStorageBlockToFiles.Add(pair.Value, files);
+				}
+				files.Add(pair.Key);
 			}
 
 			// Write all the storage blocks, and update the mapping from file to storage block
 			using (GlobalTracer.Instance.BuildSpan("TempStorage").WithTag("resource", "Write").StartActive())
 			{
 				Stopwatch timer = Stopwatch.StartNew();
+				RefName refName = TempStorage.GetRefNameForNode(RefPrefix, step.Name);
 
-				using IRpcClientRef<JobRpc.JobRpcClient> jobRpc = await RpcConnection.GetClientRefAsync<JobRpc.JobRpcClient>(cancellationToken);
-				CreateJobArtifactResponse artifact = await jobRpc.Client.CreateArtifactAsync(new CreateJobArtifactRequest { JobId = _jobId, StepId = step.StepId, Name = "Output", Type = JobArtifactType.Output }, cancellationToken: cancellationToken);
-				logger.LogInformation("Created output artifact {ArtifactId} with ref {RefName} in namespace {Namespace}", artifact.Id, artifact.RefName, artifact.NamespaceId);
+				TreeOptions treeOptions = new TreeOptions();
+				using TreeWriter treeWriter = new TreeWriter(storage, treeOptions, refName.Text);
 
-				IStorageClient storage = _storageFactory.CreateStorageClient(_session, new NamespaceId(artifact.NamespaceId), artifact.Token);
-				using TreeWriter writer = new TreeWriter(storage, new RefName(artifact.RefName));
+				DirectoryNode outputNode = new DirectoryNode();
+
+				// Create all the output blocks
+				foreach (KeyValuePair<string, HashSet<FileReference>> pair in outputStorageBlockToFiles)
+				{
+					TempStorageBlockRef outputBlock = new TempStorageBlockRef(step.Name, pair.Key);
+					foreach (FileReference file in pair.Value)
+					{
+						fileToStorageBlock.Add(file, outputBlock);
+					}
+					if (pair.Value.Any(x => referencedOutputFiles.Contains(x)))
+					{
+						outputNode.AddDirectory(await TempStorage.ArchiveBlockAsync(step.Name, pair.Key, workspaceDir, pair.Value.ToArray(), treeWriter, logger, cancellationToken));
+					}
+				}
 
-				ChunkingOptions chunkingOptions = new ChunkingOptions();
+				// Create all the output tags
+				foreach (string outputName in step.OutputNames)
+				{
+					HashSet<FileReference> files = tagNameToFileSet[outputName];
+
+					HashSet<TempStorageBlockRef> storageBlocks = new HashSet<TempStorageBlockRef>();
+					foreach (FileReference file in files)
+					{
+						TempStorageBlockRef? storageBlock;
+						if (fileToStorageBlock.TryGetValue(file, out storageBlock))
+						{
+							storageBlocks.Add(storageBlock);
+						}
+					}
 
-				DirectoryNode root = new DirectoryNode();
-				await root.CopyFilesAsync(workspaceDir, outputFiles, chunkingOptions, writer, cancellationToken);
+					outputNode.AddFile(await TempStorage.ArchiveTagAsync(step.Name, outputName, workspaceDir, files, storageBlocks.ToArray(), treeWriter, logger, cancellationToken));
+				}
 
 				// Write the final node
-				await writer.WriteAsync(artifact.RefName, root, new RefOptions { Lifetime = TimeSpan.FromDays(3.0) }, cancellationToken);
+				await treeWriter.WriteAsync(refName, outputNode, new RefOptions { Lifetime = TimeSpan.FromDays(3.0) }, cancellationToken);
 				logger.LogInformation("Upload took {Time:n1}s", timer.Elapsed.TotalSeconds);
+
+				// Create the artifact
+				using IRpcClientRef<JobRpc.JobRpcClient> jobRpc = await RpcConnection.GetClientRefAsync<JobRpc.JobRpcClient>(cancellationToken);
+
+				CreateJobArtifactRequest artifactRequest = new CreateJobArtifactRequest();
+				artifactRequest.JobId = _jobId;
+				artifactRequest.StepId = step.StepId;
+				artifactRequest.Name = "Output";
+				artifactRequest.Type = JobArtifactType.Output;
+				artifactRequest.NamespaceId = _namespaceId.ToString();
+				artifactRequest.RefName = refName.ToString();
+
+				CreateJobArtifactResponse artifact = await jobRpc.Client.CreateArtifactAsync(artifactRequest, cancellationToken: cancellationToken);
+				logger.LogInformation("Created output artifact {ArtifactId} with ref {RefName} in namespace {Namespace}", artifact.Id, artifact.RefName, artifact.NamespaceId);
 			}
 
 			return true;
diff --git a/Engine/Source/Programs/Horde/Horde.Agent/Leases/Handlers/JobHandler.cs b/Engine/Source/Programs/Horde/Horde.Agent/Leases/Handlers/JobHandler.cs
--- a/Engine/Source/Programs/Horde/Horde.Agent/Leases/Handlers/JobHandler.cs
+++ b/Engine/Source/Programs/Horde/Horde.Agent/Leases/Handlers/JobHandler.cs
@@ -76,7 +76,7 @@ namespace Horde.Agent.Leases.Handlers
 			// Execute the batch
 			try
 			{
-				JobExecutorOptions options = new JobExecutorOptions(session, _serverStorageFactory, executeTask.JobId, executeTask.BatchId, batch.AgentType, jobOptions);
+				JobExecutorOptions options = new JobExecutorOptions(session, _serverStorageFactory, executeTask.JobId, executeTask.BatchId, batch.AgentType, new NamespaceId(executeTask.NamespaceId), executeTask.Token, jobOptions);
 				await ExecuteBatchAsync(session, leaseId, executeTask.Workspace, executeTask.AutoSdkWorkspace, options, logger, cancellationToken);
 			}
 			catch (Exception ex)
diff --git a/Engine/Source/Programs/Horde/Horde.Agent/Utility/TempStorage.cs b/Engine/Source/Programs/Horde/Horde.Agent/Utility/TempStorage.cs
--- a/Engine/Source/Programs/Horde/Horde.Agent/Utility/TempStorage.cs
+++ b/Engine/Source/Programs/Horde/Horde.Agent/Utility/TempStorage.cs
@@ -523,142 +523,6 @@ namespace Horde.Storage.Utility
 		}
 	}
 
-	/// <summary>
-	/// Node representing a temp storage block
-	/// </summary>
-	[TreeNode("{73AE9604-45E6-473A-95D6-2C9732A0820E}")]
-	class TempStorageBlockNode : TreeNode
-	{
-		public TempStorageBlockManifest Manifest { get; set; }
-		public TreeNodeRef<DirectoryNode> Contents { get; }
-
-		public TempStorageBlockNode(TempStorageBlockManifest manifest, TreeNodeRef<DirectoryNode> contents)
-		{
-			Manifest = manifest;
-			Contents = contents;
-		}
-
-		public TempStorageBlockNode(ITreeNodeReader reader)
-			: base(reader)
-		{
-			TempStorageFile[] files = reader.ReadVariableLengthArray(() => ReadTempStorageFile(reader));
-			Manifest = new TempStorageBlockManifest(files);
-
-			Contents = reader.ReadRef<DirectoryNode>();
-		}
-
-		static TempStorageFile ReadTempStorageFile(ITreeNodeReader reader)
-		{
-			string relativePath = reader.ReadString();
-			long lastWriteTimeUtcTicks = (long)reader.ReadUnsignedVarInt();
-			long length = (long)reader.ReadUnsignedVarInt();
-			string digest = reader.ReadString();
-			return new TempStorageFile(relativePath, lastWriteTimeUtcTicks, length, String.IsNullOrEmpty(digest) ? null : digest);
-		}
-
-		static void WriteTempStorageFile(TempStorageFile file, ITreeNodeWriter writer)
-		{
-			writer.WriteString(file.RelativePath);
-			writer.WriteUnsignedVarInt((ulong)file.LastWriteTimeUtcTicks);
-			writer.WriteUnsignedVarInt((ulong)file.Length);
-			writer.WriteString(file.Digest ?? String.Empty);
-		}
-
-		public override IEnumerable<TreeNodeRef> EnumerateRefs()
-		{
-			yield return Contents;
-		}
-
-		public override void Serialize(ITreeNodeWriter writer)
-		{
-			writer.WriteVariableLengthArray(Manifest.Files, file => WriteTempStorageFile(file, writer));
-			writer.WriteRef(Contents);
-		}
-	}
-
-	/// <summary>
-	/// Node representing a temp storage tag
-	/// </summary>
-	[TreeNode("{CB4DEA03-3FA5-42A1-890D-1BFD7435427C}")]
-	class TempStorageTagNode : TreeNode
-	{
-		public TempStorageTagManifest FileList { get; }
-
-		public TempStorageTagNode(TempStorageTagManifest fileList)
-		{
-			FileList = fileList;
-		}
-
-		public TempStorageTagNode(ITreeNodeReader reader)
-		{
-			string[] localFiles = reader.ReadVariableLengthArray(() => reader.ReadString());
-			string[] externalFiles = reader.ReadVariableLengthArray(() => reader.ReadString());
-			TempStorageBlockRef[] blocks = reader.ReadVariableLengthArray(() => ReadTempStorageBlock(reader));
-			FileList = new TempStorageTagManifest(localFiles, externalFiles, blocks);
-		}
-
-		static TempStorageBlockRef ReadTempStorageBlock(ITreeNodeReader reader)
-		{
-			string nodeName = reader.ReadString();
-			string outputName = reader.ReadString();
-			return new TempStorageBlockRef(nodeName, outputName);
-		}
-
-		static void WriteTempStorageBlock(ITreeNodeWriter writer, TempStorageBlockRef blockRef)
-		{
-			writer.WriteString(blockRef.NodeName);
-			writer.WriteString(blockRef.OutputName);
-		}
-
-		public override void Serialize(ITreeNodeWriter writer)
-		{
-			writer.WriteVariableLengthArray(FileList.LocalFiles, x => writer.WriteString(x));
-			writer.WriteVariableLengthArray(FileList.ExternalFiles, x => writer.WriteString(x));
-			writer.WriteVariableLengthArray(FileList.Blocks, x => WriteTempStorageBlock(writer, x));
-		}
-
-		public override IEnumerable<TreeNodeRef> EnumerateRefs() => Enumerable.Empty<TreeNodeRef>();
-	}
-
-	/// <summary>
-	/// Node representing all the outputs from a node
-	/// </summary>
-	[TreeNode("{A1380FBC-06DB-4CD5-9CF2-ABEB5108E406}")]
-	class TempStorageNode : TreeNode
-	{
-		public Dictionary<string, TreeNodeRef<TempStorageTagNode>> Tags { get; } = new Dictionary<string, TreeNodeRef<TempStorageTagNode>>(StringComparer.OrdinalIgnoreCase);
-		public Dictionary<string, TreeNodeRef<TempStorageBlockNode>> Blocks { get; } = new Dictionary<string, TreeNodeRef<TempStorageBlockNode>>(StringComparer.OrdinalIgnoreCase);
-
-		public TempStorageNode()
-		{
-		}
-
-		public TempStorageNode(ITreeNodeReader reader)
-			: base(reader)
-		{
-			reader.ReadDictionary(Tags, () => reader.ReadString(), () => reader.ReadRef<TempStorageTagNode>());
-			reader.ReadDictionary(Blocks, () => reader.ReadString(), () => reader.ReadRef<TempStorageBlockNode>());
-		}
-
-		public override IEnumerable<TreeNodeRef> EnumerateRefs()
-		{
-			foreach (TreeNodeRef<TempStorageTagNode> tagNode in Tags.Values)
-			{
-				yield return tagNode;
-			}
-			foreach (TreeNodeRef<TempStorageBlockNode> blockNode in Blocks.Values)
-			{
-				yield return blockNode;
-			}
-		}
-
-		public override void Serialize(ITreeNodeWriter writer)
-		{
-			writer.WriteDictionary(Tags, k => writer.WriteString(k), v => writer.WriteRef(v));
-			writer.WriteDictionary(Blocks, k => writer.WriteString(k), v => writer.WriteRef(v));
-		}
-	}
-
 	/// <summary>
 	/// Tracks the state of the current build job using the filesystem, allowing jobs to be restarted after a failure or expanded to include larger targets, and 
 	/// providing a proxy for different machines executing parts of the build in parallel to transfer build products and share state as part of a build system.
@@ -695,116 +559,79 @@ namespace Horde.Storage.Utility
 		/// <returns>The set of files</returns>
 		public static async Task<TempStorageTagManifest> RetrieveTagAsync(TreeReader reader, string refPrefix, string nodeName, string tagName, DirectoryReference manifestDir, ILogger logger, CancellationToken cancellationToken)
 		{
-			TempStorageTagManifest fileList;
-
 			// Try to read the tag set from the local directory
 			FileReference localFileListLocation = GetTagManifestLocation(manifestDir, nodeName, tagName);
-			if(FileReference.Exists(localFileListLocation))
+			if (FileReference.Exists(localFileListLocation))
 			{
 				logger.LogInformation("Reading local file list from {File}", localFileListLocation.FullName);
-				fileList = TempStorageTagManifest.Load(localFileListLocation);
 			}
 			else
 			{
 				RefName refName = GetRefNameForNode(refPrefix, nodeName);
 				logger.LogInformation("Reading node \"{NodeName}\" tag \"{TagName}\" from temp storage (ref: {RefName}, localFile: {LocalFile})", nodeName, tagName, refName, localFileListLocation);
 
-				TempStorageNode node = await reader.ReadNodeAsync<TempStorageNode>(refName, cancellationToken: cancellationToken);
-
-				TreeNodeRef<TempStorageTagNode>? tagNodeRef;
-				if (!node.Tags.TryGetValue(tagName, out tagNodeRef))
-				{
-					throw new TempStorageException($"Missing tag {tagName} from node {nodeName}");
-				}
-
-				TempStorageTagNode tagNode = await tagNodeRef.ExpandAsync(reader, cancellationToken);
-				fileList = tagNode.FileList;
+				DirectoryNode node = await reader.ReadNodeAsync<DirectoryNode>(refName, cancellationToken: cancellationToken);
 
-				// Save the manifest locally
-				DirectoryReference.CreateDirectory(localFileListLocation.Directory);
-				fileList.Save(localFileListLocation);
+				FileEntry fileEntry = node.GetFileEntry(localFileListLocation.GetFileName());
+				await fileEntry.CopyToFileAsync(reader, localFileListLocation.ToFileInfo(), cancellationToken);
 			}
-			return fileList;
+			return TempStorageTagManifest.Load(localFileListLocation);
 		}
 
 		/// <summary>
 		/// Saves a tag to temp storage
 		/// </summary>
-		/// <param name="writer">Writer for output</param>
+		/// <param name="nodeName">Name of the node to store</param>
 		/// <param name="tagName">Name of the output tag</param>
 		/// <param name="rootDir">Root directory for the build products</param>
 		/// <param name="files">Files in the tag</param>
-		/// <param name="blocks">Blocks containing files in the tag</param>
+		/// <param name="blocks"></param>
+		/// <param name="writer">Writer for output</param>
 		/// <param name="logger">Logger for output</param>
 		/// <param name="cancellationToken">Cancellation token for the operation</param>
 		/// <returns></returns>
-		public static async Task<TreeNodeRef<TempStorageTagNode>> ArchiveTagAsync(TreeWriter writer, string tagName, DirectoryReference rootDir, IEnumerable<FileReference> files, IEnumerable<TempStorageBlockRef> blocks, ILogger logger, CancellationToken cancellationToken)
+		public static async Task<FileEntry> ArchiveTagAsync(string nodeName, string tagName, DirectoryReference rootDir, IEnumerable<FileReference> files, TempStorageBlockRef[] blocks, TreeWriter writer, ILogger logger, CancellationToken cancellationToken)
 		{
 			logger.LogInformation("Creating output tag \"{TagName}\"", tagName);
 
 			TempStorageTagManifest fileList = new TempStorageTagManifest(files, rootDir, blocks);
-			TreeNodeRef<TempStorageTagNode> tagNode = new TreeNodeRef<TempStorageTagNode>(new TempStorageTagNode(fileList));
-			await writer.WriteAsync(tagNode, cancellationToken);
 
-			return tagNode;
+			FileReference localFileListLocation = GetTagManifestLocation(rootDir, nodeName, tagName);
+			fileList.Save(localFileListLocation);
+
+			FileNodeWriter fileNodeWriter = new FileNodeWriter(writer, new ChunkingOptions());
+			NodeHandle handle = await fileNodeWriter.CreateAsync(localFileListLocation.ToFileInfo(), cancellationToken);
+
+			return new FileEntry(localFileListLocation.GetFileName(), FileEntryFlags.None, fileNodeWriter.Length, handle);
 		}
 
 		/// <summary>
 		/// Saves the given files (that should be rooted at the branch root) to a shared temp storage manifest with the given temp storage node and game.
 		/// </summary>
-		/// <param name="writer">Writer for output</param>
+		/// <param name="nodeName">Name of the node producing the block</param>
 		/// <param name="blockName">Name of the output block</param>
 		/// <param name="rootDir">Root directory for the build products</param>
-		/// <param name="buildProducts">Array of build products to be archived</param>
+		/// <param name="files">Files to add to the block</param>
+		/// <param name="writer">Writer for output</param>
 		/// <param name="logger">Logger for output</param>
 		/// <param name="cancellationToken">Cancellation token for the operation</param>
 		/// <returns>The created manifest instance (which has already been saved to disk).</returns>
-		public static async Task<TreeNodeRef<TempStorageBlockNode>> ArchiveBlockAsync(TreeWriter writer, string blockName, DirectoryReference rootDir, FileReference[] buildProducts, ILogger logger, CancellationToken cancellationToken)
+		public static async Task<DirectoryEntry> ArchiveBlockAsync(string nodeName, string blockName, DirectoryReference rootDir, IEnumerable<FileReference> files, TreeWriter writer, ILogger logger, CancellationToken cancellationToken)
 		{
 			logger.LogInformation("Creating output block \"{BlockName}\"", blockName);
 
 			// Create a manifest for the given build products
-			FileInfo[] files = buildProducts.Select(x => new FileInfo(x.FullName)).ToArray();
-
-			// Create the directory tree
-			DirectoryNode root = new DirectoryNode(DirectoryFlags.None);
+			FileInfo[] fileInfos = files.Select(x => new FileInfo(x.FullName)).ToArray();
+			TempStorageBlockManifest manifest = new TempStorageBlockManifest(fileInfos, rootDir);
 
-			Dictionary<DirectoryReference, DirectoryNode> dirRefToNode = new Dictionary<DirectoryReference, DirectoryNode>();
-			dirRefToNode.Add(rootDir, root);
-
-			List<(DirectoryNode, FileInfo)> filesToAdd = new List<(DirectoryNode, FileInfo)>();
-			foreach (FileInfo file in files)
-			{
-				DirectoryReference dirRef = new FileReference(file).Directory;
-				DirectoryNode dirNode = FindOrAddDirNode(dirRef, dirRefToNode);
-				filesToAdd.Add((dirNode, file));
-			}
+			FileReference manifestLocation = GetBlockManifestLocation(rootDir, nodeName, blockName);
+			manifest.Save(manifestLocation);
 
-			ChunkingOptions chunkingOptions = new ChunkingOptions();
-			await DirectoryNode.CopyFromDirectoryAsync(filesToAdd, chunkingOptions, writer, cancellationToken);
+			// Create the file tree
+			DirectoryNode rootNode = new DirectoryNode();
+			await rootNode.CopyFilesAsync(rootDir, Enumerable.Concat(files, new[] { manifestLocation }), new ChunkingOptions(), writer, cancellationToken);
 
-			// Create the block node
-			TempStorageBlockManifest manifest = new TempStorageBlockManifest(files, rootDir);
-			TempStorageBlockNode node = new TempStorageBlockNode(manifest, new TreeNodeRef<DirectoryNode>(root));
-			return new TreeNodeRef<TempStorageBlockNode>(node);
-		}
-
-		/// <summary>
-		/// Finds or adds a node for the given directory, using a cached lookup for existing nodes
-		/// </summary>
-		/// <param name="dirRef"></param>
-		/// <param name="dirRefToNode"></param>
-		/// <returns></returns>
-		static DirectoryNode FindOrAddDirNode(DirectoryReference dirRef, Dictionary<DirectoryReference, DirectoryNode> dirRefToNode)
-		{
-			DirectoryNode? dirNode;
-			if (!dirRefToNode.TryGetValue(dirRef, out dirNode))
-			{
-				DirectoryNode parentNode = FindOrAddDirNode(dirRef.ParentDirectory!, dirRefToNode);
-				dirNode = parentNode.AddDirectory(dirRef.GetDirectoryName());
-				dirRefToNode.Add(dirRef, dirNode);
-			}
-			return dirNode;
+			return new DirectoryEntry(blockName, rootNode);
 		}
 
 		/// <summary>
@@ -838,38 +665,27 @@ namespace Horde.Storage.Utility
 				RefName refName = GetRefNameForNode(refPrefix, nodeName);
 				logger.LogInformation("Reading node \"{NodeName}\" block \"{BlockName}\" from temp storage (ref: {RefName}, local: {LocalFile})", nodeName, blockName, refName, localManifestFile);
 
-				TempStorageNode node = await reader.ReadNodeAsync<TempStorageNode>(refName, cancellationToken: cancellationToken);
-				
-				TreeNodeRef<TempStorageBlockNode>? blockNodeRef;
-				if (!node.Blocks.TryGetValue(blockName, out blockNodeRef))
-				{
-					throw new TempStorageException($"Missing block \"{blockName}\" from node \"{nodeName}\"");
-				}
-
-				TempStorageBlockNode blockNode = await blockNodeRef.ExpandAsync(reader, cancellationToken);
-				manifest = blockNode.Manifest;
+				DirectoryNode node = await reader.ReadNodeAsync<DirectoryNode>(refName, cancellationToken: cancellationToken);
 
-				// Delete all the existing files. They may be read-only.
-				foreach (TempStorageFile ManifestFile in manifest.Files)
+				DirectoryEntry? rootDirEntry;
+				if (!node.TryGetDirectoryEntry(blockName, out rootDirEntry))
 				{
-					FileReference File = ManifestFile.ToFileReference(rootDir);
-					FileUtils.ForceDeleteFile(File);
+					throw new TempStorageException($"Missing block \"{blockName}\" from node \"{nodeName}\"");
 				}
 
 				// Add all the files and flush the ref
-				DirectoryNode rootDirNode = await blockNode.Contents.ExpandAsync(reader, cancellationToken);
+				DirectoryNode rootDirNode = await rootDirEntry.ExpandAsync(reader, cancellationToken);
 				await rootDirNode.CopyToDirectoryAsync(reader, rootDir.ToDirectoryInfo(), logger, cancellationToken);
 
+				// Read the manifest in
+				manifest = TempStorageBlockManifest.Load(localManifestFile);
+
 				// Update the timestamps to match the manifest.
 				foreach (TempStorageFile ManifestFile in manifest.Files)
 				{
 					FileReference File = ManifestFile.ToFileReference(rootDir);
 					System.IO.File.SetLastWriteTimeUtc(File.FullName, new DateTime(ManifestFile.LastWriteTimeUtcTicks, DateTimeKind.Utc));
 				}
-
-				// Save the manifest locally
-				DirectoryReference.CreateDirectory(localManifestFile.Directory);
-				manifest.Save(localManifestFile);
 			}
 
 			// Check all the local files are as expected
diff --git a/Engine/Source/Programs/Horde/Horde.Build.Tests/ArtifactTests.cs b/Engine/Source/Programs/Horde/Horde.Build.Tests/ArtifactTests.cs
--- a/Engine/Source/Programs/Horde/Horde.Build.Tests/ArtifactTests.cs
+++ b/Engine/Source/Programs/Horde/Horde.Build.Tests/ArtifactTests.cs
@@ -21,7 +21,7 @@ namespace Horde.Build.Tests
 		public async Task CreateArtifact()
 		{
 			IArtifactCollection artifactCollection = ServiceProvider.GetRequiredService<IArtifactCollection>();
-			IArtifact artifact = await artifactCollection.AddAsync("test", ArtifactType.StepOutput, new string[] { "test1", "test2" }, Namespace.Artifacts, null, AclScopeName.Root);
+			IArtifact artifact = await artifactCollection.AddAsync(ArtifactId.GenerateNewId(), "test", ArtifactType.StepOutput, new string[] { "test1", "test2" }, Namespace.Artifacts, null, AclScopeName.Root);
 
 			{
 				List<IArtifact> artifacts = await artifactCollection.FindAsync(keys: new[] { "test1" }).ToListAsync();
diff --git a/Engine/Source/Programs/Horde/Horde.Build/Artifacts/ArtifactCollection.cs b/Engine/Source/Programs/Horde/Horde.Build/Artifacts/ArtifactCollection.cs
--- a/Engine/Source/Programs/Horde/Horde.Build/Artifacts/ArtifactCollection.cs
+++ b/Engine/Source/Programs/Horde/Horde.Build/Artifacts/ArtifactCollection.cs
@@ -54,14 +54,14 @@ namespace Horde.Build.Artifacts
 			{
 			}
 
-			public Artifact(string name, ArtifactType type, IEnumerable<string> keys, NamespaceId namespaceId, string? refPrefix, AclScopeName scopeName)
+			public Artifact(ArtifactId id, string name, ArtifactType type, IEnumerable<string> keys, NamespaceId namespaceId, RefName refName, AclScopeName scopeName)
 			{
-				Id = ArtifactId.GenerateNewId();
+				Id = id;
 				Name = name;
 				Type = type;
 				Keys.AddRange(keys);
 				NamespaceId = namespaceId;
-				RefName = new RefName($"{refPrefix ?? ""}{Type}-{Id}");
+				RefName = refName;
 				AclScope = scopeName;
 			}
 		}
@@ -78,9 +78,9 @@ namespace Horde.Build.Artifacts
 		}
 
 		/// <inheritdoc/>
-		public async Task<IArtifact> AddAsync(string name, ArtifactType type, IEnumerable<string> keys, NamespaceId namespaceId, string? refPrefix, AclScopeName scopeName, CancellationToken cancellationToken)
+		public async Task<IArtifact> AddAsync(ArtifactId id, string name, ArtifactType type, IEnumerable<string> keys, NamespaceId namespaceId, RefName refName, AclScopeName scopeName, CancellationToken cancellationToken)
 		{
-			Artifact artifact = new Artifact(name, type, keys, namespaceId, refPrefix, scopeName);
+			Artifact artifact = new Artifact(id, name, type, keys, namespaceId, refName, scopeName);
 			await _artifacts.InsertOneAsync(artifact, null, cancellationToken);
 			return artifact;
 		}
diff --git a/Engine/Source/Programs/Horde/Horde.Build/Artifacts/IArtifactCollection.cs b/Engine/Source/Programs/Horde/Horde.Build/Artifacts/IArtifactCollection.cs
--- a/Engine/Source/Programs/Horde/Horde.Build/Artifacts/IArtifactCollection.cs
+++ b/Engine/Source/Programs/Horde/Horde.Build/Artifacts/IArtifactCollection.cs
@@ -1,11 +1,11 @@
 // Copyright Epic Games, Inc. All Rights Reserved.
 
+using System;
 using System.Collections.Generic;
 using System.Threading;
 using System.Threading.Tasks;
 using EpicGames.Horde.Storage;
 using Horde.Build.Acls;
-using Horde.Build.Utilities;
 
 namespace Horde.Build.Artifacts
 {
@@ -17,15 +17,16 @@ namespace Horde.Build.Artifacts
 		/// <summary>
 		/// Creates a new artifact
 		/// </summary>
+		/// <param name="artifactId">Unique id of the artifact</param>
 		/// <param name="name">Description for the artifact</param>
 		/// <param name="type">Type identifier for the artifact</param>
 		/// <param name="keys">Keys for the artifact</param>
 		/// <param name="namespaceId">Namespace containing the data</param>
-		/// <param name="refPrefix">Prefix for the ref name. Object id will be appended to this string.</param>
+		/// <param name="refName">Artifact ref name</param>
 		/// <param name="scopeName">Inherited scope used for permissions</param>
 		/// <param name="cancellationToken">Cancellation token for the operation</param>
 		/// <returns>The new log file document</returns>
-		Task<IArtifact> AddAsync(string name, ArtifactType type, IEnumerable<string> keys, NamespaceId namespaceId, string? refPrefix, AclScopeName scopeName, CancellationToken cancellationToken = default);
+		Task<IArtifact> AddAsync(ArtifactId artifactId, string name, ArtifactType type, IEnumerable<string> keys, NamespaceId namespaceId, RefName refName, AclScopeName scopeName, CancellationToken cancellationToken = default);
 
 		/// <summary>
 		/// Finds artifacts with the given keys
diff --git a/Engine/Source/Programs/Horde/Horde.Build/Jobs/JobRpcService.cs b/Engine/Source/Programs/Horde/Horde.Build/Jobs/JobRpcService.cs
--- a/Engine/Source/Programs/Horde/Horde.Build/Jobs/JobRpcService.cs
+++ b/Engine/Source/Programs/Horde/Horde.Build/Jobs/JobRpcService.cs
@@ -62,8 +62,12 @@ namespace Horde.Build.Jobs
 				throw new StructuredRpcException(StatusCode.NotFound, "Couldn't find template {TemplateId} in stream {StreamId}", job.TemplateId, job.StreamId);
 			}
 
-			NamespaceId namespaceId = Namespace.Artifacts;
-			IArtifact artifact = await _artifactCollection.AddAsync(request.Name, type, keys, namespaceId, $"job-{job.Id}/step-{step.Id}/", templateConfig.ScopeName, context.CancellationToken);
+			ArtifactId artifactId = ArtifactId.GenerateNewId();
+
+			NamespaceId namespaceId = String.IsNullOrEmpty(request.NamespaceId)? Namespace.Artifacts : new NamespaceId(request.NamespaceId);
+			RefName refName = new RefName(String.IsNullOrEmpty(request.RefName) ? $"job-{job.Id}/step-{step.Id}/{artifactId}" : request.RefName);
+
+			IArtifact artifact = await _artifactCollection.AddAsync(artifactId, request.Name, type, keys, namespaceId, refName, templateConfig.ScopeName, context.CancellationToken);
 
 			List<AclClaimConfig> claims = new List<AclClaimConfig>();
 			claims.Add(new AclClaimConfig(HordeClaimTypes.WriteNamespace, namespaceId.ToString()));
diff --git a/Engine/Source/Programs/Shared/EpicGames.Horde/Storage/Nodes/DirectoryNode.cs b/Engine/Source/Programs/Shared/EpicGames.Horde/Storage/Nodes/DirectoryNode.cs
--- a/Engine/Source/Programs/Shared/EpicGames.Horde/Storage/Nodes/DirectoryNode.cs
+++ b/Engine/Source/Programs/Shared/EpicGames.Horde/Storage/Nodes/DirectoryNode.cs
@@ -540,21 +540,27 @@ namespace EpicGames.Horde.Storage.Nodes
 		/// <summary>
 		/// Adds a new directory with the given name
 		/// </summary>
-		/// <param name="name">Name of the new directory</param>
-		/// <returns>The new directory object</returns>
-		public DirectoryNode AddDirectory(Utf8String name)
+		/// <param name="entry">Name of the new directory</param>
+		public void AddDirectory(DirectoryEntry entry)
 		{
-			if (TryGetFileEntry(name, out _))
+			if (TryGetFileEntry(entry.Name, out _))
 			{
-				throw new ArgumentException($"A file with the name '{name}' already exists in this directory", nameof(name));
+				throw new ArgumentException($"A file with the name '{entry.Name}' already exists in this directory", nameof(entry));
 			}
 
-			DirectoryNode node = new DirectoryNode(Flags);
-
-			DirectoryEntry entry = new DirectoryEntry(name, node);
-			_nameToDirectoryEntry.Add(name, entry);
+			_nameToDirectoryEntry.Add(entry.Name, entry);
 			MarkAsDirty();
+		}
 
+		/// <summary>
+		/// Adds a new directory with the given name
+		/// </summary>
+		/// <param name="name">Name of the new directory</param>
+		/// <returns>The new directory object</returns>
+		public DirectoryNode AddDirectory(Utf8String name)
+		{
+			DirectoryNode node = new DirectoryNode(Flags);
+			AddDirectory(new DirectoryEntry(name, node));
 			return node;
 		}
 
@@ -657,7 +663,7 @@ namespace EpicGames.Horde.Storage.Nodes
 			dirToNode.Add(baseDir, this);
 
 			List<(DirectoryNode, FileInfo)> groupedFiles = new List<(DirectoryNode, FileInfo)>();
-			foreach (FileReference file in files)
+			foreach (FileReference file in files.OrderBy(x => x))
 			{
 				DirectoryNode? node = FindOrAddDirectory(file.Directory, baseDir, dirToNode);
 				if (node == null)
