Horde: Fix test build error.

diff --git a/Engine/Source/Programs/Horde/Horde.Build.Tests/ArtifactTests.cs b/Engine/Source/Programs/Horde/Horde.Build.Tests/ArtifactTests.cs
--- a/Engine/Source/Programs/Horde/Horde.Build.Tests/ArtifactTests.cs
+++ b/Engine/Source/Programs/Horde/Horde.Build.Tests/ArtifactTests.cs
@@ -1,11 +1,9 @@
 // Copyright Epic Games, Inc. All Rights Reserved.
 
-using System;
 using System.Collections.Generic;
 using System.Linq;
-using System.Text;
 using System.Threading.Tasks;
-using Amazon.Runtime;
+using EpicGames.Horde.Storage;
 using Horde.Build.Acls;
 using Horde.Build.Artifacts;
 using Horde.Build.Storage;
@@ -21,7 +19,7 @@ namespace Horde.Build.Tests
 		public async Task CreateArtifact()
 		{
 			IArtifactCollection artifactCollection = ServiceProvider.GetRequiredService<IArtifactCollection>();
-			IArtifact artifact = await artifactCollection.AddAsync(ArtifactId.GenerateNewId(), "test", ArtifactType.StepOutput, new string[] { "test1", "test2" }, Namespace.Artifacts, null, AclScopeName.Root);
+			IArtifact artifact = await artifactCollection.AddAsync(ArtifactId.GenerateNewId(), "test", ArtifactType.StepOutput, new string[] { "test1", "test2" }, Namespace.Artifacts, new RefName("test"), AclScopeName.Root);
 
 			{
 				List<IArtifact> artifacts = await artifactCollection.FindAsync(keys: new[] { "test1" }).ToListAsync();
