Horde: Fix output directory for temp storage manifests.

diff --git a/Engine/Source/Programs/Horde/Horde.Agent/Execution/JobExecutor.cs b/Engine/Source/Programs/Horde/Horde.Agent/Execution/JobExecutor.cs
--- a/Engine/Source/Programs/Horde/Horde.Agent/Execution/JobExecutor.cs
+++ b/Engine/Source/Programs/Horde/Horde.Agent/Execution/JobExecutor.cs
@@ -954,7 +954,7 @@ namespace Horde.Agent.Execution
 					}
 					if (pair.Value.Any(x => referencedOutputFiles.Contains(x)))
 					{
-						outputNode.AddDirectory(await TempStorage.ArchiveBlockAsync(step.Name, pair.Key, workspaceDir, pair.Value.ToArray(), treeWriter, logger, cancellationToken));
+						outputNode.AddDirectory(await TempStorage.ArchiveBlockAsync(manifestDir, step.Name, pair.Key, workspaceDir, pair.Value.ToArray(), treeWriter, logger, cancellationToken));
 					}
 				}
 
@@ -973,7 +973,7 @@ namespace Horde.Agent.Execution
 						}
 					}
 
-					outputNode.AddFile(await TempStorage.ArchiveTagAsync(step.Name, outputName, workspaceDir, files, storageBlocks.ToArray(), treeWriter, logger, cancellationToken));
+					outputNode.AddFile(await TempStorage.ArchiveTagAsync(manifestDir, step.Name, outputName, workspaceDir, files, storageBlocks.ToArray(), treeWriter, logger, cancellationToken));
 				}
 
 				// Write the final node
diff --git a/Engine/Source/Programs/Horde/Horde.Agent/Utility/TempStorage.cs b/Engine/Source/Programs/Horde/Horde.Agent/Utility/TempStorage.cs
--- a/Engine/Source/Programs/Horde/Horde.Agent/Utility/TempStorage.cs
+++ b/Engine/Source/Programs/Horde/Horde.Agent/Utility/TempStorage.cs
@@ -581,22 +581,23 @@ namespace Horde.Storage.Utility
 		/// <summary>
 		/// Saves a tag to temp storage
 		/// </summary>
+		/// <param name="manifestDir">Directory containing manifests</param>
 		/// <param name="nodeName">Name of the node to store</param>
 		/// <param name="tagName">Name of the output tag</param>
-		/// <param name="rootDir">Root directory for the build products</param>
+		/// <param name="workspaceDir">Root directory for the build products</param>
 		/// <param name="files">Files in the tag</param>
 		/// <param name="blocks"></param>
 		/// <param name="writer">Writer for output</param>
 		/// <param name="logger">Logger for output</param>
 		/// <param name="cancellationToken">Cancellation token for the operation</param>
 		/// <returns></returns>
-		public static async Task<FileEntry> ArchiveTagAsync(string nodeName, string tagName, DirectoryReference rootDir, IEnumerable<FileReference> files, TempStorageBlockRef[] blocks, TreeWriter writer, ILogger logger, CancellationToken cancellationToken)
+		public static async Task<FileEntry> ArchiveTagAsync(DirectoryReference manifestDir, string nodeName, string tagName, DirectoryReference workspaceDir, IEnumerable<FileReference> files, TempStorageBlockRef[] blocks, TreeWriter writer, ILogger logger, CancellationToken cancellationToken)
 		{
 			logger.LogInformation("Creating output tag \"{TagName}\"", tagName);
 
-			TempStorageTagManifest fileList = new TempStorageTagManifest(files, rootDir, blocks);
+			TempStorageTagManifest fileList = new TempStorageTagManifest(files, workspaceDir, blocks);
 
-			FileReference localFileListLocation = GetTagManifestLocation(rootDir, nodeName, tagName);
+			FileReference localFileListLocation = GetTagManifestLocation(manifestDir, nodeName, tagName);
 			fileList.Save(localFileListLocation);
 
 			FileNodeWriter fileNodeWriter = new FileNodeWriter(writer, new ChunkingOptions());
@@ -608,28 +609,29 @@ namespace Horde.Storage.Utility
 		/// <summary>
 		/// Saves the given files (that should be rooted at the branch root) to a shared temp storage manifest with the given temp storage node and game.
 		/// </summary>
+		/// <param name="manifestDir">Directory containing manifests</param>
 		/// <param name="nodeName">Name of the node producing the block</param>
 		/// <param name="blockName">Name of the output block</param>
-		/// <param name="rootDir">Root directory for the build products</param>
+		/// <param name="workspaceDir">Root directory for the build products</param>
 		/// <param name="files">Files to add to the block</param>
 		/// <param name="writer">Writer for output</param>
 		/// <param name="logger">Logger for output</param>
 		/// <param name="cancellationToken">Cancellation token for the operation</param>
 		/// <returns>The created manifest instance (which has already been saved to disk).</returns>
-		public static async Task<DirectoryEntry> ArchiveBlockAsync(string nodeName, string blockName, DirectoryReference rootDir, IEnumerable<FileReference> files, TreeWriter writer, ILogger logger, CancellationToken cancellationToken)
+		public static async Task<DirectoryEntry> ArchiveBlockAsync(DirectoryReference manifestDir, string nodeName, string blockName, DirectoryReference workspaceDir, IEnumerable<FileReference> files, TreeWriter writer, ILogger logger, CancellationToken cancellationToken)
 		{
 			logger.LogInformation("Creating output block \"{BlockName}\"", blockName);
 
 			// Create a manifest for the given build products
 			FileInfo[] fileInfos = files.Select(x => new FileInfo(x.FullName)).ToArray();
-			TempStorageBlockManifest manifest = new TempStorageBlockManifest(fileInfos, rootDir);
+			TempStorageBlockManifest manifest = new TempStorageBlockManifest(fileInfos, workspaceDir);
 
-			FileReference manifestLocation = GetBlockManifestLocation(rootDir, nodeName, blockName);
+			FileReference manifestLocation = GetBlockManifestLocation(manifestDir, nodeName, blockName);
 			manifest.Save(manifestLocation);
 
 			// Create the file tree
 			DirectoryNode rootNode = new DirectoryNode();
-			await rootNode.CopyFilesAsync(rootDir, Enumerable.Concat(files, new[] { manifestLocation }), new ChunkingOptions(), writer, cancellationToken);
+			await rootNode.CopyFilesAsync(workspaceDir, Enumerable.Concat(files, new[] { manifestLocation }), new ChunkingOptions(), writer, cancellationToken);
 
 			return new DirectoryEntry(blockName, rootNode);
 		}
