Horde: Fix 404 errors when uploading data to temp storage namespace.

diff --git a/Engine/Source/Programs/Horde/Horde.Build/Jobs/JobTaskSource.cs b/Engine/Source/Programs/Horde/Horde.Build/Jobs/JobTaskSource.cs
--- a/Engine/Source/Programs/Horde/Horde.Build/Jobs/JobTaskSource.cs
+++ b/Engine/Source/Programs/Horde/Horde.Build/Jobs/JobTaskSource.cs
@@ -842,6 +842,7 @@ namespace Horde.Build.Jobs
 			task.LogId = logId.ToString();
 			task.JobName = leaseName.ToString();
 			task.JobOptions = job.JobOptions;
+			task.NamespaceId = Namespace.Artifacts.ToString();
 			task.Token = await _aclService.IssueBearerTokenAsync(claims, null);
 
 			List<HordeCommon.Rpc.Messages.AgentWorkspace> workspaces = new ();
