Horde: Allow any ref writes with a namespace claim.

diff --git a/Engine/Source/Programs/Horde/Horde.Build/Storage/StorageController.cs b/Engine/Source/Programs/Horde/Horde.Build/Storage/StorageController.cs
--- a/Engine/Source/Programs/Horde/Horde.Build/Storage/StorageController.cs
+++ b/Engine/Source/Programs/Horde/Horde.Build/Storage/StorageController.cs
@@ -339,7 +339,7 @@ namespace Horde.Build.Storage
 			{
 				return NotFound(namespaceId);
 			}
-			if (!namespaceConfig.Authorize(AclAction.WriteRefs, User) && (!User.HasClaim(HordeClaimTypes.WriteNamespace, namespaceConfig.Id.ToString()) || !User.HasClaim(HordeClaimTypes.WriteRef, refName.ToString())))
+			if (!namespaceConfig.Authorize(AclAction.WriteRefs, User) && !User.HasClaim(HordeClaimTypes.WriteNamespace, namespaceConfig.Id.ToString()))
 			{
 				return Forbid(AclAction.WriteRefs, namespaceId);
 			}
