Horde: Allow passing claims that give access to reading or writing only a subpath of a particular namespace.

diff --git a/Engine/Source/Programs/Horde/Horde.Agent.Tests/Execution/WorkspaceExecutorTest.cs b/Engine/Source/Programs/Horde/Horde.Agent.Tests/Execution/WorkspaceExecutorTest.cs
--- a/Engine/Source/Programs/Horde/Horde.Agent.Tests/Execution/WorkspaceExecutorTest.cs
+++ b/Engine/Source/Programs/Horde/Horde.Agent.Tests/Execution/WorkspaceExecutorTest.cs
@@ -49,7 +49,7 @@ public sealed class WorkspaceExecutorTest : IDisposable
 		_workspace.SetFile(1, "main.cpp", "main");
 		_workspace.SetFile(1, "foo/bar/baz.h", "baz");
 
-		JobExecutorOptions executorOptions = new JobExecutorOptions(_session, null!, JobId, "batch1", AgentType, default, null!, new JobOptions());
+		JobExecutorOptions executorOptions = new JobExecutorOptions(_session, null!, JobId, "batch1", AgentType, default, "", null!, new JobOptions());
 		_executor = new (executorOptions, null, _workspace, NullLogger.Instance);
 	}
 
@@ -71,7 +71,7 @@ public sealed class WorkspaceExecutorTest : IDisposable
 	[TestMethod]
 	public async Task RegularAndAutoSdkWorkspace()
 	{
-		JobExecutorOptions executorOptions = new JobExecutorOptions(_session, null!, JobId, "batch1", AgentType, default, null!, new JobOptions());
+		JobExecutorOptions executorOptions = new JobExecutorOptions(_session, null!, JobId, "batch1", AgentType, default, "", null!, new JobOptions());
 		WorkspaceExecutor executor = new (executorOptions, _autoSdkWorkspace, _workspace, NullLogger.Instance);
 
 		await executor.InitializeAsync(_logger, CancellationToken.None);
@@ -84,7 +84,7 @@ public sealed class WorkspaceExecutorTest : IDisposable
 	[TestMethod]
 	public async Task EnvVars()
 	{
-		JobExecutorOptions executorOptions = new JobExecutorOptions(_session, null!, JobId, "batch1", AgentType, default, null!, new JobOptions());
+		JobExecutorOptions executorOptions = new JobExecutorOptions(_session, null!, JobId, "batch1", AgentType, default, "", null!, new JobOptions());
 		WorkspaceExecutor executor = new (executorOptions, _autoSdkWorkspace, _workspace, NullLogger.Instance);
 		await executor.InitializeAsync(_logger, CancellationToken.None);
 
@@ -109,7 +109,7 @@ public sealed class WorkspaceExecutorTest : IDisposable
 		_server.AddJob("jobPreflight", StreamId, 1, 1000);		
 		_workspace.SetFile(1000, "New/Feature/Foo.cs", "foo");
 
-		JobExecutorOptions executorOptions = new JobExecutorOptions(_session, null!, "jobPreflight", "batch1", AgentType, default, null!, new JobOptions());
+		JobExecutorOptions executorOptions = new JobExecutorOptions(_session, null!, "jobPreflight", "batch1", AgentType, default, "", null!, new JobOptions());
 		WorkspaceExecutor executor = new (executorOptions, null, _workspace, NullLogger.Instance);
 		
 		await executor.InitializeAsync(_logger, CancellationToken.None);
@@ -123,7 +123,7 @@ public sealed class WorkspaceExecutorTest : IDisposable
 	public async Task JobWithNoChange()
 	{
 		_server.AddJob("jobNoChange", StreamId, 0, 0);		
-		JobExecutorOptions executorOptions = new JobExecutorOptions(_session, null!, "jobNoChange", "batch1", AgentType, default, null!, new JobOptions());
+		JobExecutorOptions executorOptions = new JobExecutorOptions(_session, null!, "jobNoChange", "batch1", AgentType, default, "", null!, new JobOptions());
 		WorkspaceExecutor executor = new (executorOptions, null, _workspace, NullLogger.Instance);
 		await Assert.ThrowsExceptionAsync<WorkspaceMaterializationException>(() => executor.InitializeAsync(_logger, CancellationToken.None));
 	}
diff --git a/Engine/Source/Programs/Horde/Horde.Agent/Execution/JobExecutor.cs b/Engine/Source/Programs/Horde/Horde.Agent/Execution/JobExecutor.cs
--- a/Engine/Source/Programs/Horde/Horde.Agent/Execution/JobExecutor.cs
+++ b/Engine/Source/Programs/Horde/Horde.Agent/Execution/JobExecutor.cs
@@ -51,10 +51,11 @@ namespace Horde.Agent.Execution
 		public string BatchId { get; }
 		public string AgentType { get; }
 		public NamespaceId NamespaceId { get; }
+		public string StoragePrefix { get; }
 		public string Token { get; }
 		public JobOptions JobOptions { get; }
 
-		public JobExecutorOptions(ISession session, IServerStorageFactory storageFactory, string jobId, string batchId, string agentType, NamespaceId namespaceId, string token, JobOptions jobOptions)
+		public JobExecutorOptions(ISession session, IServerStorageFactory storageFactory, string jobId, string batchId, string agentType, NamespaceId namespaceId, string storagePrefix, string token, JobOptions jobOptions)
 		{
 			Session = session;
 			StorageFactory = storageFactory;
@@ -62,6 +63,7 @@ namespace Horde.Agent.Execution
 			BatchId = batchId;
 			AgentType = agentType;
 			NamespaceId = namespaceId;
+			StoragePrefix = storagePrefix;
 			Token = token;
 			JobOptions = jobOptions;
 		}
@@ -198,6 +200,7 @@ namespace Horde.Agent.Execution
 		protected readonly ISession _session;
 		protected readonly IServerStorageFactory _storageFactory;
 		private readonly NamespaceId _namespaceId;
+		private readonly string _storagePrefix;
 		private readonly string _token;
 		protected readonly JobOptions _jobOptions;
 
@@ -218,6 +221,7 @@ namespace Horde.Agent.Execution
 			_agentTypeName = options.AgentType;
 
 			_namespaceId = options.NamespaceId;
+			_storagePrefix = options.StoragePrefix;
 			_token = options.Token;
 			_jobOptions = options.JobOptions;
 
@@ -796,7 +800,7 @@ namespace Horde.Agent.Execution
 				string nodeName = input.Substring(0, slashIdx);
 				string tagName = input.Substring(slashIdx + 1);
 
-				TempStorageTagManifest fileList = await TempStorage.RetrieveTagAsync(reader, RefPrefix, nodeName, tagName, manifestDir, logger, cancellationToken);
+				TempStorageTagManifest fileList = await TempStorage.RetrieveTagAsync(reader, _storagePrefix, nodeName, tagName, manifestDir, logger, cancellationToken);
 				tagNameToFileSet[tagName] = fileList.ToFileSet(workspaceDir);
 				inputStorageBlocks.UnionWith(fileList.Blocks);
 			}
@@ -809,7 +813,7 @@ namespace Horde.Agent.Execution
 				scope.Span.SetTag("blocks", inputStorageBlocks.Count);
 				foreach (TempStorageBlockRef inputStorageBlock in inputStorageBlocks)
 				{
-					TempStorageBlockManifest manifest = await TempStorage.RetrieveBlockAsync(reader, RefPrefix, inputStorageBlock.NodeName, inputStorageBlock.OutputName, workspaceDir, manifestDir, logger, cancellationToken);
+					TempStorageBlockManifest manifest = await TempStorage.RetrieveBlockAsync(reader, _storagePrefix, inputStorageBlock.NodeName, inputStorageBlock.OutputName, workspaceDir, manifestDir, logger, cancellationToken);
 					inputManifests[inputStorageBlock] = manifest;
 				}
 				scope.Span.SetTag("size", inputManifests.Sum(x => x.Value.GetTotalSize()));
@@ -937,7 +941,7 @@ namespace Horde.Agent.Execution
 			using (GlobalTracer.Instance.BuildSpan("TempStorage").WithTag("resource", "Write").StartActive())
 			{
 				Stopwatch timer = Stopwatch.StartNew();
-				RefName refName = TempStorage.GetRefNameForNode(RefPrefix, step.Name);
+				RefName refName = TempStorage.GetRefNameForNode(_storagePrefix, step.Name);
 
 				TreeOptions treeOptions = new TreeOptions();
 				using TreeWriter treeWriter = new TreeWriter(storage, treeOptions, refName.Text);
@@ -1146,8 +1150,6 @@ namespace Horde.Agent.Execution
 			}
 		}
 		
-		string RefPrefix => $"{_job.StreamId}/{_job.Change}-{_jobId}";
-
 		async Task<int> ExecuteCommandAsync(BeginStepResponse step, DirectoryReference workspaceDir, string fileName, string arguments, ILogger jobLogger, CancellationToken cancellationToken)
 		{
 			// Combine all the supplied environment variables together
diff --git a/Engine/Source/Programs/Horde/Horde.Agent/Leases/Handlers/JobHandler.cs b/Engine/Source/Programs/Horde/Horde.Agent/Leases/Handlers/JobHandler.cs
--- a/Engine/Source/Programs/Horde/Horde.Agent/Leases/Handlers/JobHandler.cs
+++ b/Engine/Source/Programs/Horde/Horde.Agent/Leases/Handlers/JobHandler.cs
@@ -76,7 +76,7 @@ namespace Horde.Agent.Leases.Handlers
 			// Execute the batch
 			try
 			{
-				JobExecutorOptions options = new JobExecutorOptions(session, _serverStorageFactory, executeTask.JobId, executeTask.BatchId, batch.AgentType, new NamespaceId(executeTask.NamespaceId), executeTask.Token, jobOptions);
+				JobExecutorOptions options = new JobExecutorOptions(session, _serverStorageFactory, executeTask.JobId, executeTask.BatchId, batch.AgentType, new NamespaceId(executeTask.NamespaceId), executeTask.StoragePrefix, executeTask.Token, jobOptions);
 				await ExecuteBatchAsync(session, leaseId, executeTask.Workspace, executeTask.AutoSdkWorkspace, options, logger, cancellationToken);
 			}
 			catch (Exception ex)
diff --git a/Engine/Source/Programs/Horde/Horde.Build/Artifacts/ArtifactType.cs b/Engine/Source/Programs/Horde/Horde.Build/Artifacts/ArtifactType.cs
--- a/Engine/Source/Programs/Horde/Horde.Build/Artifacts/ArtifactType.cs
+++ b/Engine/Source/Programs/Horde/Horde.Build/Artifacts/ArtifactType.cs
@@ -14,6 +14,11 @@ namespace Horde.Build.Artifacts
 	[TypeConverter(typeof(StringIdTypeConverter<ArtifactType, ArtifactTypeConverter>))]
 	public record struct ArtifactType(StringId Id)
 	{
+		/// <summary>
+		/// Default artifact type 
+		/// </summary>
+		public static ArtifactType Unknown { get; } = default;
+
 		/// <summary>
 		/// Output from a build step
 		/// </summary>
diff --git a/Engine/Source/Programs/Horde/Horde.Build/Jobs/JobTaskSource.cs b/Engine/Source/Programs/Horde/Horde.Build/Jobs/JobTaskSource.cs
--- a/Engine/Source/Programs/Horde/Horde.Build/Jobs/JobTaskSource.cs
+++ b/Engine/Source/Programs/Horde/Horde.Build/Jobs/JobTaskSource.cs
@@ -833,8 +833,11 @@ namespace Horde.Build.Jobs
 			// Create a bearer token for the job executor
 			List<AclClaimConfig> claims = new List<AclClaimConfig>();
 			claims.Add(new AclClaimConfig(HordeClaimTypes.Lease, leaseId.ToString()));
-			claims.Add(new AclClaimConfig(HordeClaimTypes.WriteNamespace, Namespace.Artifacts.ToString()));
 
+			string storagePrefix = $"{job.StreamId}/{job.Change}-{job.Id}";
+			claims.Add(new AclClaimConfig(HordeClaimTypes.ReadNamespace, $"{Namespace.Artifacts}:{storagePrefix}"));
+			claims.Add(new AclClaimConfig(HordeClaimTypes.WriteNamespace, $"{Namespace.Artifacts}:{storagePrefix}"));
+		
 			// Encode the payload
 			ExecuteJobTask task = new ExecuteJobTask();
 			task.JobId = job.Id.ToString();
@@ -843,6 +846,7 @@ namespace Horde.Build.Jobs
 			task.JobName = leaseName.ToString();
 			task.JobOptions = job.JobOptions;
 			task.NamespaceId = Namespace.Artifacts.ToString();
+			task.StoragePrefix = storagePrefix;
 			task.Token = await _aclService.IssueBearerTokenAsync(claims, null);
 
 			List<HordeCommon.Rpc.Messages.AgentWorkspace> workspaces = new ();
diff --git a/Engine/Source/Programs/Horde/Horde.Build/Storage/StorageController.cs b/Engine/Source/Programs/Horde/Horde.Build/Storage/StorageController.cs
--- a/Engine/Source/Programs/Horde/Horde.Build/Storage/StorageController.cs
+++ b/Engine/Source/Programs/Horde/Horde.Build/Storage/StorageController.cs
@@ -4,6 +4,7 @@ using System;
 using System.Collections.Generic;
 using System.IO;
 using System.Linq;
+using System.Security.Claims;
 using System.Text.RegularExpressions;
 using System.Threading;
 using System.Threading.Tasks;
@@ -176,7 +177,7 @@ namespace Horde.Build.Storage
 			{
 				return NotFound(namespaceId);
 			}
-			if (!namespaceConfig.Authorize(AclAction.WriteBlobs, User) && !User.HasClaim(HordeClaimTypes.WriteNamespace, namespaceId.ToString()))
+			if (!namespaceConfig.Authorize(AclAction.WriteBlobs, User) && !HasPathClaim(User, HordeClaimTypes.WriteNamespace, namespaceId, prefix ?? String.Empty))
 			{
 				return Forbid(AclAction.WriteBlobs, namespaceId);
 			}
@@ -255,7 +256,7 @@ namespace Horde.Build.Storage
 			{
 				return NotFound(namespaceId);
 			}
-			if (!namespaceConfig.Authorize(AclAction.ReadBlobs, User) && !User.HasClaim(HordeClaimTypes.ReadNamespace, namespaceId.ToString()))
+			if (!namespaceConfig.Authorize(AclAction.ReadBlobs, User) && !HasPathClaim(User, HordeClaimTypes.ReadNamespace, namespaceId, locator.Inner.ToString()))
 			{
 				return Forbid(AclAction.ReadBlobs, namespaceId);
 			}
@@ -339,7 +340,7 @@ namespace Horde.Build.Storage
 			{
 				return NotFound(namespaceId);
 			}
-			if (!namespaceConfig.Authorize(AclAction.WriteRefs, User) && !User.HasClaim(HordeClaimTypes.WriteNamespace, namespaceConfig.Id.ToString()))
+			if (!namespaceConfig.Authorize(AclAction.WriteRefs, User) && !HasPathClaim(User, HordeClaimTypes.WriteNamespace, namespaceId, refName.ToString()))
 			{
 				return Forbid(AclAction.WriteRefs, namespaceId);
 			}
@@ -366,7 +367,7 @@ namespace Horde.Build.Storage
 			{
 				return NotFound(namespaceId);
 			}
-			if (!namespaceConfig.Authorize(AclAction.ReadRefs, User) && !User.HasClaim(HordeClaimTypes.ReadNamespace, namespaceId.ToString()))
+			if (!namespaceConfig.Authorize(AclAction.ReadRefs, User) && !HasPathClaim(User, HordeClaimTypes.ReadNamespace, namespaceId, refName.ToString()))
 			{
 				return Forbid(AclAction.ReadRefs, namespaceId);
 			}
@@ -381,5 +382,44 @@ namespace Horde.Build.Storage
 
 			return new ReadRefResponse(target);
 		}
+
+		/// <summary>
+		/// Checks whether the user has an explicit claim to read or write to a path within a namespace
+		/// </summary>
+		/// <param name="user">User to query</param>
+		/// <param name="claimType">The claim name</param>
+		/// <param name="namespaceId">Namespace id to check for</param>
+		/// <param name="entity">Path to the entity to query</param>
+		/// <returns>True if the user is authorized for access to the given path</returns>
+		static bool HasPathClaim(ClaimsPrincipal user, string claimType, NamespaceId namespaceId, string entity)
+		{
+			foreach (Claim claim in user.Claims)
+			{
+				if (claim.Type.Equals(claimType, StringComparison.Ordinal))
+				{
+					int colonIdx = claim.Value.IndexOf(':', StringComparison.Ordinal);
+					if (colonIdx == -1)
+					{
+						if (namespaceId.Text.Equals(claim.Value))
+						{
+							return true;
+						}
+					}
+					else
+					{
+						if (namespaceId.Text.Equals(claim.Value.AsMemory(0, colonIdx)) && HasPathPrefix(entity, claim.Value.AsSpan(colonIdx + 1)))
+						{
+							return true;
+						}
+					}
+				}
+			}
+			return false;
+		}
+
+		static bool HasPathPrefix(string name, ReadOnlySpan<char> prefix)
+		{
+			return name.Length > prefix.Length && name[prefix.Length] == '/' && name.AsSpan(0, prefix.Length).SequenceEqual(prefix);
+		}
 	}
 }
diff --git a/Engine/Source/Programs/Horde/Horde.Build/Utilities/HordeClaimTypes.cs b/Engine/Source/Programs/Horde/Horde.Build/Utilities/HordeClaimTypes.cs
--- a/Engine/Source/Programs/Horde/Horde.Build/Utilities/HordeClaimTypes.cs
+++ b/Engine/Source/Programs/Horde/Horde.Build/Utilities/HordeClaimTypes.cs
@@ -67,14 +67,14 @@ namespace Horde.Build.Utilities
 		public const string Lease = Prefix + "lease";
 
 		/// <summary>
-		/// Claim allowing a certain namespace to be read from
+		/// Claim allowing a certain namespace and/or path to be read from. See <see cref="WriteNamespace"/>.
 		/// </summary>
-		public const string ReadNamespace = Prefix + "read-ns";
+		public const string ReadNamespace = Prefix + "read-namespace";
 
 		/// <summary>
-		/// Claim allowing a certain namespace to be written to
+		/// Claim allowing a certain namespace and/or path to be written to. Value may be a namespace name, or a subpath (eg. ns:a/b will only allow blobs and refs with an a/b prefix to be written).
 		/// </summary>
-		public const string WriteNamespace = Prefix + "write-ns";
+		public const string WriteNamespace = Prefix + "write-namespace";
 
 		/// <summary>
 		/// Claim allowing a certain ref to be written
diff --git a/Engine/Source/Programs/Shared/EpicGames.Horde/StringId.cs b/Engine/Source/Programs/Shared/EpicGames.Horde/StringId.cs
--- a/Engine/Source/Programs/Shared/EpicGames.Horde/StringId.cs
+++ b/Engine/Source/Programs/Shared/EpicGames.Horde/StringId.cs
@@ -10,7 +10,7 @@ namespace EpicGames.Horde
 	/// <summary>
 	/// Normalized string identifier for a resource
 	/// </summary>
-	public struct StringId : IEquatable<StringId>
+	public struct StringId : IEquatable<StringId>, IEquatable<string>, IEquatable<ReadOnlyMemory<char>>
 	{
 		/// <summary>
 		/// Enum used to disable validation on string arguments
@@ -178,6 +178,27 @@ namespace EpicGames.Horde
 		/// <inheritdoc/>
 		public bool Equals(StringId other) => Text.Equals(other.Text);
 
+		/// <inheritdoc/>
+		public bool Equals(string? other) => other != null && Equals(other.AsMemory());
+	
+		/// <inheritdoc/>
+		public bool Equals(ReadOnlyMemory<char> other)
+		{
+			ReadOnlySpan<char> span = other.Span;
+			if (span.Length != Text.Length)
+			{
+				return false;
+			}
+			for (int idx = 0; idx < Text.Length; idx++)
+			{
+				if (span[idx] != Text[idx])
+				{
+					return false;
+				}
+			}
+			return true;
+		}
+
 		/// <inheritdoc/>
 		public override string ToString() => Text.ToString();
 
