Horde: Support new artifact storage backend for all artifact types.

diff --git a/Engine/Source/Programs/Horde/Horde.Agent/Execution/JobExecutor.cs b/Engine/Source/Programs/Horde/Horde.Agent/Execution/JobExecutor.cs
--- a/Engine/Source/Programs/Horde/Horde.Agent/Execution/JobExecutor.cs
+++ b/Engine/Source/Programs/Horde/Horde.Agent/Execution/JobExecutor.cs
@@ -755,6 +755,41 @@ namespace Horde.Agent.Execution
 			}
 		}
 
+		protected async Task CreateArtifactsAsync(string stepId, string name, JobArtifactType type, DirectoryReference baseDir, IEnumerable<(string, FileReference)> files, ILogger logger, CancellationToken cancellationToken)
+		{
+			if (_jobOptions.UseNewTempStorage ?? false)
+			{
+				await CreateArtifactAsync(stepId, name, type, baseDir, files.Select(x => x.Item2), cancellationToken);
+			}
+			else
+			{
+				await ArtifactUploader.UploadAsync(RpcConnection, _jobId, _batchId, stepId, files, logger, CancellationToken.None);
+			}
+		}
+
+		protected async Task CreateArtifactAsync(string stepId, string name, JobArtifactType type, DirectoryReference baseDir, IEnumerable<FileReference> files, CancellationToken cancellationToken)
+		{
+			try
+			{
+				using IRpcClientRef<JobRpc.JobRpcClient> jobRpc = await RpcConnection.GetClientRefAsync<JobRpc.JobRpcClient>(cancellationToken);
+
+				CreateJobArtifactResponse artifact = await jobRpc.Client.CreateArtifactAsync(new CreateJobArtifactRequest { JobId = _jobId, StepId = stepId, Name = name, Type = type }, cancellationToken: cancellationToken);
+				_logger.LogInformation("Created artifact {ArtifactId} with ref {RefName} in ns {Namespace}", artifact.Id, artifact.RefName, artifact.NamespaceId);
+
+				IStorageClient storage = _storageFactory.CreateStorageClient(_session, new NamespaceId(artifact.NamespaceId), artifact.Token);
+				using TreeWriter writer = new TreeWriter(storage, new RefName(artifact.RefName));
+
+				DirectoryNode dir = new DirectoryNode();
+				await dir.CopyFilesAsync(baseDir, files, new ChunkingOptions(), writer, cancellationToken);
+
+				await writer.WriteAsync(new RefName(artifact.RefName), dir, cancellationToken: cancellationToken);
+			}
+			catch (Exception ex)
+			{
+				_logger.LogInformation(ex, "Error creating artifact '{Name}'", name);
+			}
+		}
+
 		protected async Task TestArtifactAsync(string stepId, CancellationToken cancellationToken)
 		{
 			try
@@ -1244,6 +1279,8 @@ namespace Horde.Agent.Execution
 
 			if (DirectoryReference.Exists(telemetryDir))
 			{
+				List<(string, FileReference)> telemetryFiles = new List<(string, FileReference)>();
+
 				List<TraceEventList> telemetryList = new List<TraceEventList>();
 				foreach (FileReference telemetryFile in DirectoryReference.EnumerateFiles(telemetryDir, "*.json"))
 				{
@@ -1261,8 +1298,7 @@ namespace Horde.Agent.Execution
 						telemetryList.Add(telemetry);
 					}
 
-					await ArtifactUploader.UploadAsync(RpcConnection, _jobId, _batchId, step.StepId, $"Telemetry/{telemetryFile.GetFileName()}", telemetryFile, jobLogger, CancellationToken.None);
-					FileUtils.ForceDeleteFile(telemetryFile);
+					telemetryFiles.Add(($"Telemetry/{telemetryFile.GetFileName()}", telemetryFile));
 				}
 
 				List<TraceEvent> telemetrySpans = new List<TraceEvent>();
@@ -1278,6 +1314,7 @@ namespace Horde.Agent.Execution
 					}
 				}
 
+				FileReference? traceFile = null;
 				if (telemetrySpans.Count > 0)
 				{
 					TraceSpan rootSpan = new TraceSpan();
@@ -1324,25 +1361,34 @@ namespace Horde.Agent.Execution
 					rootSpan.Start = rootSpan.Children!.First().Start;
 					rootSpan.Finish = rootSpan.Children!.Last().Finish;
 
-					FileReference traceFile = FileReference.Combine(telemetryDir, "Trace.json");
+					traceFile = FileReference.Combine(telemetryDir, "Trace.json");
 					using (FileStream stream = FileReference.Open(traceFile, FileMode.Create))
 					{
 						JsonSerializerOptions options = new JsonSerializerOptions { DefaultIgnoreCondition = JsonIgnoreCondition.WhenWritingNull };
 						await JsonSerializer.SerializeAsync(stream, rootSpan, options, cancellationToken);
 					}
-					await ArtifactUploader.UploadAsync(RpcConnection, _jobId, _batchId, step.StepId, "Trace.json", traceFile, jobLogger, CancellationToken.None);
+					telemetryFiles.Add(("Trace.json", traceFile));
 
 					CreateTracingData(GlobalTracer.Instance.ActiveSpan, rootSpan);
 				}
+
+				await CreateArtifactsAsync(step.StepId, "Trace", JobArtifactType.Trace, workspaceDir, telemetryFiles, jobLogger, cancellationToken);
+
+				foreach ((_, FileReference telemetryFile) in telemetryFiles)
+				{
+					FileUtils.ForceDeleteFile(telemetryFile);
+				}
 			}
-					
+
 			if (DirectoryReference.Exists(testDataDir))
 			{
+				List<(string, FileReference)> testDataFiles = new List<(string, FileReference)>();
+
 				Dictionary<string, object> combinedTestData = new Dictionary<string, object>();
 				foreach (FileReference testDataFile in DirectoryReference.EnumerateFiles(testDataDir, "*.json", SearchOption.AllDirectories))
 				{
 					jobLogger.LogInformation("Reading test data {TestDataFile}", testDataFile);
-					await ArtifactUploader.UploadAsync(RpcConnection, _jobId, _batchId, step.StepId, $"TestData/{testDataFile.MakeRelativeTo(testDataDir)}", testDataFile, jobLogger, CancellationToken.None);
+					testDataFiles.Add(($"TestData/{testDataFile.MakeRelativeTo(testDataDir)}", testDataFile));
 
 					TestData testData;
 					using (FileStream stream = FileReference.Open(testDataFile, FileMode.Open))
@@ -1367,32 +1413,21 @@ namespace Horde.Agent.Execution
 
 				jobLogger.LogInformation("Found {NumResults} test results", combinedTestData.Count);
 				await UploadTestDataAsync(step.StepId, combinedTestData);
-			}
 
-			if (_jobOptions.UseNewTempStorage ?? false)
-			{
-				using IRpcClientRef<JobRpc.JobRpcClient> jobRpc = await RpcConnection.GetClientRefAsync<JobRpc.JobRpcClient>(cancellationToken);
-				CreateJobArtifactResponse artifact = await jobRpc.Client.CreateArtifactAsync(new CreateJobArtifactRequest { JobId = _jobId, StepId = step.StepId, Name = "Agent State", Type = JobArtifactType.Saved }, cancellationToken: cancellationToken);
-				jobLogger.LogInformation("Created saved artifact {ArtifactId} with ref {RefName} in ns {Namespace}", artifact.Id, artifact.RefName, artifact.NamespaceId);
-
-				IStorageClient storage = _storageFactory.CreateStorageClient(_session, new NamespaceId(artifact.NamespaceId), artifact.Token);
-
-				TreeOptions treeOptions = new TreeOptions();
-				using TreeWriter treeWriter = new TreeWriter(storage, treeOptions, artifact.RefName);
-
-				DirectoryNode directoryNode = new DirectoryNode(DirectoryFlags.None);
-
-				ChunkingOptions options = new ChunkingOptions();
-				await directoryNode.CopyFromDirectoryAsync(logDir.ToDirectoryInfo(), options, treeWriter, cancellationToken);
-
-				await treeWriter.WriteAsync(new RefName(artifact.RefName), directoryNode, cancellationToken: cancellationToken);
+				await CreateArtifactsAsync(step.StepId, "TestData", JobArtifactType.TestData, workspaceDir, testDataFiles, jobLogger, cancellationToken);
 			}
-			else
+
+			if (DirectoryReference.Exists(logDir))
 			{
-				if (DirectoryReference.Exists(logDir))
+				List<FileReference> artifactFiles = DirectoryReference.EnumerateFiles(logDir, "*", SearchOption.AllDirectories).ToList();
+				if (_jobOptions.UseNewTempStorage ?? false)
+				{
+					await CreateArtifactAsync(step.StepId, "Saved", JobArtifactType.Saved, workspaceDir, artifactFiles, cancellationToken);
+				}
+				else
 				{
 					Dictionary<FileReference, string> artifactFileToId = new Dictionary<FileReference, string>();
-					foreach (FileReference artifactFile in DirectoryReference.EnumerateFiles(logDir, "*", SearchOption.AllDirectories))
+					foreach (FileReference artifactFile in artifactFiles)
 					{
 						string artifactName = artifactFile.MakeRelativeTo(logDir);
 
diff --git a/Engine/Source/Programs/Horde/Horde.Agent/Utility/ArtifactUploader.cs b/Engine/Source/Programs/Horde/Horde.Agent/Utility/ArtifactUploader.cs
--- a/Engine/Source/Programs/Horde/Horde.Agent/Utility/ArtifactUploader.cs
+++ b/Engine/Source/Programs/Horde/Horde.Agent/Utility/ArtifactUploader.cs
@@ -46,20 +46,39 @@ namespace Horde.Agent.Utility
             }
             return contentType;
         }
-        
-        /// <summary>
-        /// Uploads an artifact (with retries)
-        /// </summary>
-        /// <param name="rpcConnection">The grpc client</param>
-        /// <param name="jobId">Job id</param>
-        /// <param name="batchId">Job batch id</param>
-        /// <param name="stepId">Job step id</param>
-        /// <param name="artifactName">Name of the artifact</param>
-        /// <param name="artifactFile">File to upload</param>
-        /// <param name="logger">Logger interfact</param>
-        /// <param name="cancellationToken">Cancellation token</param>
-        /// <returns></returns>
-        public static async Task<string?> UploadAsync(IRpcConnection rpcConnection, string jobId, string batchId, string stepId, string artifactName, FileReference artifactFile, ILogger logger, CancellationToken cancellationToken)
+
+		/// <summary>
+		/// Uploads an artifact (with retries)
+		/// </summary>
+		/// <param name="rpcConnection">The grpc client</param>
+		/// <param name="jobId">Job id</param>
+		/// <param name="batchId">Job batch id</param>
+		/// <param name="stepId">Job step id</param>
+		/// <param name="artifacts">List of artifacts</param>
+		/// <param name="logger">Logger interfact</param>
+		/// <param name="cancellationToken">Cancellation token</param>
+		/// <returns></returns>
+		public static async Task UploadAsync(IRpcConnection rpcConnection, string jobId, string batchId, string stepId, IEnumerable<(string, FileReference)> artifacts, ILogger logger, CancellationToken cancellationToken)
+		{
+			foreach ((string name, FileReference file) in artifacts)
+			{
+				await UploadAsync(rpcConnection, jobId, batchId, stepId, name, file, logger, cancellationToken);
+			}
+		}
+
+		/// <summary>
+		/// Uploads an artifact (with retries)
+		/// </summary>
+		/// <param name="rpcConnection">The grpc client</param>
+		/// <param name="jobId">Job id</param>
+		/// <param name="batchId">Job batch id</param>
+		/// <param name="stepId">Job step id</param>
+		/// <param name="artifactName">Name of the artifact</param>
+		/// <param name="artifactFile">File to upload</param>
+		/// <param name="logger">Logger interfact</param>
+		/// <param name="cancellationToken">Cancellation token</param>
+		/// <returns></returns>
+		public static async Task<string?> UploadAsync(IRpcConnection rpcConnection, string jobId, string batchId, string stepId, string artifactName, FileReference artifactFile, ILogger logger, CancellationToken cancellationToken)
         {
 			try
 			{
diff --git a/Engine/Source/Programs/Horde/Horde.Build/Artifacts/ArtifactType.cs b/Engine/Source/Programs/Horde/Horde.Build/Artifacts/ArtifactType.cs
--- a/Engine/Source/Programs/Horde/Horde.Build/Artifacts/ArtifactType.cs
+++ b/Engine/Source/Programs/Horde/Horde.Build/Artifacts/ArtifactType.cs
@@ -29,6 +29,16 @@ namespace Horde.Build.Artifacts
 		/// </summary>
 		public static ArtifactType StepSaved { get; } = new ArtifactType("step-saved");
 
+		/// <summary>
+		/// Traces from the machine executing a build step
+		/// </summary>
+		public static ArtifactType StepTrace { get; } = new ArtifactType("step-trace");
+
+		/// <summary>
+		/// Test data generated by a machine executing a build step
+		/// </summary>
+		public static ArtifactType StepTestData { get; } = new ArtifactType("step-testdata");
+
 		/// <summary>
 		/// Constructor
 		/// </summary>
diff --git a/Engine/Source/Programs/Horde/Horde.Build/Jobs/JobRpcService.cs b/Engine/Source/Programs/Horde/Horde.Build/Jobs/JobRpcService.cs
--- a/Engine/Source/Programs/Horde/Horde.Build/Jobs/JobRpcService.cs
+++ b/Engine/Source/Programs/Horde/Horde.Build/Jobs/JobRpcService.cs
@@ -50,6 +50,8 @@ namespace Horde.Build.Jobs
 			{
 				JobArtifactType.Output => ArtifactType.StepOutput,
 				JobArtifactType.Saved => ArtifactType.StepSaved,
+				JobArtifactType.Trace => ArtifactType.StepTrace,
+				JobArtifactType.TestData => ArtifactType.StepTestData,
 				_ => throw new StructuredRpcException(StatusCode.InvalidArgument, "Invalid artifact type")
 			};
 
