Horde: Remove name field from artifacts. We haven't needed this in practice.

diff --git a/Engine/Source/Programs/Horde/Horde.Agent/Execution/JobExecutor.cs b/Engine/Source/Programs/Horde/Horde.Agent/Execution/JobExecutor.cs
--- a/Engine/Source/Programs/Horde/Horde.Agent/Execution/JobExecutor.cs
+++ b/Engine/Source/Programs/Horde/Horde.Agent/Execution/JobExecutor.cs
@@ -456,7 +456,6 @@ namespace Horde.Agent.Execution
 
 		protected virtual async Task<bool> SetupAsync(BeginStepResponse step, DirectoryReference workspaceDir, DirectoryReference? sharedStorageDir, ILogger logger, CancellationToken cancellationToken)
 		{
-			await TestArtifactAsync(step.StepId, cancellationToken);
 			FileReference definitionFile = FileReference.Combine(workspaceDir, "Engine", "Saved", "Horde", "Exported.json");
 
 			StringBuilder arguments = new StringBuilder($"BuildGraph");
@@ -755,11 +754,11 @@ namespace Horde.Agent.Execution
 			}
 		}
 
-		protected async Task CreateArtifactsAsync(string stepId, string name, JobArtifactType type, DirectoryReference baseDir, IEnumerable<(string, FileReference)> files, ILogger logger, CancellationToken cancellationToken)
+		protected async Task CreateArtifactsAsync(string stepId, JobArtifactType type, DirectoryReference baseDir, IEnumerable<(string, FileReference)> files, ILogger logger, CancellationToken cancellationToken)
 		{
 			if (_jobOptions.UseNewTempStorage ?? false)
 			{
-				await CreateArtifactAsync(stepId, name, type, baseDir, files.Select(x => x.Item2), cancellationToken);
+				await CreateArtifactAsync(stepId, type, baseDir, files.Select(x => x.Item2), cancellationToken);
 			}
 			else
 			{
@@ -767,13 +766,13 @@ namespace Horde.Agent.Execution
 			}
 		}
 
-		protected async Task CreateArtifactAsync(string stepId, string name, JobArtifactType type, DirectoryReference baseDir, IEnumerable<FileReference> files, CancellationToken cancellationToken)
+		protected async Task CreateArtifactAsync(string stepId, JobArtifactType type, DirectoryReference baseDir, IEnumerable<FileReference> files, CancellationToken cancellationToken)
 		{
 			try
 			{
 				using IRpcClientRef<JobRpc.JobRpcClient> jobRpc = await RpcConnection.GetClientRefAsync<JobRpc.JobRpcClient>(cancellationToken);
 
-				CreateJobArtifactResponse artifact = await jobRpc.Client.CreateArtifactAsync(new CreateJobArtifactRequest { JobId = _jobId, StepId = stepId, Name = name, Type = type }, cancellationToken: cancellationToken);
+				CreateJobArtifactResponse artifact = await jobRpc.Client.CreateArtifactAsync(new CreateJobArtifactRequest { JobId = _jobId, StepId = stepId, Type = type }, cancellationToken: cancellationToken);
 				_logger.LogInformation("Created artifact {ArtifactId} with ref {RefName} in ns {Namespace}", artifact.Id, artifact.RefName, artifact.NamespaceId);
 
 				IStorageClient storage = _storageFactory.CreateStorageClient(_session, new NamespaceId(artifact.NamespaceId), artifact.Token);
@@ -786,27 +785,7 @@ namespace Horde.Agent.Execution
 			}
 			catch (Exception ex)
 			{
-				_logger.LogInformation(ex, "Error creating artifact '{Name}'", name);
-			}
-		}
-
-		protected async Task TestArtifactAsync(string stepId, CancellationToken cancellationToken)
-		{
-			try
-			{
-				using IRpcClientRef<JobRpc.JobRpcClient> jobRpc = await RpcConnection.GetClientRefAsync<JobRpc.JobRpcClient>(cancellationToken);
-				CreateJobArtifactResponse artifact = await jobRpc.Client.CreateArtifactAsync(new CreateJobArtifactRequest { JobId = _jobId, StepId = stepId, Name = "test artifact", Type = JobArtifactType.Saved }, cancellationToken: cancellationToken);
-				_logger.LogInformation("Created artifact {ArtifactId} with ref {RefName} in ns {Namespace}", artifact.Id, artifact.RefName, artifact.NamespaceId);
-
-				IStorageClient storage = _storageFactory.CreateStorageClient(_session, new NamespaceId(artifact.NamespaceId), artifact.Token);
-				using TreeWriter writer = new TreeWriter(storage, new RefName(artifact.RefName));
-
-				DirectoryNode dir = new DirectoryNode();
-				await writer.WriteAsync(new RefName(artifact.RefName), dir, cancellationToken: cancellationToken);
-			}
-			catch (Exception ex)
-			{
-				_logger.LogInformation(ex, "Error creating artifact");
+				_logger.LogInformation(ex, "Error creating artifact '{Type}'", type);
 			}
 		}
 
@@ -1025,7 +1004,6 @@ namespace Horde.Agent.Execution
 				CreateJobArtifactRequest artifactRequest = new CreateJobArtifactRequest();
 				artifactRequest.JobId = _jobId;
 				artifactRequest.StepId = step.StepId;
-				artifactRequest.Name = "Output";
 				artifactRequest.Type = JobArtifactType.Output;
 				artifactRequest.NamespaceId = _namespaceId.ToString();
 				artifactRequest.RefName = refName.ToString();
@@ -1372,7 +1350,7 @@ namespace Horde.Agent.Execution
 					CreateTracingData(GlobalTracer.Instance.ActiveSpan, rootSpan);
 				}
 
-				await CreateArtifactsAsync(step.StepId, "Trace", JobArtifactType.Trace, workspaceDir, telemetryFiles, jobLogger, cancellationToken);
+				await CreateArtifactsAsync(step.StepId, JobArtifactType.Trace, workspaceDir, telemetryFiles, jobLogger, cancellationToken);
 
 				foreach ((_, FileReference telemetryFile) in telemetryFiles)
 				{
@@ -1414,7 +1392,7 @@ namespace Horde.Agent.Execution
 				jobLogger.LogInformation("Found {NumResults} test results", combinedTestData.Count);
 				await UploadTestDataAsync(step.StepId, combinedTestData);
 
-				await CreateArtifactsAsync(step.StepId, "TestData", JobArtifactType.TestData, workspaceDir, testDataFiles, jobLogger, cancellationToken);
+				await CreateArtifactsAsync(step.StepId, JobArtifactType.TestData, workspaceDir, testDataFiles, jobLogger, cancellationToken);
 			}
 
 			if (DirectoryReference.Exists(logDir))
@@ -1422,7 +1400,7 @@ namespace Horde.Agent.Execution
 				List<FileReference> artifactFiles = DirectoryReference.EnumerateFiles(logDir, "*", SearchOption.AllDirectories).ToList();
 				if (_jobOptions.UseNewTempStorage ?? false)
 				{
-					await CreateArtifactAsync(step.StepId, "Saved", JobArtifactType.Saved, workspaceDir, artifactFiles, cancellationToken);
+					await CreateArtifactAsync(step.StepId, JobArtifactType.Saved, workspaceDir, artifactFiles, cancellationToken);
 				}
 				else
 				{
diff --git a/Engine/Source/Programs/Horde/Horde.Build.Tests/ArtifactTests.cs b/Engine/Source/Programs/Horde/Horde.Build.Tests/ArtifactTests.cs
--- a/Engine/Source/Programs/Horde/Horde.Build.Tests/ArtifactTests.cs
+++ b/Engine/Source/Programs/Horde/Horde.Build.Tests/ArtifactTests.cs
@@ -19,7 +19,7 @@ namespace Horde.Build.Tests
 		public async Task CreateArtifact()
 		{
 			IArtifactCollection artifactCollection = ServiceProvider.GetRequiredService<IArtifactCollection>();
-			IArtifact artifact = await artifactCollection.AddAsync(ArtifactId.GenerateNewId(), "test", ArtifactType.StepOutput, new string[] { "test1", "test2" }, Namespace.Artifacts, new RefName("test"), AclScopeName.Root);
+			IArtifact artifact = await artifactCollection.AddAsync(ArtifactId.GenerateNewId(), ArtifactType.StepOutput, new string[] { "test1", "test2" }, Namespace.Artifacts, new RefName("test"), AclScopeName.Root);
 
 			{
 				List<IArtifact> artifacts = await artifactCollection.FindAsync(keys: new[] { "test1" }).ToListAsync();
diff --git a/Engine/Source/Programs/Horde/Horde.Build/Artifacts/ArtifactCollection.cs b/Engine/Source/Programs/Horde/Horde.Build/Artifacts/ArtifactCollection.cs
--- a/Engine/Source/Programs/Horde/Horde.Build/Artifacts/ArtifactCollection.cs
+++ b/Engine/Source/Programs/Horde/Horde.Build/Artifacts/ArtifactCollection.cs
@@ -26,9 +26,6 @@ namespace Horde.Build.Artifacts
 			[BsonRequired, BsonId]
 			public ArtifactId Id { get; set; }
 
-			[BsonElement("nam")]
-			public string Name { get; set; } = String.Empty;
-
 			[BsonElement("typ")]
 			public ArtifactType Type { get; set; }
 
@@ -54,10 +51,9 @@ namespace Horde.Build.Artifacts
 			{
 			}
 
-			public Artifact(ArtifactId id, string name, ArtifactType type, IEnumerable<string> keys, NamespaceId namespaceId, RefName refName, AclScopeName scopeName)
+			public Artifact(ArtifactId id, ArtifactType type, IEnumerable<string> keys, NamespaceId namespaceId, RefName refName, AclScopeName scopeName)
 			{
 				Id = id;
-				Name = name;
 				Type = type;
 				Keys.AddRange(keys);
 				NamespaceId = namespaceId;
@@ -78,9 +74,9 @@ namespace Horde.Build.Artifacts
 		}
 
 		/// <inheritdoc/>
-		public async Task<IArtifact> AddAsync(ArtifactId id, string name, ArtifactType type, IEnumerable<string> keys, NamespaceId namespaceId, RefName refName, AclScopeName scopeName, CancellationToken cancellationToken)
+		public async Task<IArtifact> AddAsync(ArtifactId id, ArtifactType type, IEnumerable<string> keys, NamespaceId namespaceId, RefName refName, AclScopeName scopeName, CancellationToken cancellationToken)
 		{
-			Artifact artifact = new Artifact(id, name, type, keys, namespaceId, refName, scopeName);
+			Artifact artifact = new Artifact(id, type, keys, namespaceId, refName, scopeName);
 			await _artifacts.InsertOneAsync(artifact, null, cancellationToken);
 			return artifact;
 		}
diff --git a/Engine/Source/Programs/Horde/Horde.Build/Artifacts/ArtifactsController.cs b/Engine/Source/Programs/Horde/Horde.Build/Artifacts/ArtifactsController.cs
--- a/Engine/Source/Programs/Horde/Horde.Build/Artifacts/ArtifactsController.cs
+++ b/Engine/Source/Programs/Horde/Horde.Build/Artifacts/ArtifactsController.cs
@@ -29,9 +29,6 @@ namespace Horde.Build.Artifacts
 		/// <inheritdoc cref="IArtifact.Id"/>
 		public string Id => _artifact.Id.ToString();
 
-		/// <inheritdoc cref="IArtifact.Name"/>
-		public string Name => _artifact.Name;
-
 		/// <inheritdoc cref="IArtifact.Type"/>
 		public ArtifactType Type => _artifact.Type;
 
diff --git a/Engine/Source/Programs/Horde/Horde.Build/Artifacts/IArtifact.cs b/Engine/Source/Programs/Horde/Horde.Build/Artifacts/IArtifact.cs
--- a/Engine/Source/Programs/Horde/Horde.Build/Artifacts/IArtifact.cs
+++ b/Engine/Source/Programs/Horde/Horde.Build/Artifacts/IArtifact.cs
@@ -16,11 +16,6 @@ namespace Horde.Build.Artifacts
 		/// </summary>
 		public ArtifactId Id { get; }
 
-		/// <summary>
-		/// Name for this artifact object.
-		/// </summary>
-		public string Name { get; }
-
 		/// <summary>
 		/// Type of artifact
 		/// </summary>
diff --git a/Engine/Source/Programs/Horde/Horde.Build/Artifacts/IArtifactCollection.cs b/Engine/Source/Programs/Horde/Horde.Build/Artifacts/IArtifactCollection.cs
--- a/Engine/Source/Programs/Horde/Horde.Build/Artifacts/IArtifactCollection.cs
+++ b/Engine/Source/Programs/Horde/Horde.Build/Artifacts/IArtifactCollection.cs
@@ -18,7 +18,6 @@ namespace Horde.Build.Artifacts
 		/// Creates a new artifact
 		/// </summary>
 		/// <param name="artifactId">Unique id of the artifact</param>
-		/// <param name="name">Description for the artifact</param>
 		/// <param name="type">Type identifier for the artifact</param>
 		/// <param name="keys">Keys for the artifact</param>
 		/// <param name="namespaceId">Namespace containing the data</param>
@@ -26,7 +25,7 @@ namespace Horde.Build.Artifacts
 		/// <param name="scopeName">Inherited scope used for permissions</param>
 		/// <param name="cancellationToken">Cancellation token for the operation</param>
 		/// <returns>The new log file document</returns>
-		Task<IArtifact> AddAsync(ArtifactId artifactId, string name, ArtifactType type, IEnumerable<string> keys, NamespaceId namespaceId, RefName refName, AclScopeName scopeName, CancellationToken cancellationToken = default);
+		Task<IArtifact> AddAsync(ArtifactId artifactId, ArtifactType type, IEnumerable<string> keys, NamespaceId namespaceId, RefName refName, AclScopeName scopeName, CancellationToken cancellationToken = default);
 
 		/// <summary>
 		/// Finds artifacts with the given keys
diff --git a/Engine/Source/Programs/Horde/Horde.Build/Jobs/JobRpcService.cs b/Engine/Source/Programs/Horde/Horde.Build/Jobs/JobRpcService.cs
--- a/Engine/Source/Programs/Horde/Horde.Build/Jobs/JobRpcService.cs
+++ b/Engine/Source/Programs/Horde/Horde.Build/Jobs/JobRpcService.cs
@@ -69,7 +69,7 @@ namespace Horde.Build.Jobs
 			NamespaceId namespaceId = String.IsNullOrEmpty(request.NamespaceId)? Namespace.Artifacts : new NamespaceId(request.NamespaceId);
 			RefName refName = new RefName(String.IsNullOrEmpty(request.RefName) ? $"job-{job.Id}/step-{step.Id}/{artifactId}" : request.RefName);
 
-			IArtifact artifact = await _artifactCollection.AddAsync(artifactId, request.Name, type, keys, namespaceId, refName, templateConfig.ScopeName, context.CancellationToken);
+			IArtifact artifact = await _artifactCollection.AddAsync(artifactId, type, keys, namespaceId, refName, templateConfig.ScopeName, context.CancellationToken);
 
 			List<AclClaimConfig> claims = new List<AclClaimConfig>();
 			claims.Add(new AclClaimConfig(HordeClaimTypes.WriteNamespace, namespaceId.ToString()));
