Horde: Fix exception when extracting temp storage manifests.

diff --git a/Engine/Source/Programs/Horde/Horde.Agent/Utility/TempStorage.cs b/Engine/Source/Programs/Horde/Horde.Agent/Utility/TempStorage.cs
--- a/Engine/Source/Programs/Horde/Horde.Agent/Utility/TempStorage.cs
+++ b/Engine/Source/Programs/Horde/Horde.Agent/Utility/TempStorage.cs
@@ -573,6 +573,7 @@ namespace Horde.Storage.Utility
 				DirectoryNode node = await reader.ReadNodeAsync<DirectoryNode>(refName, cancellationToken: cancellationToken);
 
 				FileEntry fileEntry = node.GetFileEntry(localFileListLocation.GetFileName());
+				DirectoryReference.CreateDirectory(localFileListLocation.Directory);
 				await fileEntry.CopyToFileAsync(reader, localFileListLocation.ToFileInfo(), cancellationToken);
 			}
 			return TempStorageTagManifest.Load(localFileListLocation);
