Horde: Add simple text job status to pinned jobs to better parse overall job
#rnx
#skipci

diff --git a/Engine/Source/Programs/Horde/HordeDashboard/src/components/UserHome.tsx b/Engine/Source/Programs/Horde/HordeDashboard/src/components/UserHome.tsx
--- a/Engine/Source/Programs/Horde/HordeDashboard/src/components/UserHome.tsx
+++ b/Engine/Source/Programs/Horde/HordeDashboard/src/components/UserHome.tsx
@@ -7,7 +7,7 @@ import moment from "moment-timezone";
 import { default as React, useEffect, useState } from 'react';
 import { Link, useLocation } from "react-router-dom";
 import backend, { useBackend } from "../backend";
-import { GetIssueResponse, GetStepResponse, JobData, JobQuery, JobState, JobStepOutcome, LabelData, LabelOutcome, LabelState, ProjectData, StepData, StreamData } from "../backend/Api";
+import { GetIssueResponse, GetJobResponse, GetStepResponse, JobData, JobQuery, JobState, JobStepBatchError, JobStepBatchState, JobStepOutcome, JobStepState, LabelData, LabelOutcome, LabelState, ProjectData, StepData, StreamData } from "../backend/Api";
 import dashboard, { StatusColor } from "../backend/Dashboard";
 import graphCache, { GraphQuery } from '../backend/GraphCache';
 import { useWindowSize } from '../base/utilities/hooks';
@@ -119,7 +119,7 @@ const ProjectsPanel: React.FC = observer(() => {
 class HealthHandler {
 
    constructor() {
-      makeObservable(this);      
+      makeObservable(this);
    }
 
    startPolling() {
@@ -413,7 +413,7 @@ class UserJobsHandler {
       const maxJobs = 100;
 
       try {
-         let filter = "id,streamId,name,change,preflightChange,templateId,templateHash,graphHash,startedByUserInfo,createTime,state,arguments,updateTime,labels,defaultLabel,batches,autoSubmit,autoSubmitChange,preflightDescription";
+         let filter = "id,streamId,name,change,preflightChange,templateId,templateHash,graphHash,startedByUserInfo,createTime,state,arguments,updateTime,labels,defaultLabel,batches,autoSubmit,autoSubmitChange,preflightDescription,abortedByUserInfo";
 
          const query: JobQuery = {
             filter: filter,
@@ -538,6 +538,72 @@ class UserJobsHandler {
 
 const jobHandler = new UserJobsHandler();
 
+function getJobSummaryText(job: GetJobResponse) {
+   const batches = job.batches;
+
+   const cancelled = job.abortedByUserInfo?.id ? "Canceled" : "";
+
+   const numBatches = batches?.length;
+
+   if (cancelled || !batches || !numBatches) {
+      return cancelled;
+   }
+
+   let batchErrors = 0;
+   let batchFinished = 0;
+
+   let stepsCanceled = 0;
+   let stepsSkipped = 0;
+   let stepsWarnings = 0;
+   let stepsFailures = 0;
+
+   let stepsComplete = true;
+
+   batches.forEach(b => {
+
+      if (!!b.finishTime || b.state === JobStepBatchState.Complete) {
+         batchFinished++;
+      }
+      if (b.error && b.error !== JobStepBatchError.None) {
+         batchErrors++;
+      }
+
+      b.steps?.forEach(s => {
+
+         if (s.state === JobStepState.Waiting || s.state === JobStepState.Ready || s.state === JobStepState.Running) {
+            stepsComplete = false;
+         }
+
+         if (s.state === JobStepState.Aborted) {
+            stepsCanceled++;
+         }
+         if (s.state === JobStepState.Skipped) {
+            stepsSkipped++;
+         }
+
+         if (s.outcome === JobStepOutcome.Warnings) {
+            stepsWarnings++;
+         }
+
+         if (s.outcome === JobStepOutcome.Failure) {
+            stepsFailures++;
+         }
+      })
+   });
+
+   let state = (stepsComplete || batchFinished === batches.length) ? "Completed" : "Running";
+   if (batchErrors || stepsFailures) {
+      state += " with errors";
+   } else if (stepsSkipped || stepsCanceled) {
+      state += " with skipped steps";
+   } else if (stepsWarnings) {
+      state += " with warnings";
+   }
+
+   return state;
+
+}
+
 const JobsPanel: React.FC<{ includeOtherPreflights: boolean }> = observer(({ includeOtherPreflights }) => {
 
    const { projectStore } = useBackend();
@@ -775,16 +841,20 @@ const JobsPanel: React.FC<{ includeOtherPreflights: boolean }> = observer(({ inc
       }
 
       if (column!.name === "Change") {
+
+         const summary = getJobSummaryText(item.job);
+
          return <Stack verticalAlign="center" verticalFill={true} tokens={{ childrenGap: 4 }}>
             <Stack horizontalAlign="start" style={{ whiteSpace: "normal" }}><Text style={{ fontSize: 13, fontFamily: "Horde Open Sans SemiBold", whiteSpace: "pre" }}>{item.job.name}</Text></Stack>
             {!!item.stream && <Stack horizontalAlign="start" style={{ whiteSpace: "normal" }}><Text style={{ fontSize: 11, fontFamily: "Horde Open Sans SemiBold", whiteSpace: "pre" }}>{`//${item.stream.project!.name}/${item.stream.name}`}</Text></Stack>}
             <Stack><ChangeButton job={job} /></Stack>
+            {!!summary && <Stack horizontalAlign="start" style={{ whiteSpace: "normal" }}><Text style={{ fontSize: 11}}>{summary}</Text></Stack>}
          </Stack>
 
       }
 
       if (column!.name === "Submit") {
-         
+
          if (!job.autoSubmit) {
             return null;
          }
@@ -823,7 +893,7 @@ const JobsPanel: React.FC<{ includeOtherPreflights: boolean }> = observer(({ inc
          return <Stack horizontalAlign="start" verticalAlign="center" verticalFill={true} tokens={{ childrenGap: 12 }}>
             <Stack horizontal verticalAlign='center' verticalFill={true} tokens={{ childrenGap: 6 }}>
                {!url && <Stack horizontal verticalAlign='center' verticalFill={true} tokens={{ childrenGap: 8 }}>{icon}<Stack><Text style={{ fontSize: 11 }}>{message}</Text></Stack></Stack>}
-               {!!url && <Stack horizontal verticalAlign='center' verticalFill={true} tokens={{ childrenGap: 8 }}>{icon}<Stack style={{textAlign: "left"}}><a href={url} target="_blank" rel="noreferrer" onClick={ev => ev?.stopPropagation()}><Text variant="tiny" style={{whiteSpace: "pre" }}>{message}</Text></a></Stack></Stack>}
+               {!!url && <Stack horizontal verticalAlign='center' verticalFill={true} tokens={{ childrenGap: 8 }}>{icon}<Stack style={{ textAlign: "left" }}><a href={url} target="_blank" rel="noreferrer" onClick={ev => ev?.stopPropagation()}><Text variant="tiny" style={{ whiteSpace: "pre" }}>{message}</Text></a></Stack></Stack>}
             </Stack>
          </Stack>
       }
