Horde: Add error dialog for job 404 due to the stream being removed
#rnx
#skipci

diff --git a/Engine/Source/Programs/Horde/HordeDashboard/src/components/jobDetailsV2/JobDetailViewV2.tsx b/Engine/Source/Programs/Horde/HordeDashboard/src/components/jobDetailsV2/JobDetailViewV2.tsx
--- a/Engine/Source/Programs/Horde/HordeDashboard/src/components/jobDetailsV2/JobDetailViewV2.tsx
+++ b/Engine/Source/Programs/Horde/HordeDashboard/src/components/jobDetailsV2/JobDetailViewV2.tsx
@@ -6,6 +6,7 @@ import { GetJobsTabResponse } from "../../backend/Api";
 import { useWindowSize } from "../../base/utilities/hooks";
 import { hordeClasses, modeColors } from "../../styles/Styles";
 import { BreadcrumbItem, Breadcrumbs } from "../Breadcrumbs";
+import ErrorHandler from "../ErrorHandler";
 import { useQuery } from "../JobDetailCommon";
 import { TopNav } from "../TopNav";
 import { HealthPanel } from "./JobDetailHealthV2";
@@ -48,10 +49,20 @@ const JobBreadCrumbs: React.FC<{ jobDetails: JobDetailsV2 }> = observer(({ jobDe
 
    const jobData = jobDetails.jobData;
    if (!jobData) {
-      return <Breadcrumbs items={[{ text: "Loading Job" }]} title={"Loading Job"} spinner={true} />
+      if (!jobDetails.jobError) {
+         return <Breadcrumbs items={[{ text: "Loading Job" }]} title={"Loading Job"} spinner={true} />
+      }
+      else {
+         ErrorHandler.set({
+            reason: `Error loading job data`,
+            title: `Unable to Load Job`,
+            message: `Job ${jobDetails.jobId} could not be loaded.\n\nPlease check that you are on the network and the job stream exists.`
+         }, true);
+      }
 
-      // return <Stack horizontal horizontalAlign="center" tokens={{childrenGap: 18}}> <Text variant="mediumPlus">Loading Job</Text><Spinner size={SpinnerSize.large} /></Stack>;
+      return null;
    }
+
    const crumbItems: BreadcrumbItem[] = [];
 
    const stream = jobDetails.stream;
@@ -320,7 +331,7 @@ const ScrollRestore: React.FC<{ jobDetails: JobDetailsV2, scrollRef: React.RefOb
    } catch (message) {
       console.error(message);
    }
-   
+
    return null;
 }
 
diff --git a/Engine/Source/Programs/Horde/HordeDashboard/src/components/jobDetailsV2/JobDetailsViewCommon.tsx b/Engine/Source/Programs/Horde/HordeDashboard/src/components/jobDetailsV2/JobDetailsViewCommon.tsx
--- a/Engine/Source/Programs/Horde/HordeDashboard/src/components/jobDetailsV2/JobDetailsViewCommon.tsx
+++ b/Engine/Source/Programs/Horde/HordeDashboard/src/components/jobDetailsV2/JobDetailsViewCommon.tsx
@@ -153,6 +153,7 @@ export class JobDetailsV2 extends PollBase {
    clear() {
       super.stop();
       this.jobId = undefined;
+      this.jobError = undefined;
       this.jobData = undefined;
       this.stream = undefined;
       this.template = undefined;
@@ -604,7 +605,7 @@ export class JobDetailsV2 extends PollBase {
 
    async poll() {
 
-      if (!this.jobId) {
+      if (!this.jobId || this.jobError) {
          return;
       }
 
@@ -618,8 +619,14 @@ export class JobDetailsV2 extends PollBase {
 
       await Promise.all(requests as any).then(r => results = r).catch(reason => {
          console.error(reason);
+         this.jobError = reason;
       });
 
+      if (this.jobError) {
+         this.setRootUpdated();
+         return;
+      }
+
       const lastUpdateTime = this.jobData?.updateTime;
 
       const jobData = this.jobData = results[0] as JobData;
@@ -697,6 +704,7 @@ export class JobDetailsV2 extends PollBase {
    }
 
    jobId?: string;
+   jobError?: string;
 
    overview?: boolean;
 
