Horde: Fix bad padding on main job detail div
#rnx
#skipci

diff --git a/Engine/Source/Programs/Horde/HordeDashboard/src/components/jobDetailsV2/JobDetailViewV2.tsx b/Engine/Source/Programs/Horde/HordeDashboard/src/components/jobDetailsV2/JobDetailViewV2.tsx
--- a/Engine/Source/Programs/Horde/HordeDashboard/src/components/jobDetailsV2/JobDetailViewV2.tsx
+++ b/Engine/Source/Programs/Horde/HordeDashboard/src/components/jobDetailsV2/JobDetailViewV2.tsx
@@ -1,4 +1,4 @@
-import { mergeStyleSets, Pivot, PivotItem, Stack } from "@fluentui/react";
+import { mergeStyleSets, Pivot, PivotItem, Stack, Text } from "@fluentui/react";
 import { observer } from "mobx-react-lite";
 import { useEffect, useRef, useState } from "react";
 import { useLocation, useNavigate, useParams } from "react-router-dom";
@@ -369,7 +369,7 @@ const DetailsView: React.FC<{ jobDetails: JobDetailsV2 }> = ({ jobDetails }) =>
                <Stack horizontal>
                   <Stack /*className={classNames.pointerSuppress} */ style={{ position: "relative", width: "100%", height: 'calc(100vh - 228px)' }}>
                      <div id="hordeContentArea" ref={scrollRef} style={{ overflowX: "hidden", overflowY: "visible" }}>
-                        <Stack horizontal>
+                        <Stack horizontal style={{paddingBottom: 48}}>
                            <Stack key={`${key}_2`} style={{ paddingLeft: centerAlign }} />
                            <Stack style={{ width: rootWidth }}>
                               <DetailsViewOverview jobDetails={details} />
