Horde: Better pinned job status badge
#rnx
#skipci

diff --git a/Engine/Source/Programs/Horde/HordeDashboard/src/components/ChangeButton.tsx b/Engine/Source/Programs/Horde/HordeDashboard/src/components/ChangeButton.tsx
--- a/Engine/Source/Programs/Horde/HordeDashboard/src/components/ChangeButton.tsx
+++ b/Engine/Source/Programs/Horde/HordeDashboard/src/components/ChangeButton.tsx
@@ -3,7 +3,7 @@
 import { Point } from '@fluentui/react';
 import { ContextualMenu, ContextualMenuItemType, IContextualMenuItem, Stack, Text } from '@fluentui/react';
 import React, { MutableRefObject, useState } from 'react';
-import { GetChangeSummaryResponse, GetJobStepRefResponse, JobData } from '../backend/Api';
+import { GetChangeSummaryResponse, GetJobResponse, GetJobStepRefResponse, JobData, JobStepBatchError, JobStepBatchState, JobStepOutcome, JobStepState } from '../backend/Api';
 import dashboard, { StatusColor } from "../backend/Dashboard";
 import { projectStore } from '../backend/ProjectStore';
 
@@ -151,7 +151,81 @@ export const ChangeContextMenu: React.FC<{ target: ChangeContextMenuTarget, onDi
 
 }
 
-export const ChangeButton: React.FC<{ job?: JobData, stepRef?: GetJobStepRefResponse, commit?: GetChangeSummaryResponse, hideAborted?: boolean, rangeCL?: number }> = ({ job, stepRef, commit, hideAborted, rangeCL }) => {
+function getJobSummary(job: GetJobResponse): { text: string, color: string } {
+
+   const colors = dashboard.getStatusColors();
+   let color = colors.get(StatusColor.Skipped)!;
+
+   const batches = job.batches;
+
+   const cancelled = job.abortedByUserInfo?.id ? "Canceled" : "";
+
+   const numBatches = batches?.length;
+
+   if (cancelled || !batches || !numBatches) {
+      return { text: cancelled, color: colors.get(StatusColor.Skipped)! };
+   }
+
+   let batchErrors = 0;
+   let batchFinished = 0;
+
+   let stepsCanceled = 0;
+   let stepsSkipped = 0;
+   let stepsWarnings = 0;
+   let stepsFailures = 0;
+
+   let stepsComplete = true;
+
+   batches.forEach(b => {
+
+      if (!!b.finishTime || b.state === JobStepBatchState.Complete) {
+         batchFinished++;
+      }
+      if (b.error && b.error !== JobStepBatchError.None) {
+         batchErrors++;
+      }
+
+      b.steps?.forEach(s => {
+
+         if (s.state === JobStepState.Waiting || s.state === JobStepState.Ready || s.state === JobStepState.Running) {
+            stepsComplete = false;
+         }
+
+         if (s.state === JobStepState.Aborted) {
+            stepsCanceled++;
+         }
+         if (s.state === JobStepState.Skipped) {
+            stepsSkipped++;
+         }
+
+         if (s.outcome === JobStepOutcome.Warnings) {
+            stepsWarnings++;
+         }
+
+         if (s.outcome === JobStepOutcome.Failure) {
+            stepsFailures++;
+         }
+      })
+   });
+
+   let state = (stepsComplete || batchFinished === batches.length) ? "Completed" : "Running";
+   color = state === "Running" ? colors.get(StatusColor.Running)! : colors.get(StatusColor.Success)!
+
+   if (batchErrors || stepsFailures) {
+      color = colors.get(StatusColor.Failure)!
+   } else if (stepsSkipped || stepsCanceled) {
+      state += " with skipped steps";
+      color = colors.get(StatusColor.Skipped)!
+   } else if (stepsWarnings) {
+      color = colors.get(StatusColor.Warnings)!
+   }
+
+   return { text: state, color: color };
+
+}
+
+
+export const ChangeButton: React.FC<{ job?: JobData, stepRef?: GetJobStepRefResponse, commit?: GetChangeSummaryResponse, hideAborted?: boolean, rangeCL?: number, pinned?: boolean }> = ({ job, stepRef, commit, hideAborted, rangeCL, pinned }) => {
 
    const [menuShown, setMenuShown] = useState(false);
 
@@ -161,10 +235,9 @@ export const ChangeButton: React.FC<{ job?: JobData, stepRef?: GetJobStepRefResp
       return null;
    }
 
-   if (!job.abortedByUserInfo) {
-      hideAborted = true;
-   }
+   let showStatus = pinned || (!hideAborted && job.abortedByUserInfo?.id);
 
+   const { text, color } = getJobSummary(job);
 
    let change = job?.change?.toString() ?? "Latest";
    if (job?.preflightChange) {
@@ -174,26 +247,13 @@ export const ChangeButton: React.FC<{ job?: JobData, stepRef?: GetJobStepRefResp
    if (stepRef) {
       change = stepRef.change.toString();
    }
-
-   const colors = dashboard.getStatusColors();
-
+   
    return (<Stack verticalAlign="center" verticalFill={true} horizontalAlign="start"> <div style={{ paddingBottom: "1px" }}>
       <Stack tokens={{ childrenGap: 4 }}>
          <span ref={spanRef} style={{ padding: "2px 6px 2px 6px", height: "15px", cursor: "pointer" }} className={job.startedByUserInfo ? "cl-callout-button-user" : "cl-callout-button"} onClick={(ev) => { ev.stopPropagation(); ev.preventDefault(); setMenuShown(!menuShown) }} >{change}</span>
-         {!hideAborted && <span ref={spanRef} style={{ padding: "2px 6px 2px 6px", height: "16px", cursor: "pointer", userSelect: "none", fontFamily: "Horde Open Sans SemiBold", fontSize: "10px", backgroundColor: colors.get(StatusColor.Failure), color: "rgb(255, 255, 255)" }} onClick={(ev) => { ev.preventDefault(); setMenuShown(!menuShown) }}>Canceled</span>}
+         {(!!showStatus) && <span ref={spanRef} style={{ padding: "2px 6px 2px 6px", height: "16px", cursor: "pointer", userSelect: "none", fontFamily: "Horde Open Sans SemiBold", fontSize: "10px", backgroundColor: color, color: "rgb(255, 255, 255)" }} onClick={(ev) => { ev.preventDefault(); setMenuShown(!menuShown) }}>{text}</span>}
       </Stack>
       {menuShown && <ChangeContextMenu target={{ ref: spanRef }} job={job} commit={commit} stepRef={stepRef} rangeCL={rangeCL} onDismiss={() => setMenuShown(false)} />}
    </div></Stack>);
 };
 
-/*
-
-text-decoration: none;
-  white-space: nowrap;
-  background:#0288ee;
-  color: rgb(255, 255, 255);
-  user-select: none;
-  font-family: "Horde Open Sans SemiBold";
-  font-size: 10px;
-  border-width: 0;
-  */
\ No newline at end of file
diff --git a/Engine/Source/Programs/Horde/HordeDashboard/src/components/UserHome.tsx b/Engine/Source/Programs/Horde/HordeDashboard/src/components/UserHome.tsx
--- a/Engine/Source/Programs/Horde/HordeDashboard/src/components/UserHome.tsx
+++ b/Engine/Source/Programs/Horde/HordeDashboard/src/components/UserHome.tsx
@@ -7,7 +7,7 @@ import moment from "moment-timezone";
 import { default as React, useEffect, useState } from 'react';
 import { Link, useLocation } from "react-router-dom";
 import backend, { useBackend } from "../backend";
-import { GetIssueResponse, GetJobResponse, GetStepResponse, JobData, JobQuery, JobState, JobStepBatchError, JobStepBatchState, JobStepOutcome, JobStepState, LabelData, LabelOutcome, LabelState, ProjectData, StepData, StreamData } from "../backend/Api";
+import { GetIssueResponse, GetStepResponse, JobData, JobQuery, JobState, JobStepOutcome, LabelData, LabelOutcome, LabelState, ProjectData, StepData, StreamData } from "../backend/Api";
 import dashboard, { StatusColor } from "../backend/Dashboard";
 import graphCache, { GraphQuery } from '../backend/GraphCache';
 import { useWindowSize } from '../base/utilities/hooks';
@@ -538,72 +538,6 @@ class UserJobsHandler {
 
 const jobHandler = new UserJobsHandler();
 
-function getJobSummaryText(job: GetJobResponse) {
-   const batches = job.batches;
-
-   const cancelled = job.abortedByUserInfo?.id ? "Canceled" : "";
-
-   const numBatches = batches?.length;
-
-   if (cancelled || !batches || !numBatches) {
-      return cancelled;
-   }
-
-   let batchErrors = 0;
-   let batchFinished = 0;
-
-   let stepsCanceled = 0;
-   let stepsSkipped = 0;
-   let stepsWarnings = 0;
-   let stepsFailures = 0;
-
-   let stepsComplete = true;
-
-   batches.forEach(b => {
-
-      if (!!b.finishTime || b.state === JobStepBatchState.Complete) {
-         batchFinished++;
-      }
-      if (b.error && b.error !== JobStepBatchError.None) {
-         batchErrors++;
-      }
-
-      b.steps?.forEach(s => {
-
-         if (s.state === JobStepState.Waiting || s.state === JobStepState.Ready || s.state === JobStepState.Running) {
-            stepsComplete = false;
-         }
-
-         if (s.state === JobStepState.Aborted) {
-            stepsCanceled++;
-         }
-         if (s.state === JobStepState.Skipped) {
-            stepsSkipped++;
-         }
-
-         if (s.outcome === JobStepOutcome.Warnings) {
-            stepsWarnings++;
-         }
-
-         if (s.outcome === JobStepOutcome.Failure) {
-            stepsFailures++;
-         }
-      })
-   });
-
-   let state = (stepsComplete || batchFinished === batches.length) ? "Completed" : "Running";
-   if (batchErrors || stepsFailures) {
-      state += " with errors";
-   } else if (stepsSkipped || stepsCanceled) {
-      state += " with skipped steps";
-   } else if (stepsWarnings) {
-      state += " with warnings";
-   }
-
-   return state;
-
-}
-
 const JobsPanel: React.FC<{ includeOtherPreflights: boolean }> = observer(({ includeOtherPreflights }) => {
 
    const { projectStore } = useBackend();
@@ -842,13 +776,10 @@ const JobsPanel: React.FC<{ includeOtherPreflights: boolean }> = observer(({ inc
 
       if (column!.name === "Change") {
 
-         const summary = getJobSummaryText(item.job);
-
          return <Stack verticalAlign="center" verticalFill={true} tokens={{ childrenGap: 4 }}>
             <Stack horizontalAlign="start" style={{ whiteSpace: "normal" }}><Text style={{ fontSize: 13, fontFamily: "Horde Open Sans SemiBold", whiteSpace: "pre" }}>{item.job.name}</Text></Stack>
             {!!item.stream && <Stack horizontalAlign="start" style={{ whiteSpace: "normal" }}><Text style={{ fontSize: 11, fontFamily: "Horde Open Sans SemiBold", whiteSpace: "pre" }}>{`//${item.stream.project!.name}/${item.stream.name}`}</Text></Stack>}
-            <Stack><ChangeButton job={job} /></Stack>
-            {!!summary && <Stack horizontalAlign="start" style={{ whiteSpace: "normal" }}><Text style={{ fontSize: 11}}>{summary}</Text></Stack>}
+            <Stack><ChangeButton job={job} pinned={true} /></Stack>
          </Stack>
 
       }
