Horde: Register json converter for ObjectId types.

diff --git a/Engine/Source/Programs/Horde/Horde.Build/Startup.cs b/Engine/Source/Programs/Horde/Horde.Build/Startup.cs
--- a/Engine/Source/Programs/Horde/Horde.Build/Startup.cs
+++ b/Engine/Source/Programs/Horde/Horde.Build/Startup.cs
@@ -780,6 +780,7 @@ namespace Horde.Build
 			options.Converters.Add(new JsonStringEnumConverter());
 			options.Converters.Add(new JsonKnownTypesConverterFactory());
 			options.Converters.Add(new StringIdJsonConverterFactory());
+			options.Converters.Add(new ObjectIdJsonConverterFactory());
 			options.Converters.Add(new JsonDateTimeConverter());
 			options.Converters.Add(new JsonTimeSpanConverter());
 		}
