Horde: Fix for line numbers getting out of sync when processing bursty output.

It's possible that a log packet written to the server is truncated to stay below the maximum packet size, which results in some lines being retained in the log writer buffer. Code was previously assuming that subsequent packets would start at the next line index, causing the line count to be offset incorrectly. Since packets are flushed as soon as the threshold size is reached, it would commonly end up with one unused line in the buffer causing subsequent line counts to be offset by one.

diff --git a/Engine/Source/Programs/Horde/Horde.Agent/Utility/JsonRpcLogWriter.cs b/Engine/Source/Programs/Horde/Horde.Agent/Utility/JsonRpcLogWriter.cs
--- a/Engine/Source/Programs/Horde/Horde.Agent/Utility/JsonRpcLogWriter.cs
+++ b/Engine/Source/Programs/Horde/Horde.Agent/Utility/JsonRpcLogWriter.cs
@@ -74,8 +74,8 @@ namespace Horde.Agent.Utility
 		/// <summary>
 		/// Creates a packet from the current data
 		/// </summary>
-		/// <returns>Packet data</returns>
-		public ReadOnlyMemory<byte> CreatePacket()
+		/// <returns>Packet data and number of lines written</returns>
+		public (ReadOnlyMemory<byte>, int) CreatePacket()
 		{
 			_packetWriter.Clear();
 
@@ -93,7 +93,7 @@ namespace Horde.Agent.Utility
 			}
 			_logEvents.RemoveRange(0, eventCount);
 
-			return _packetWriter.WrittenMemory;
+			return (_packetWriter.WrittenMemory, eventCount);
 		}
 
 		/// <summary>
diff --git a/Engine/Source/Programs/Horde/Horde.Agent/Utility/JsonRpcLogger.cs b/Engine/Source/Programs/Horde/Horde.Agent/Utility/JsonRpcLogger.cs
--- a/Engine/Source/Programs/Horde/Horde.Agent/Utility/JsonRpcLogger.cs
+++ b/Engine/Source/Programs/Horde/Horde.Agent/Utility/JsonRpcLogger.cs
@@ -147,8 +147,11 @@ namespace Horde.Agent.Utility
 		async Task RunDataWriter()
 		{
 			// Current position and line number in the log file
-			long offset = 0;
-			int lineIndex = 0;
+			long packetOffset = 0;
+			int packetLineIndex = 0;
+
+			// Index of the next line to write to the log
+			int nextLineIndex = 0;
 
 			// Total number of errors and warnings
 			const int MaxErrors = 50;
@@ -168,9 +171,6 @@ namespace Horde.Agent.Utility
 			{
 				events.Clear();
 
-				// Save off the current line number for sending to the server
-				int initialLineIndex = lineIndex;
-
 				// Get the next data
 				Task waitTask = Task.Delay(TimeSpan.FromSeconds(2.0));
 				while (writer.PacketLength < writer.MaxPacketLength)
@@ -183,14 +183,14 @@ namespace Horde.Agent.Utility
 						{
 							if (jsonLogEvent.Level == LogLevel.Warning && ++numWarnings <= MaxWarnings)
 							{
-								AddEvent(jsonLogEvent.Data.Span, lineIndex, Math.Max(lineCount, jsonLogEvent.LineCount), EventSeverity.Warning, events);
+								AddEvent(jsonLogEvent.Data.Span, nextLineIndex, Math.Max(lineCount, jsonLogEvent.LineCount), EventSeverity.Warning, events);
 							}
 							else if ((jsonLogEvent.Level == LogLevel.Error || jsonLogEvent.Level == LogLevel.Critical) && ++numErrors <= MaxErrors)
 							{
-								AddEvent(jsonLogEvent.Data.Span, lineIndex, Math.Max(lineCount, jsonLogEvent.LineCount), EventSeverity.Error, events);
+								AddEvent(jsonLogEvent.Data.Span, nextLineIndex, Math.Max(lineCount, jsonLogEvent.LineCount), EventSeverity.Error, events);
 							}
 						}
-						lineIndex += lineCount;
+						nextLineIndex += lineCount;
 					}
 					{
 						Task<bool> readTask = _dataChannel.Reader.WaitToReadAsync().AsTask();
@@ -208,16 +208,17 @@ namespace Horde.Agent.Utility
 				// Upload it to the server
 				if (writer.PacketLength > 0)
 				{
-					ByteString packet = UnsafeByteOperations.UnsafeWrap(writer.CreatePacket());
+					(ReadOnlyMemory<byte> packet, int packetLineCount) = writer.CreatePacket();
 					try
 					{
-						await _sink.WriteOutputAsync(new WriteOutputRequest(_logId, offset, initialLineIndex, packet, false), CancellationToken.None);
+						await _sink.WriteOutputAsync(new WriteOutputRequest(_logId, packetOffset, packetLineIndex, UnsafeByteOperations.UnsafeWrap(packet), false), CancellationToken.None);
+						packetOffset += packet.Length;
+						packetLineIndex += packetLineCount;
 					}
 					catch (Exception ex)
 					{
-						_inner.LogWarning(ex, "Unable to write data to server (log {LogId}, offset {Offset}, length {Length}, lines {StartLine}-{EndLine})", _logId, offset, packet.Length, initialLineIndex, lineIndex);
+						_inner.LogWarning(ex, "Unable to write data to server (log {LogId}, offset {Offset}, length {Length}, lines {StartLine}-{EndLine})", _logId, packetOffset, packet.Length, packetLineIndex, packetLineIndex + packetLineCount);
 					}
-					offset += packet.Length;
 				}
 
 				// Write all the events
@@ -248,15 +249,15 @@ namespace Horde.Agent.Utility
 				}
 
 				// Wait for more data to be available
-				if (!await _dataChannel.Reader.WaitToReadAsync())
+				if (writer.PacketLength <= 0 && !await _dataChannel.Reader.WaitToReadAsync())
 				{
 					try
 					{
-						await _sink.WriteOutputAsync(new WriteOutputRequest(_logId, offset, lineIndex, ByteString.Empty, true), CancellationToken.None);
+						await _sink.WriteOutputAsync(new WriteOutputRequest(_logId, packetOffset, packetLineIndex, ByteString.Empty, true), CancellationToken.None);
 					}
 					catch (Exception ex)
 					{
-						_inner.LogWarning(ex, "Unable to flush data to server (log {LogId}, offset {Offset})", _logId, offset);
+						_inner.LogWarning(ex, "Unable to flush data to server (log {LogId}, offset {Offset})", _logId, packetOffset);
 					}
 					break;
 				}
diff --git a/Engine/Source/Programs/Horde/Horde.Build/Logs/LogFileService.cs b/Engine/Source/Programs/Horde/Horde.Build/Logs/LogFileService.cs
--- a/Engine/Source/Programs/Horde/Horde.Build/Logs/LogFileService.cs
+++ b/Engine/Source/Programs/Horde/Horde.Build/Logs/LogFileService.cs
@@ -1691,17 +1691,17 @@ namespace Horde.Build.Logs
 			if (chunkIdx + 1 < logFile.Chunks.Count)
 			{
 				ILogChunk nextChunk = logFile.Chunks[chunkIdx + 1];
-				chunkData = await RepairChunkDataAsync(logFile, chunkIdx, chunkData, (int)(nextChunk.Offset - chunk.Offset), nextChunk.LineIndex - chunk.LineIndex);
+				chunkData = await RepairChunkDataAsync(logFile, chunkIdx, chunkData, (int)(nextChunk.Offset - chunk.Offset), nextChunk.LineIndex - chunk.LineIndex, $"before next");
 			}
 			else
 			{
 				if (logFile.MaxLineIndex != null && chunk.Length != 0)
 				{
-					chunkData = await RepairChunkDataAsync(logFile, chunkIdx, chunkData, chunk.Length, logFile.MaxLineIndex.Value - chunk.LineIndex);
+					chunkData = await RepairChunkDataAsync(logFile, chunkIdx, chunkData, chunk.Length, logFile.MaxLineIndex.Value - chunk.LineIndex, $"last chunk (max line index = {logFile.MaxLineIndex})");
 				}
 				else
 				{
-					chunkData ??= await RepairChunkDataAsync(logFile, chunkIdx, chunkData, 1024, 1);
+					chunkData ??= await RepairChunkDataAsync(logFile, chunkIdx, chunkData, 1024, 1, "default");
 				}
 			}
 
@@ -1716,8 +1716,9 @@ namespace Horde.Build.Logs
 		/// <param name="chunkData">The chunk data that was read</param>
 		/// <param name="length">Expected length of the data</param>
 		/// <param name="lineCount">Expected number of lines in the data</param>
+		/// <param name="context">Context string for diagnostic output</param>
 		/// <returns>Repaired chunk data</returns>
-		async Task<LogChunkData> RepairChunkDataAsync(ILogFile logFile, int chunkIdx, LogChunkData? chunkData, int length, int lineCount)
+		async Task<LogChunkData> RepairChunkDataAsync(ILogFile logFile, int chunkIdx, LogChunkData? chunkData, int length, int lineCount, string context)
 		{
 			int currentLength = 0;
 			int currentLineCount = 0;
@@ -1729,7 +1730,7 @@ namespace Horde.Build.Logs
 
 			if (chunkData == null || currentLength < length || currentLineCount < lineCount)
 			{
-				_logger.LogWarning("Creating placeholder subchunk for log {LogId} chunk {ChunkIdx} (length {Length} vs expected {ExpLength}, lines {LineCount} vs expected {ExpLineCount})", logFile.Id, chunkIdx, currentLength, length, currentLineCount, lineCount);
+				_logger.LogWarning("Creating placeholder subchunk for log {LogId} chunk {ChunkIdx} (length {Length} vs expected {ExpLength}, lines {LineCount} vs expected {ExpLineCount}, context {Context})", logFile.Id, chunkIdx, currentLength, length, currentLineCount, lineCount, context);
 
 				List<LogSubChunkData> subChunks = new List<LogSubChunkData>();
 				if (chunkData != null && chunkData.Length < length && chunkData.LineCount < lineCount)
