Horde: Fix test compile errors.

diff --git a/Engine/Source/Programs/Horde/Horde.Agent.Tests/JsonRpcLoggerTests.cs b/Engine/Source/Programs/Horde/Horde.Agent.Tests/JsonRpcLoggerTests.cs
--- a/Engine/Source/Programs/Horde/Horde.Agent.Tests/JsonRpcLoggerTests.cs
+++ b/Engine/Source/Programs/Horde/Horde.Agent.Tests/JsonRpcLoggerTests.cs
@@ -63,7 +63,7 @@ namespace Horde.Agent.Tests
 			JsonRpcLogWriter writer = new JsonRpcLogWriter();
 			int count = writer.SanitizeAndWriteEvent(jsonLogEvent);
 
-			string[] result = Encoding.UTF8.GetString(writer.CreatePacket().Span).Split('\n', StringSplitOptions.RemoveEmptyEntries);
+			string[] result = Encoding.UTF8.GetString(writer.CreatePacket().Item1.Span).Split('\n', StringSplitOptions.RemoveEmptyEntries);
 			Assert.AreEqual(count, result.Length);
 			return result;
 		}
@@ -220,7 +220,7 @@ namespace Horde.Agent.Tests
 			JsonRpcLogWriter writer = new JsonRpcLogWriter();
 			writer.SanitizeAndWriteEvent(new JsonLogEvent(baseEvent));
 
-			string[] output = Encoding.UTF8.GetString(writer.CreatePacket().Span).Split('\n');
+			string[] output = Encoding.UTF8.GetString(writer.CreatePacket().Item1.Span).Split('\n');
 			string[] expected =
 			{
 				@"{""time"":""2023-01-01T00:00:00"",""level"":""Information"",""message"":""start A0"",""format"":""start {LongLineA$0}"",""properties"":{""LongLineA$0"":""A0""},""line"":0,""lineCount"":14}",
@@ -259,7 +259,7 @@ namespace Horde.Agent.Tests
 
 			for(int idx = 0; idx < 2; idx++)
 			{
-				string output = Encoding.UTF8.GetString(writer.CreatePacket().Span);
+				string output = Encoding.UTF8.GetString(writer.CreatePacket().Item1.Span);
 				string expected = @"{""time"":""2023-01-01T00:00:00"",""level"":""Information"",""message"":""abcdefghijklmnopqrstuvwxyzABCDEFGHIJKLMNOPQRSTUVWXYZ""}" + "\n";
 				Assert.AreEqual(expected, output);
 			}
@@ -275,7 +275,7 @@ namespace Horde.Agent.Tests
 			JsonRpcLogWriter writer = new JsonRpcLogWriter(85);
 			writer.SanitizeAndWriteEvent(new JsonLogEvent(baseEvent));
 
-			string[] output = Encoding.UTF8.GetString(writer.CreatePacket().Span).Split('\n');
+			string[] output = Encoding.UTF8.GetString(writer.CreatePacket().Item1.Span).Split('\n');
 			string[] expected =
 			{
 				@"{""time"":""2023-01-01T00:00:00"",""level"":""Information"",""message"":""abcdefghijklmnopqrstuvwxyzABCDEFGHI [...]""}",
