Horde: Rename job artifact message types to fix swagger doc generation.

diff --git a/Engine/Source/Programs/Horde/Horde.Build/Jobs/Artifacts/ArtifactMessagesV1.cs b/Engine/Source/Programs/Horde/Horde.Build/Jobs/Artifacts/ArtifactMessagesV1.cs
--- a/Engine/Source/Programs/Horde/Horde.Build/Jobs/Artifacts/ArtifactMessagesV1.cs
+++ b/Engine/Source/Programs/Horde/Horde.Build/Jobs/Artifacts/ArtifactMessagesV1.cs
@@ -7,7 +7,7 @@ namespace Horde.Build.Jobs.Artifacts
 	/// <summary>
 	/// Response from creating an artifact
 	/// </summary>
-	public class CreateArtifactResponse
+	public class CreateJobArtifactResponse
 	{
 		/// <summary>
 		/// Unique id for this artifact
@@ -18,7 +18,7 @@ namespace Horde.Build.Jobs.Artifacts
 		/// Constructor
 		/// </summary>
 		/// <param name="id">The artifact file id</param>
-		public CreateArtifactResponse(string id)
+		public CreateJobArtifactResponse(string id)
 		{
 			Id = id;
 		}
@@ -27,7 +27,7 @@ namespace Horde.Build.Jobs.Artifacts
 	/// <summary>
 	/// Request to get a zip of artifacts
 	/// </summary>
-	public class GetArtifactZipRequest
+	public class GetJobArtifactZipRequest
 	{
 		/// <summary>
 		/// Job id to get all artifacts for
@@ -50,7 +50,7 @@ namespace Horde.Build.Jobs.Artifacts
 		/// <param name="jobId">Job id to get all artifacts for</param>
 		/// <param name="stepId">step to filter by</param>
 		/// <param name="artifactIds">The artifact ids.  Returns all artifacts for a job </param>
-		public GetArtifactZipRequest(string? jobId, string? stepId, List<string>? artifactIds)
+		public GetJobArtifactZipRequest(string? jobId, string? stepId, List<string>? artifactIds)
 		{
 			JobId = jobId;
 			StepId = stepId;
@@ -61,7 +61,7 @@ namespace Horde.Build.Jobs.Artifacts
 	/// <summary>
 	/// Response describing an artifact
 	/// </summary>
-	public class GetArtifactResponse
+	public class GetJobArtifactResponse
 	{
 		/// <summary>
 		/// Unique id of the artifact
@@ -103,7 +103,7 @@ namespace Horde.Build.Jobs.Artifacts
 		/// </summary>
 		/// <param name="artifact">The artifact to construct from</param>
 		/// <param name="code">The direct download code</param>
-		public GetArtifactResponse(IArtifactV1 artifact, string? code)
+		public GetJobArtifactResponse(IArtifactV1 artifact, string? code)
 		{
 			Id = artifact.Id.ToString();
 			Name = artifact.Name;
diff --git a/Engine/Source/Programs/Horde/Horde.Build/Jobs/Artifacts/ArtifactsControllerV1.cs b/Engine/Source/Programs/Horde/Horde.Build/Jobs/Artifacts/ArtifactsControllerV1.cs
--- a/Engine/Source/Programs/Horde/Horde.Build/Jobs/Artifacts/ArtifactsControllerV1.cs
+++ b/Engine/Source/Programs/Horde/Horde.Build/Jobs/Artifacts/ArtifactsControllerV1.cs
@@ -58,7 +58,7 @@ namespace Horde.Build.Jobs.Artifacts
 		[HttpPost]
 		[Authorize]
 		[Route("/api/v1/artifacts")]
-		public async Task<ActionResult<CreateArtifactResponse>> CreateArtifact([FromQuery] JobId jobId, [FromQuery]string? stepId, IFormFile file)
+		public async Task<ActionResult<CreateJobArtifactResponse>> CreateArtifact([FromQuery] JobId jobId, [FromQuery]string? stepId, IFormFile file)
 		{
 			IJob? job = await _jobService.GetJobAsync(jobId);
 			if(job == null)
@@ -94,7 +94,7 @@ namespace Horde.Build.Jobs.Artifacts
 			}
 
 			IArtifactV1 newArtifact = await _artifactCollection.CreateArtifactAsync(job.Id, step?.Id, file.FileName, file.ContentType ?? "horde-mime/unknown", file.OpenReadStream());
-			return new CreateArtifactResponse(newArtifact.Id.ToString());
+			return new CreateJobArtifactResponse(newArtifact.Id.ToString());
 		}
 
 		/// <summary>
@@ -106,7 +106,7 @@ namespace Horde.Build.Jobs.Artifacts
 		[HttpPut]
 		[Authorize]
 		[Route("/api/v1/artifacts/{artifactId}")]
-		public async Task<ActionResult<CreateArtifactResponse>> UpdateArtifact(string artifactId, IFormFile file)
+		public async Task<ActionResult<CreateJobArtifactResponse>> UpdateArtifact(string artifactId, IFormFile file)
 		{
 			IArtifactV1? artifact = await _artifactCollection.GetArtifactAsync(artifactId.ToObjectId());
 			if (artifact == null)
@@ -133,7 +133,7 @@ namespace Horde.Build.Jobs.Artifacts
 		[HttpGet]
 		[Authorize]
 		[Route("/api/v1/artifacts")]
-		[ProducesResponseType(typeof(List<GetArtifactResponse>), 200)]
+		[ProducesResponseType(typeof(List<GetJobArtifactResponse>), 200)]
 		public async Task<ActionResult<List<object>>> GetArtifacts([FromQuery] JobId jobId, [FromQuery] string? stepId = null, [FromQuery] bool code = false, [FromQuery] PropertyFilter? filter = null)
 		{
 			if (!await _jobService.AuthorizeAsync(jobId, AclAction.DownloadArtifact, User, _globalConfig.Value))
@@ -144,7 +144,7 @@ namespace Horde.Build.Jobs.Artifacts
 			string? downloadCode = code ? (string?)await GetDirectDownloadCodeForJobAsync(jobId) : null;
 
 			List<IArtifactV1> artifacts = await _artifactCollection.GetArtifactsAsync(jobId, stepId?.ToSubResourceId(), null);
-			return artifacts.ConvertAll(x => new GetArtifactResponse(x, downloadCode).ApplyFilter(filter));
+			return artifacts.ConvertAll(x => new GetJobArtifactResponse(x, downloadCode).ApplyFilter(filter));
 		}
 
 		/// <summary>
@@ -178,7 +178,7 @@ namespace Horde.Build.Jobs.Artifacts
 		[HttpGet]
 		[Authorize]
 		[Route("/api/v1/artifacts/{artifactId}")]
-		[ProducesResponseType(typeof(GetArtifactResponse), 200)]
+		[ProducesResponseType(typeof(GetJobArtifactResponse), 200)]
 		public async Task<ActionResult<object>> GetArtifact(string artifactId, bool code = false, [FromQuery] PropertyFilter? filter = null)
 		{
 			IArtifactV1? artifact = await _artifactCollection.GetArtifactAsync(artifactId.ToObjectId());
@@ -192,7 +192,7 @@ namespace Horde.Build.Jobs.Artifacts
 			}
 
 			string? downloadCode = code? (string?)await GetDirectDownloadCodeForJobAsync(artifact.JobId) : null;
-			return new GetArtifactResponse(artifact, downloadCode).ApplyFilter(filter);
+			return new GetJobArtifactResponse(artifact, downloadCode).ApplyFilter(filter);
 		}
 
 		/// <summary>
@@ -329,7 +329,7 @@ namespace Horde.Build.Jobs.Artifacts
 		[HttpPost]
 		[Authorize]
 		[Route("/api/v1/artifacts/zip")]
-		public async Task<ActionResult> ZipArtifacts(GetArtifactZipRequest artifactZipRequest)
+		public async Task<ActionResult> ZipArtifacts(GetJobArtifactZipRequest artifactZipRequest)
 		{
 			if (artifactZipRequest.JobId == null)
 			{
