Horde: Add alternate platform selection to suite test events panel
#rnx
#skipci

diff --git a/Engine/Source/Programs/Horde/HordeDashboard/src/backend/AutomationTestData.ts b/Engine/Source/Programs/Horde/HordeDashboard/src/backend/AutomationTestData.ts
new file mode 100644
--- /dev/null
+++ b/Engine/Source/Programs/Horde/HordeDashboard/src/backend/AutomationTestData.ts
@@ -0,0 +1,1334 @@
+import { isNumber } from '@datadog/browser-core';
+import { action, makeObservable, observable } from 'mobx';
+import backend from '../backend';
+import { GetTestDataRefResponse, GetTestMetaResponse, GetTestResponse, GetTestStreamResponse, GetTestSuiteResponse } from '../backend/Api';
+import { projectStore } from '../backend/ProjectStore';
+
+export type FilterState = "Success" | "Failed" | "Consecutive Failures" | "Skipped";
+
+// some aliases
+type StreamId = string;
+type TestId = string;
+type SuiteId = string;
+type MetaId = string;
+
+type TestStatus = {
+    suite: boolean;
+    refs: GetTestDataRefResponse[];
+    success?: number;
+    error?: number;
+    skip?: number;
+}
+
+
+const defaultQueryWeeks = 2;
+const minuteInWeek = 10080;
+const queryRefreshMS = 2500;
+
+// query states
+type TestDataState = {
+    streams?: string[];
+    automation?: string;
+    platforms?: string[];
+    configurations?: string[];
+    targets?: string[];
+    rhi?: string[];
+    variation?: string[];
+    tests?: string[];
+    suites?: string[];
+    weeks?: number;
+}
+
+export class TestDataHandler {
+
+    constructor(search: URLSearchParams) {
+        TestDataHandler.instance = this;
+        makeObservable(this);
+        this.state = this.stateFromSearch(search);
+        this.load();
+    }
+
+    setAutomation(automation: string) {
+        this.state.automation = automation;
+
+        if (automation) {
+            const streams = this.getAutomationStreams(automation);
+            this.state.streams = this.state.streams?.filter(s => streams.indexOf(s) !== -1);
+            if (!this.state.streams?.length) {
+                this.state.streams = undefined;
+            }
+
+            //this.state.tests = undefined;
+            //this.state.suites = undefined;
+
+            const tests = this.streamTests;
+            this.state.tests = this.state.tests?.filter(test => !!tests.find(t => t.name === test));
+            if (!this.state.tests?.length) {
+                this.state.tests = undefined;
+            }
+
+            const suites = this.streamSuites;
+            this.state.suites = this.state.suites?.filter(suite => !!suites.find(s => s.name === suite));
+            if (!this.state.suites?.length) {
+                this.state.suites = undefined;
+            }
+
+        }
+
+        if (this.updateSearch()) {
+            this.setUpdated();
+        }
+    }
+
+    setQueryWeeks(weeks: number) {
+        this.state.weeks = weeks;
+        if (this.updateSearch()) {
+            this.setUpdated();
+        }
+    }
+
+    setSuiteRef(suiteRef?: string) {
+        if (this.suiteRef === suiteRef) {
+            return;
+        }
+
+        console.log(`Setting suite ref to ${suiteRef}}`)
+
+        this.suiteRef = suiteRef;
+        this.setUpdated();
+    }
+
+    addPlatform(platform: string) {
+
+        const state = this.state;
+
+        if (state.platforms) {
+            if (state.platforms.indexOf(platform) !== -1) {
+                return;
+            }
+        }
+
+        if (!state.platforms) {
+            state.platforms = [];
+        }
+
+        state.platforms.push(platform);
+
+        if (this.updateSearch()) {
+            this.setUpdated();
+        }
+    }
+
+    removePlatform(platform: string) {
+
+        const state = this.state;
+
+        if (state.platforms?.indexOf(platform) === -1) {
+            return;
+        }
+
+        const idx = state.platforms?.indexOf(platform);
+        if (!isNumber(idx) || idx < 0) {
+            return;
+        }
+
+        state.platforms!.splice(idx, 1);
+
+        if (!state.platforms!.length) {
+            state.platforms = undefined;
+        }
+
+        if (this.updateSearch()) {
+            this.setUpdated();
+        }
+    }
+
+    addConfiguration(config: string) {
+
+        const state = this.state;
+
+        if (state.configurations) {
+            if (state.configurations.indexOf(config) !== -1) {
+                return;
+            }
+        }
+
+        if (!state.configurations) {
+            state.configurations = [];
+        }
+
+        state.configurations.push(config);
+
+        if (this.updateSearch()) {
+            this.setUpdated();
+        }
+    }
+
+    removeConfiguration(config: string) {
+
+        const state = this.state;
+
+        if (state.configurations?.indexOf(config) === -1) {
+            return;
+        }
+
+        const idx = state.configurations?.indexOf(config);
+        if (!isNumber(idx) || idx < 0) {
+            return;
+        }
+
+        state.configurations!.splice(idx, 1);
+
+        if (!state.configurations!.length) {
+            state.configurations = undefined;
+        }
+
+        if (this.updateSearch()) {
+            this.setUpdated();
+        }
+    }
+
+    addTarget(target: string) {
+
+        const state = this.state;
+
+        if (state.targets) {
+            if (state.targets.indexOf(target) !== -1) {
+                return;
+            }
+        }
+
+        if (!state.targets) {
+            state.targets = [];
+        }
+
+        state.targets.push(target);
+
+        if (this.updateSearch()) {
+            this.setUpdated();
+        }
+    }
+
+    removeTarget(target: string) {
+
+        const state = this.state;
+
+        if (state.targets?.indexOf(target) === -1) {
+            return;
+        }
+
+        const idx = state.targets?.indexOf(target);
+        if (!isNumber(idx) || idx < 0) {
+            return;
+        }
+
+        state.targets!.splice(idx, 1);
+
+        if (!state.targets!.length) {
+            state.targets = undefined;
+        }
+
+        if (this.updateSearch()) {
+            this.setUpdated();
+        }
+    }
+
+    addRHI(rhi: string) {
+
+        const state = this.state;
+
+        if (state.rhi) {
+            if (state.rhi.indexOf(rhi) !== -1) {
+                return;
+            }
+        }
+
+        if (!state.rhi) {
+            state.rhi = [];
+        }
+
+        state.rhi.push(rhi);
+
+        if (this.updateSearch()) {
+            this.setUpdated();
+        }
+    }
+
+    removeRHI(rhi: string) {
+
+        const state = this.state;
+
+        if (state.rhi?.indexOf(rhi) === -1) {
+            return;
+        }
+
+        const idx = state.rhi?.indexOf(rhi);
+        if (!isNumber(idx) || idx < 0) {
+            return;
+        }
+
+        state.rhi!.splice(idx, 1);
+
+        if (!state.rhi!.length) {
+            state.rhi = undefined;
+        }
+
+        if (this.updateSearch()) {
+            this.setUpdated();
+        }
+    }
+
+    addVariation(variation: string) {
+
+        const state = this.state;
+
+        if (state.variation) {
+            if (state.variation.indexOf(variation) !== -1) {
+                return;
+            }
+        }
+
+        if (!state.variation) {
+            state.variation = [];
+        }
+
+        state.variation.push(variation);
+
+        if (this.updateSearch()) {
+            this.setUpdated();
+        }
+    }
+
+    removeVariation(variation: string) {
+
+        const state = this.state;
+
+        if (state.variation?.indexOf(variation) === -1) {
+            return;
+        }
+
+        const idx = state.variation?.indexOf(variation);
+        if (!isNumber(idx) || idx < 0) {
+            return;
+        }
+
+        state.variation!.splice(idx, 1);
+
+        if (!state.variation!.length) {
+            state.variation = undefined;
+        }
+
+        if (this.updateSearch()) {
+            this.setUpdated();
+        }
+    }
+
+
+    addTest(test: string) {
+
+        const state = this.state;
+
+        if (state.tests) {
+            if (state.tests.indexOf(test) !== -1) {
+                return;
+            }
+        }
+
+        if (!state.tests) {
+            state.tests = [];
+        }
+
+        state.tests.push(test);
+
+        if (this.updateSearch()) {
+            this.setUpdated();
+        }
+    }
+
+    removeTest(test: string) {
+
+        const state = this.state;
+
+        if (state.tests?.indexOf(test) === -1) {
+            return;
+        }
+
+        const idx = state.tests?.indexOf(test);
+        if (!isNumber(idx) || idx < 0) {
+            return;
+        }
+
+        state.tests!.splice(idx, 1);
+
+        if (!state.tests!.length) {
+            state.tests = undefined;
+        }
+
+        if (this.updateSearch()) {
+            this.setUpdated();
+        }
+    }
+
+    addStream(streamId: StreamId) {
+
+        const state = this.state;
+
+        if (state.streams) {
+            if (state.streams.indexOf(streamId) !== -1) {
+                return;
+            }
+        }
+
+        if (!state.streams) {
+            state.streams = [];
+        }
+
+        state.streams.push(streamId);
+
+        if (this.updateSearch()) {
+            this.setUpdated();
+        }
+    }
+
+    removeStream(streamId: string) {
+
+        const state = this.state;
+
+        if (state.streams?.indexOf(streamId) === -1) {
+            return;
+        }
+
+        const idx = state.streams?.indexOf(streamId);
+        if (!isNumber(idx) || idx < 0) {
+            return;
+        }
+
+        state.streams!.splice(idx, 1);
+
+        if (!state.streams!.length) {
+            state.streams = undefined;
+        }
+
+        if (this.updateSearch()) {
+            this.setUpdated();
+        }
+    }
+
+
+    addSuite(suite: string) {
+
+        const state = this.state;
+
+        if (state.suites) {
+            if (state.suites.indexOf(suite) !== -1) {
+                return;
+            }
+        }
+
+        if (!state.suites) {
+            state.suites = [];
+        }
+
+        state.suites.push(suite);
+
+        if (this.updateSearch()) {
+            this.setUpdated();
+        }
+    }
+
+    removeSuite(suite: string) {
+
+        const state = this.state;
+
+        if (state.suites?.indexOf(suite) === -1) {
+            return;
+        }
+
+        const idx = state.suites?.indexOf(suite);
+        if (!isNumber(idx) || idx < 0) {
+            return;
+        }
+
+        state.suites!.splice(idx, 1);
+
+        if (!state.suites!.length) {
+            state.suites = undefined;
+        }
+
+        if (this.updateSearch()) {
+            this.setUpdated();
+        }
+    }
+
+    get currentMeta(): GetTestMetaResponse[] {
+
+        const state = this.state;
+
+        if (!state.targets?.length || !state.platforms?.length || !state.configurations?.length || !state.rhi?.length || !state.variation?.length) {
+            return [];
+        }
+
+        if (!state.automation || !state.streams?.length) {
+            return [];
+        }
+
+        const meta = new Map<string, GetTestMetaResponse>();
+
+        state.streams.forEach(streamId => {
+
+            const stream = this.testStreams.get(streamId);
+
+            if (!stream) {
+                return;
+            }
+
+            stream.testMetadata.forEach(m => {
+
+                if (m.projectName !== state.automation) {
+                    return;
+                }
+
+                if (state.targets?.length) {
+                    if (!m.buildTargets.find(v => !!state.targets?.find(v2 => v === v2))) {
+                        return;
+                    }
+                }
+
+                if (state.platforms?.length) {
+                    if (!m.platforms.find(v => !!state.platforms?.find(v2 => v === v2))) {
+                        return;
+                    }
+                }
+
+                if (state.configurations?.length) {
+                    if (!m.configurations.find(v => !!state.configurations?.find(v2 => v === v2))) {
+                        return;
+                    }
+                }
+
+                if (state.rhi?.length) {
+                    if (!state.rhi.find(r => m.rhi === r)) {
+                        return;
+                    }
+                }
+
+                if (state.variation?.length) {
+                    if (!state.variation.find(r => m.variation === r)) {
+                        return;
+                    }
+                }
+
+                meta.set(m.id, m);
+
+            });
+
+        });
+
+        return Array.from(meta.values());
+    };
+
+    /// Tests which match current stream
+    get streamTests(): GetTestResponse[] {
+
+        const state = this.state;
+
+        if (!state.automation || !state.streams?.length) {
+            return [];
+        }
+
+        const tests = new Map<string, GetTestResponse>();
+
+        state.streams.forEach(streamId => {
+
+            const stream = this.testStreams.get(streamId);
+            if (!stream) {
+                return;
+            }
+
+            stream.tests.forEach(test => {
+                if (!!test.metadata.find(m => this.metaData.get(m)?.projectName === state.automation)) {
+                    tests.set(test.id, test);
+                }
+            });
+
+        });
+
+        return Array.from(tests.values()).sort((a, b) => a.name.localeCompare(b.name));
+    }
+
+    /// Tests that match current meta
+    get metaTests(): GetTestResponse[] {
+
+        const state = this.state;
+
+        if (!state.automation || !state.streams?.length) {
+            return [];
+        }
+
+        const meta = this.currentMeta;
+        const metaIds = new Set<string>();
+        meta.forEach(m => metaIds.add(m.id));
+
+        const tests = new Map<string, GetTestResponse>();
+        const streamTests = this.streamTests;
+
+        streamTests.forEach(test => {
+
+            if (test.metadata.find(id => metaIds.has(id))) {
+                tests.set(test.id, test);
+            }
+        });
+
+        return Array.from(tests.values()).sort((a, b) => a.name.localeCompare(b.name));
+    }
+
+    /// Suites which match current stream
+    get streamSuites(): GetTestSuiteResponse[] {
+
+        const state = this.state;
+
+        if (!state.automation || !state.streams?.length) {
+            return [];
+        }
+
+        const suites = new Map<string, GetTestSuiteResponse>();
+
+        state.streams.forEach(streamId => {
+
+            const stream = this.testStreams.get(streamId);
+            if (!stream) {
+                return;
+            }
+
+            stream.testSuites.forEach(suite => {
+                if (!!suite.metadata.find(m => this.metaData.get(m)?.projectName === state.automation)) {
+                    suites.set(suite.id, suite);
+                }
+            });
+
+        });
+
+        return Array.from(suites.values()).sort((a, b) => a.name.localeCompare(b.name));
+    }
+
+    /// Suites that match current meta
+    get metaSuites(): GetTestSuiteResponse[] {
+
+        const state = this.state;
+
+        if (!state.automation || !state.streams?.length) {
+            return [];
+        }
+
+        const meta = this.currentMeta;
+        const metaIds = new Set<string>();
+        meta.forEach(m => metaIds.add(m.id));
+
+        const suites = new Map<string, GetTestResponse>();
+        const streamSuites = this.streamSuites;
+
+        streamSuites.forEach(suite => {
+            if (suite.metadata.find(id => metaIds.has(id))) {
+                suites.set(suite.id, suite);
+            }
+        });
+
+
+        return Array.from(suites.values()).sort((a, b) => a.name.localeCompare(b.name));
+    }
+
+
+    updateSearch(): boolean {
+
+        const state = this.state;
+
+        state.weeks = state.weeks ? state.weeks : undefined;
+
+        state.streams = state.streams?.sort((a, b) => a.localeCompare(b));
+        state.automation = state.automation?.length ? state.automation : undefined;
+        state.platforms = state.platforms?.sort((a, b) => a.localeCompare(b));
+
+        state.configurations = state.configurations?.sort((a, b) => a.localeCompare(b));
+        state.targets = state.targets?.sort((a, b) => a.localeCompare(b));
+        state.rhi = state.rhi?.sort((a, b) => a.localeCompare(b));
+        state.variation = state.variation?.sort((a, b) => a.localeCompare(b));
+
+        state.tests = state.tests?.sort((a, b) => a.localeCompare(b));
+        state.suites = state.suites?.sort((a, b) => a.localeCompare(b));
+
+        const search = new URLSearchParams();
+
+        const csearch = this.search.toString();
+
+        if (state.automation) {
+            search.append("automation", state.automation);
+        }
+
+        if (state.weeks) {
+            search.append("weeks", state.weeks.toString());
+        }
+
+        state.streams?.forEach(s => {
+            search.append("stream", s);
+        });
+
+        state.platforms?.forEach(p => {
+            search.append("platform", p);
+        });
+
+        state.configurations?.forEach(c => {
+            search.append("configurations", c);
+        });
+
+        state.targets?.forEach(t => {
+            search.append("targets", t);
+        });
+
+        state.rhi?.forEach(r => {
+            search.append("rhi", r);
+        });
+
+        state.variation?.forEach(v => {
+            search.append("var", v);
+        });
+
+        state.tests?.forEach(r => {
+            search.append("test", r);
+        });
+
+        state.suites?.forEach(r => {
+            search.append("suite", r);
+        });
+
+        if (search.toString() !== csearch) {
+
+            if (this.queryTimeoutId) {
+                clearTimeout(this.queryTimeoutId);
+            }
+
+            this.queryTimeoutId = setTimeout(() => {
+                this.queryTimeoutId = undefined;
+                this.query();
+            }, queryRefreshMS);
+
+            this.search = search;
+            return true;
+        }
+
+        return false;
+    }
+
+    private stateFromSearch(search: URLSearchParams): TestDataState {
+
+        const state: TestDataState = {};
+
+        const streams = search.getAll("stream") ?? undefined;
+        const automation = search.get("automation") ?? undefined;
+        const platforms = search.getAll("platform") ?? undefined;
+        const configurations = search.getAll("configurations") ?? undefined;
+        const targets = search.getAll("targets") ?? undefined;
+        const rhi = search.getAll("rhi") ?? undefined;
+        const variation = search.getAll("var") ?? undefined;
+
+        const tests = search.getAll("test") ?? undefined;
+        const suites = search.getAll("suite") ?? undefined;
+
+        const weeks = search.get("weeks") ?? undefined;
+
+        state.streams = streams?.sort((a, b) => a.localeCompare(b));
+        state.automation = automation?.length ? automation : undefined;
+        state.platforms = platforms?.sort((a, b) => a.localeCompare(b));
+
+        state.configurations = configurations?.sort((a, b) => a.localeCompare(b));
+        state.targets = targets?.sort((a, b) => a.localeCompare(b));
+        state.rhi = rhi?.sort((a, b) => a.localeCompare(b));
+        state.variation = variation?.sort((a, b) => a.localeCompare(b));
+
+        state.tests = tests?.sort((a, b) => a.localeCompare(b));
+        state.suites = suites?.sort((a, b) => a.localeCompare(b));
+
+        state.weeks = weeks ? parseInt(weeks) : defaultQueryWeeks;
+        if (typeof (state.weeks) !== "number") {
+            state.weeks = defaultQueryWeeks;
+        }
+
+        return state;
+    }
+
+    getMetaString(metaId: string, platform = true, config = true, target = true, rhi = true, variation = true) {
+
+        const meta = this.metaData.get(metaId);
+        if (!meta) {
+            return "";
+        }
+
+        const elements: string[] = [];
+        if (platform) {
+            elements.push(`${meta.platforms.join(" - ")}`);
+        }
+        if (config) {
+            elements.push(`${meta.configurations.join(" - ")}`);
+        }
+
+        if (target) {
+            elements.push(`${meta.buildTargets.join(" - ")}`);
+        }
+
+        if (rhi) {
+            elements.push(meta.rhi === "default" ? "Default" : meta.rhi.toUpperCase());
+        }
+
+        if (variation) {
+            elements.push(meta.variation === "default" ? "Default" : meta.variation.toUpperCase());
+        }
+
+        if (elements.length) {
+            return elements.join(" / ");
+        }
+
+        return "";
+
+    }
+
+    getCommonMeta(metaIds: string[], config = true, target = true, rhi = true, variation = true) {
+
+        let cconfigs = "";
+        let uniqueConfigs = false;
+
+        let ctargets = "";
+        let uniqueTargets = false;
+
+        let crhi = "";
+        let uniqueRHI = false;
+
+        let cvariation = "";
+        let uniqueVariation = false;
+
+        metaIds.forEach(m => {
+
+            const meta = this.metaData.get(m)!;
+
+            if (config) {
+                const configs = meta.configurations.join(" - ");
+                if (cconfigs && cconfigs !== configs) {
+                    uniqueConfigs = true;
+                } else {
+                    cconfigs = configs;
+                }
+            }
+
+            if (target) {
+                const targets = meta.buildTargets.join(" - ");
+                if (ctargets && ctargets !== targets) {
+                    uniqueTargets = true;
+                } else {
+                    ctargets = targets;
+                }
+            }
+
+            if (rhi) {
+                const rhi = meta.rhi === "default" ? "Default" : meta.rhi.toUpperCase()
+
+                if (crhi && crhi !== rhi) {
+                    uniqueRHI = true;
+                } else {
+                    crhi = rhi;
+                }
+            }
+
+            if (variation) {
+                const variation = meta.variation === "default" ? "Default" : meta.variation.toUpperCase()
+
+                if (cvariation && cvariation !== variation) {
+                    uniqueVariation = true;
+                } else {
+                    cvariation = variation;
+                }
+            }
+
+        });
+
+        const elements: string[] = [];
+        if (!uniqueConfigs || !uniqueTargets || !uniqueRHI) {
+            if (config && !uniqueConfigs) {
+                elements.push(cconfigs);
+            }
+            if (target && !uniqueTargets) {
+                elements.push(ctargets);
+            }
+            if (rhi && !uniqueRHI) {
+                elements.push(crhi);
+            }
+            if (variation && !uniqueVariation) {
+                elements.push(cvariation);
+            }
+        }
+
+        return {
+            commonConfigs: !uniqueConfigs,
+            commonTargets: !uniqueTargets,
+            commonRHI: !uniqueRHI,
+            commonVariation: !uniqueVariation,
+            commonMetaString: elements.length ? elements.join(" / ") : ""
+        }
+    }
+
+    async query() {
+
+        const state = this.state;
+
+        if (!state.automation || !state.streams?.length || (!state.tests?.length && !state.suites?.length)) {
+            return;
+        }
+
+        this.setQueryLoading(true);
+
+        const testIds: Set<string> = new Set();
+        const suiteIds: Set<string> = new Set();
+        const metaIds: Set<string> = new Set(this.currentMeta.map(m => m.id));
+
+        const streamTestMap: Map<string, GetTestResponse> = new Map();
+        this.streamTests.forEach(t => {
+            streamTestMap.set(t.id, t);
+        });
+
+        const streamSuiteMap: Map<string, GetTestResponse> = new Map();
+        this.streamSuites.forEach(s => {
+            streamSuiteMap.set(s.id, s);
+        });
+
+        state.tests?.forEach(name => {
+            streamTestMap.forEach((t, id) => {
+                if (t.name === name || t.displayName === name) {
+                    testIds.add(id);
+                }
+            });
+        });
+
+        state.suites?.forEach(name => {
+            streamSuiteMap.forEach((s, id) => {
+                if (s.name === name) {
+                    suiteIds.add(id);
+                }
+            });
+        });
+
+
+        // notify user?
+        if (!metaIds.size || (!testIds.size && !suiteIds.size)) {
+            this.refs = [];
+            this.streamStatus = new Map();
+            this.setQueryLoading(false);
+            this.setQueryUpdated();
+            return;
+        }
+
+        let queryWeeks = this.state.weeks;
+        if (typeof (queryWeeks) !== "number") {
+            queryWeeks = defaultQueryWeeks;
+        }
+
+        const minQueryDate = new Date(new Date().valueOf() - (minuteInWeek * queryWeeks * 60000));
+        const maxQueryDate = new Date();
+
+        let refs = await backend.getTestRefs(state.streams, Array.from(metaIds), minQueryDate.toISOString(), maxQueryDate.toISOString(), undefined, undefined, testIds.size ? Array.from(testIds) : undefined, suiteIds.size ? Array.from(suiteIds) : undefined);
+
+        // deduplicate and sort intro stream -> test -> meta -> refs descending by CL
+        const streamStatus = this.streamStatus;
+        streamStatus.clear();
+        const dedupedRefs: GetTestDataRefResponse[] = [];
+        const dupeWarnings = new Set<string>();
+
+        this.changeDates = new Map();
+
+        refs.forEach(r => {
+
+            if (!this.idDates.get(r.id)) {
+                const timestamp = r.id.substring(0, 8)
+                this.idDates.set(r.id, new Date(parseInt(timestamp, 16) * 1000));
+            }
+
+            const date = this.changeDates.get(r.buildChangeList);
+            const rdate = this.idDates.get(r.id)!;
+            if (!date || date.getTime() > rdate.getTime()) {
+                this.changeDates.set(r.buildChangeList, rdate);
+            }
+
+            let streamRefs = streamStatus.get(r.streamId);
+            if (!streamRefs) {
+                streamRefs = new Map();
+                streamStatus.set(r.streamId, streamRefs);
+            }
+
+            const id = r.testId ?? r.suiteId!;
+            let testRefs = streamRefs.get(id);
+            if (!testRefs) {
+                testRefs = new Map();
+                streamRefs.set(id, testRefs);
+            }
+
+            let metaStatus = testRefs.get(r.metaId);
+            if (!metaStatus) {
+                metaStatus = { refs: [], suite: !!r.suiteId };
+                testRefs.set(r.metaId, metaStatus);
+            }
+
+            let cref = metaStatus.refs.find(tr => r.buildChangeList === tr.buildChangeList);
+            if (cref) {
+                // todo: better replace logic, use latest, should elevate errors/skips/etc
+                dupeWarnings.add(id);
+            } else {
+                metaStatus.refs.push(r);
+                dedupedRefs.push(r);
+            }
+        });
+
+        this.minRefDate = new Date();
+        this.maxRefDate = new Date();
+
+        const timeRefs = refs.sort((a, b) => this.idDates.get(a.id)!.getTime() - this.idDates.get(b.id)!.getTime());
+        if (timeRefs.length === 1) {
+            this.minRefDate = this.maxRefDate = this.idDates.get(timeRefs[0].id)!;
+        } else if (timeRefs.length > 1) {
+            this.minRefDate = this.idDates.get(timeRefs[0].id)!;
+            this.maxRefDate = this.idDates.get(timeRefs[timeRefs.length - 1].id)!;
+        }
+
+        for (const [, testMap] of streamStatus) {
+            for (const [, metaMap] of testMap) {
+                for (const [, metaStatus] of metaMap) {
+
+                    const suite = metaStatus.suite;
+                    metaStatus.refs = metaStatus.refs.sort((a, b) => a.buildChangeList - b.buildChangeList).reverse();
+
+                    // errors
+                    for (const ref of metaStatus.refs) {
+                        const error = suite ? !!ref.suiteErrorCount : ref.outcome === "Failure";
+                        if (!error) {
+                            break;
+                        }
+                        metaStatus.error = metaStatus.error ? metaStatus.error + 1 : 1;
+                    }
+
+                    // skip
+                    for (const ref of metaStatus.refs) {
+                        const skip = suite ? !!ref.suiteSkipCount : ref.outcome === "Skipped";
+                        if (!skip) {
+                            break;
+                        }
+                        metaStatus.skip = metaStatus.skip ? metaStatus.skip + 1 : 1;
+                    }
+
+                    // success
+                    for (const ref of metaStatus.refs) {
+                        const success = suite ? (!ref.suiteSkipCount && !ref.suiteErrorCount) : (ref.outcome !== "Failure" && ref.outcome !== "Skipped");
+                        if (!success) {
+                            break;
+                        }
+                        metaStatus.success = metaStatus.success ? metaStatus.success + 1 : 1;
+                    }
+                }
+            }
+        }
+
+        // generate warnings for non unique tests in stream
+        for (let id of dupeWarnings) {
+            if (this.testMap.get(id)) {
+                console.warn(`Test ${this.testMap.get(id)!.name} has duplicate refs for stream changelists and may require a more specific name`);
+            }
+            if (this.suiteMap.get(id)) {
+                console.warn(`Suite ${this.suiteMap.get(id)!.name} has duplicate refs for stream changelists and may require a more specific name`);
+            }
+        }
+
+        this.refs = dedupedRefs;
+
+        const changelists = this.refs.map(r => r.buildChangeList);
+
+        this.minCL = this.maxCL = 0;
+
+        this.minCL = Math.min(...changelists);
+        this.maxCL = Math.max(...changelists);
+
+        this.hasQueried = true;
+
+        this.setQueryLoading(false);
+        this.setQueryUpdated();
+
+        // debug
+        /*
+        const debugRef = `63e3f1685c0ce8f11b01b7a5`;
+        const ref = this.refs.find(r => r.id === debugRef);
+        if (ref) {
+            this.setSuiteRef(ref.id);
+        } 
+        */
+    }
+
+    // Loads immutable data for this view session
+    async load() {
+
+        const allStreams = projectStore.projects.map(p => p.streams).flat().filter(s => !!s).map(s => s!.id);
+        //const streams = testingStreams;
+        const testStreams = await backend.getTestStreams(allStreams);
+
+        // get unique meta, tests, and suites
+        const metaMap: Map<string, GetTestMetaResponse> = new Map();
+
+        const testMap = this.testMap;
+        const suiteMap = this.suiteMap;
+
+        testMap.clear();
+        suiteMap.clear();
+
+        testStreams.forEach(stream => {
+
+            // it looks like there was a misconfiguration, need a way to remove or age out tests
+            stream.tests = stream.tests.filter(t => !t.name.startsWith("EditorAutomation") && !t.name.startsWith("TargetAutomation"));
+
+            this.testStreams.set(stream.streamId, stream);
+
+            const desktop = ["Win64", "Mac", "Linux"];
+            stream.testMetadata.forEach(m => {
+
+                m.platforms = m.platforms.sort((a, b) => {
+                    if (desktop.indexOf(a) !== -1) {
+                        return 1;
+                    }
+                    if (desktop.indexOf(b) !== -1) {
+                        return -1;
+                    }
+
+                    return a.localeCompare(b);
+
+                });
+
+                if (!metaMap.has(m.id)) {
+                    metaMap.set(m.id, m);
+                }
+            });
+
+            stream.tests.forEach(t => {
+                if (!testMap.has(t.id)) {
+                    testMap.set(t.id, t);
+                }
+            });
+
+            stream.testSuites.forEach(s => {
+                if (!suiteMap.has(s.id)) {
+                    suiteMap.set(s.id, s);
+                }
+            });
+        });
+
+        const metaData = Array.from(metaMap.values());
+        this.tests = Array.from(testMap.values()).map(t => t.displayName ?? t.name).sort((a, b) => a.localeCompare(b));
+        this.suites = Array.from(suiteMap.values()).map(s => s.name).sort((a, b) => a.localeCompare(b));
+
+        this.metaData = metaMap;
+        metaMap.forEach((meta, id) => {
+            if (this.metaNames.has(id)) {
+                return;
+            }
+
+            const metaName = `${meta.platforms.join(" - ")} / ${meta.configurations.join(" - ")} / ${meta.buildTargets.join(" - ")} / ${meta.rhi} / ${meta.variation}`;
+
+            this.metaNames.set(id, metaName);
+
+        });
+
+        const automation = new Set<string>();
+        const platforms = new Set<string>();
+        const configurations = new Set<string>();
+        const targets = new Set<string>();
+        const rhi = new Set<string>();
+        const variation = new Set<string>();
+
+        metaData.forEach(m => {
+            automation.add(m.projectName);
+            m.platforms.forEach(p => platforms.add(p));
+            m.buildTargets.forEach(t => targets.add(t));
+            m.configurations.forEach(c => configurations.add(c));
+            rhi.add(m.rhi);
+            variation.add(m.variation)
+        });
+
+        this.automation = Array.from(automation).sort((a, b) => a.localeCompare(b));
+        this.platforms = Array.from(platforms).sort((a, b) => a.localeCompare(b));
+        this.configurations = Array.from(configurations).sort((a, b) => a.localeCompare(b));
+        this.targets = Array.from(targets).sort((a, b) => a.localeCompare(b));
+        this.rhi = Array.from(rhi).sort((a, b) => a.localeCompare(b));
+        this.variation = Array.from(variation).sort((a, b) => a.localeCompare(b));
+
+        this.loaded = true;
+
+        const state = this.state;
+        if (!state.platforms?.length && this.platforms.length) {
+            this.platforms.forEach(p => this.addPlatform(p));
+        }
+        if (!state.configurations?.length && this.configurations.length) {
+            this.configurations.forEach(c => this.addConfiguration(c));
+        }
+        if (!state.targets?.length && this.targets.length) {
+            this.targets.forEach(t => this.addTarget(t));
+        }
+        if (!state.rhi?.length && this.rhi.length) {
+            this.rhi.forEach(r => this.addRHI(r));
+        }
+        if (!state.variation?.length && this.variation.length) {
+            this.variation.forEach(v => this.addVariation(v));
+        }
+
+        this.setUpdated();
+
+        this.query()
+    }
+
+    @action
+    setUpdated() {
+        this.updated++;
+    }
+
+    @action
+    setQueryUpdated() {
+        this.queryUpdated++;
+    }
+
+    @action
+    setQueryLoading(value: boolean) {
+        this.queryLoading = value;
+    }
+
+
+    setFilterState(nstate?: FilterState) {
+        if (this.filterState !== nstate) {
+            this.filterState = nstate;
+            this.setUpdated();
+            // need to update the center view
+            this.setQueryUpdated();
+        }
+    }
+
+    getAutomationStreams(automation: string): StreamId[] {
+
+        const streams: Set<StreamId> = new Set();
+
+        this.testStreams.forEach(testStream => {
+            if (testStream.testMetadata.find(t => t.projectName === automation)) {
+                streams.add(testStream.streamId);
+            }
+
+        });
+
+        return Array.from(streams);
+    }
+
+    getFilterState(): FilterState | undefined {
+        return this.filterState;
+    }
+
+    getRef(refId: string) {
+        return this.refs.find(r => r.id === refId);
+    }
+
+    getSuiteRefs(suiteId: string, metaId: string) {
+        return this.refs.filter(r => r.suiteId === suiteId && r.metaId === metaId);
+    }
+
+    getFilteredRefs(streamIdIn?: StreamId, idIn?: TestId | SuiteId, metaIdIn?: MetaId): GetTestDataRefResponse[] {
+
+        let status: TestStatus[] = [];
+
+        for (const [streamId, testMap] of this.streamStatus) {
+            if (streamIdIn && streamId !== streamIdIn) {
+                continue;
+            }
+            for (const [id, metaMap] of testMap) {
+                if (idIn && id !== idIn) {
+                    continue;
+                }
+                for (const [metaId, metaStatus] of metaMap) {
+                    if (metaIdIn && metaId !== metaIdIn) {
+                        continue;
+                    }
+                    status.push(metaStatus);
+                }
+            }
+        }
+
+        if (!this.filterState) {
+            return status.map(m => m.refs).flat();
+        }
+
+        if (this.filterState === "Failed") {
+            return status.filter(s => !!s.error).map(m => m.refs).flat();
+        }
+
+        if (this.filterState === "Consecutive Failures") {
+            return status.filter(s => s.error && s.error > 1).map(m => m.refs).flat();
+        }
+
+        if (this.filterState === "Skipped") {
+            return status.filter(s => !!s.skip).map(m => m.refs).flat();
+        }
+
+        if (this.filterState === "Success") {
+            return status.filter(s => !!s.success).map(m => m.refs).flat();
+        }
+
+        return [];
+
+    }
+
+    @observable
+    updated: number = 0;
+
+    @observable
+    queryUpdated: number = 0;
+
+    @observable
+    queryLoading: boolean = false;
+
+    state: TestDataState = {};
+    search: URLSearchParams = new URLSearchParams();
+
+    testStreams: Map<string, GetTestStreamResponse> = new Map();
+
+    metaData: Map<MetaId, GetTestMetaResponse> = new Map();
+    metaNames: Map<MetaId, string> = new Map();
+
+    testMap: Map<TestId, GetTestResponse> = new Map();
+    suiteMap: Map<SuiteId, GetTestSuiteResponse> = new Map();
+
+    automation: string[] = [];
+    platforms: string[] = [];
+    configurations: string[] = [];
+    targets: string[] = [];
+    rhi: string[] = [];
+    variation: string[] = [];
+
+    tests: string[] = [];
+    suites: string[] = [];
+
+    private refs: GetTestDataRefResponse[] = [];
+    minCL: number = 0;
+    maxCL: number = 0;
+    minRefDate: Date = new Date();
+    maxRefDate: Date = new Date();
+    idDates = new Map<string, Date>();
+    // quantized so common CL -> lands at same time in views
+    changeDates = new Map<number, Date>();
+
+    private filterState?: FilterState;
+    private streamStatus: Map<StreamId, Map<TestId | SuiteId, Map<MetaId, TestStatus>>> = new Map();
+
+    loaded = false;
+
+    hasQueried = false;
+
+    suiteRef?: string;
+
+    static instance: TestDataHandler;
+
+    private queryTimeoutId: any = undefined;
+}
\ No newline at end of file
diff --git a/Engine/Source/Programs/Horde/HordeDashboard/src/components/AutomationSuiteDetails.tsx b/Engine/Source/Programs/Horde/HordeDashboard/src/components/AutomationSuiteDetails.tsx
--- a/Engine/Source/Programs/Horde/HordeDashboard/src/components/AutomationSuiteDetails.tsx
+++ b/Engine/Source/Programs/Horde/HordeDashboard/src/components/AutomationSuiteDetails.tsx
@@ -50,7 +50,7 @@ class SuiteHandler {
    constructor(suite: GetTestSuiteResponse, suiteRefs: GetTestDataRefResponse[], metaData: GetTestMetaResponse) {
 
       makeObservable(this);
-
+      
       this.suite = suite;
 
       const dedupe = new Set<number>();
@@ -70,7 +70,7 @@ class SuiteHandler {
       this.metaData = metaData;
 
       this.metaName = `${metaData.platforms.join(" - ")} / ${metaData.configurations.join(" - ")} / ${metaData.buildTargets.join(" - ")} / ${metaData.rhi === "default" ? "Default" : metaData.rhi?.toUpperCase()} / ${metaData.variation === "default" ? "Default" : metaData.variation?.toUpperCase()}`;
-         
+
       this.load();
    }
 
@@ -248,8 +248,13 @@ class SuiteHandler {
 
       this.setUpdated();
 
-      // const sdetails = this.testDetails.get(debugRef)!;
-      //this.setSelection({ ref: this.suiteRefs.find(s => s.id === debugRef)!, details: sdetails, testId: debugTestId })
+      /*
+      const debugRef = "63d9a44acf52968117e67b39";
+      const debugTestId = "6376b26ee30d43884952c55d";
+
+      const sdetails = this.testDetails.get(debugRef)!;
+      this.setSelection({ ref: this.suiteRefs.find(s => s.id === debugRef)!, details: sdetails, testId: debugTestId })
+      */
    }
 
    setFilterState(nstate?: FilterState) {
@@ -385,7 +390,7 @@ class SuiteHandler {
 
    selection?: TestSelection;
 
-   metrics?: TestMetrics;
+   metrics?: TestMetrics;   
 
    static fixedTestNames = new Map<TestId, string>();
    static lowerCaseTestNames = new Map<TestId, string>();
@@ -898,28 +903,28 @@ const SuiteTestList: React.FC<{ handler: SuiteHandler }> = observer(({ handler }
    return <Stack key={`key_automation_suite_list_${id_counter++}`} className={hordeClasses.raised} style={{ flexBasis: flexBasic, flexShrink: 0 }}>
       <Stack >
          <Stack>
-            {!!metrics && <Stack horizontal tokens={{childrenGap: 24}} style={{paddingBottom: 8}}>
+            {!!metrics && <Stack horizontal tokens={{ childrenGap: 24 }} style={{ paddingBottom: 8 }}>
                <Stack>
                   <Text>{`Tests: ${metrics.total}`}</Text>
                </Stack>
                {!!metrics.success && <Stack>
                   <Stack>
-                     <Text>{`Success: ${metrics.success} / ${Math.ceil((metrics.success / metrics.total)* 100)}%`}</Text>
+                     <Text>{`Success: ${metrics.success} / ${Math.ceil((metrics.success / metrics.total) * 100)}%`}</Text>
                   </Stack>
                </Stack>}
                {!!metrics.error && <Stack>
                   <Stack>
-                     <Text>{`Error: ${metrics.error} / ${Math.ceil((metrics.error / metrics.total)* 100)}%`}</Text>
+                     <Text>{`Error: ${metrics.error} / ${Math.ceil((metrics.error / metrics.total) * 100)}%`}</Text>
                   </Stack>
                </Stack>}
                {!!metrics.skipped && <Stack>
                   <Stack>
-                     <Text>{`Skipped: ${metrics.skipped} / ${Math.ceil((metrics.skipped / metrics.total)* 100)}%`}</Text>
+                     <Text>{`Skipped: ${metrics.skipped} / ${Math.ceil((metrics.skipped / metrics.total) * 100)}%`}</Text>
                   </Stack>
                </Stack>}
                {!!metrics.unspecified && <Stack>
                   <Stack>
-                     <Text>{`Unspecified: ${metrics.unspecified} / ${Math.ceil((metrics.unspecified / metrics.total)* 100)}%`}</Text>
+                     <Text>{`Unspecified: ${metrics.unspecified} / ${Math.ceil((metrics.unspecified / metrics.total) * 100)}%`}</Text>
                   </Stack>
                </Stack>}
             </Stack>}
@@ -958,8 +963,6 @@ const SuiteTestViewPanel: React.FC<{ handler: SuiteHandler }> = observer(({ hand
    </Stack>
 });
 
-
-
 export const AutomationSuiteDetails: React.FC<{ suite: GetTestSuiteResponse, suiteRefs: GetTestDataRefResponse[], metaData: GetTestMetaResponse, onClose: () => void }> = observer(({ suite, suiteRefs, metaData, onClose }) => {
 
    const [state, setState] = useState<{ handler?: SuiteHandler }>({});
diff --git a/Engine/Source/Programs/Horde/HordeDashboard/src/components/AutomationSuiteTest.tsx b/Engine/Source/Programs/Horde/HordeDashboard/src/components/AutomationSuiteTest.tsx
--- a/Engine/Source/Programs/Horde/HordeDashboard/src/components/AutomationSuiteTest.tsx
+++ b/Engine/Source/Programs/Horde/HordeDashboard/src/components/AutomationSuiteTest.tsx
@@ -1,15 +1,15 @@
 
-import { FontIcon, Image, Label, ScrollablePane, ScrollbarVisibility, Spinner, SpinnerSize, Stack, Text } from "@fluentui/react";
+import { DefaultButton, FontIcon, IContextualMenuItem, IContextualMenuItemProps, IContextualMenuProps, Image, Label, ScrollablePane, ScrollbarVisibility, Spinner, SpinnerSize, Stack, Text } from "@fluentui/react";
 import { action, makeObservable, observable } from 'mobx';
 import { observer } from 'mobx-react-lite';
 import React, { useState } from "react";
 import { Link } from "react-router-dom";
 import backend from "../backend";
-import { ArtifactData, GetSuiteTestDataResponse, GetTestDataDetailsResponse, GetTestDataRefResponse, GetTestMetaResponse, GetTestResponse, GetTestSuiteResponse, TestOutcome } from "../backend/Api";
+import { ArtifactData, GetSuiteTestDataResponse, GetTestDataDetailsResponse, GetTestDataRefResponse, GetTestMetaResponse, GetTestResponse, GetTestSuiteResponse, TestData, TestOutcome } from "../backend/Api";
+import { TestDataHandler } from "../backend/AutomationTestData";
 import dashboard, { StatusColor } from "../backend/Dashboard";
 import { getShortNiceTime } from "../base/utilities/timeUtils";
 
-
 type EventType = "Error" | "Info" | "Warning";
 
 enum ArtifactTag {
@@ -44,64 +44,168 @@ class SuiteTestHandler {
       this.suite = suite;
       this.test = test;
 
-      this.details = details;
-      this.testRef = testRef;
-      this.metaData = metaData;
+      this.refs.set(testRef.id, testRef);
+      this.details.set(details.id, details);
+      this.metaData.set(metaData.id, metaData);
 
-      this.metaName = `${metaData.platforms.join(" - ")} / ${metaData.configurations.join(" - ")} / ${metaData.buildTargets.join(" - ")} / ${metaData.rhi === "default" ? "Default" : metaData.rhi?.toUpperCase()} / ${metaData.variation === "default" ? "Default" : metaData.variation?.toUpperCase()}`;
+      this.initialize(metaData, testRef);
+   }
 
-      this.load();
+   clear() {
+      this.events = [];
+      this.jobId = undefined;
+      this.stepId = undefined;
    }
 
-   async load() {
+   get jobArtifacts(): ArtifactData[] | undefined {
+      if (!this.jobId || !this.stepId) {
+         return undefined;
+      }
 
-      const suiteTest = this.suiteTest = this.details.suiteTests?.find(t => t.testId === this.test.id);
-      if (!suiteTest) {
-         console.error("SuiteTestHandler: Missing suite test");
+      const key = this.jobId + this.stepId;
+      return this.artifacts.get(key);
+   }
+
+   getSuiteTest(metaId: string) {
+
+      const ref = Array.from(this.refs.values()).find(r => r.metaId === metaId);
+
+      if (!ref) {
+         return undefined;
+      }
+
+      const details = this.details.get(ref.id);
+
+      if (!details) {
+         return undefined;
+      }
+
+      return details.suiteTests?.find(t => t.testId === this.test.id);
+   }
+
+   async set(testRef: GetTestDataRefResponse, initialUpdate = false) {
+
+      if (this.cref?.id === testRef.id) {
          return;
       }
 
-      if (!this.details.testDataIds?.length) {
-         console.error("SuiteTestHandler: Suite test has no test id data");
+      this.clear();
+
+      this.loading = true;
+
+      if (!initialUpdate) {
+         this.setUpdated();
+      }
+
+      let details = this.details.get(testRef.id);
+      if (!details) {
+         console.error("Missing details");
          return;
       }
 
-      // @todo: can cache these
-      const requests = this.details.testDataIds.map(id => backend.getTestData(id));
-      const responses = await Promise.all(requests);
+      this.cref = testRef;
+      this.cmetaData = this.metaData.get(testRef.metaId)!;
+      this.cdetails = details;
+
+      // handle test data
+      const requestIds = details.testDataIds.filter(id => !this.testData.get(id));
+      if (requestIds) {
+         const requests = requestIds.map(id => backend.getTestData(id));
+         const responses = await Promise.all(requests);
+         responses.forEach(r => {
+            this.testData.set(r.id, r);
+         });
+      }
+      // Automated Test Session
+      // Unreal Automated Tests::UE.EditorAutomation(RunTest=Rendering) Win64 d3d12
+      // Automated Test Session Result Details::e5f21037-c022-47df-b3d7-028d54e03205
+
+      const testData = details.testDataIds.map(id => this.testData.get(id)!);
 
-      const jobId = this.jobId = responses[0].jobId as string;
-      const stepId = this.stepId = responses[1].stepId as string;
+      if (!!testData.find(td => !td)) {
+         console.error("Missing testdata");
+         return;
+      }
 
-      if (!jobId || !stepId) {
-         console.error("SuiteTestHandler: Missing job or step id");
+      const suiteTest = this.csuiteTest = details.suiteTests?.find(t => t.testId === this.test.id);
+      if (!suiteTest) {
+         console.error("SuiteTestHandler: Missing suite test");
          return;
       }
 
-      const data = responses.map(r => r.data as Record<string, any>);
-      data.forEach(d => {
 
-         const testData = d[suiteTest.uid];
-         if (testData) {
-            if (testData.Events) {
-               for (let e of (testData.Events as TestEvent[])) {
-                  this.events.push(e);
-               }
+      // @todo: these lookups aren't great 
+      const session = testData.find(t => (t.data as any)["Type"]?.trim() === "Automated Test Session");
+      //const tests = testData.find(t => (t.data as any)["Type"]?.trim() === "Unreal Automated Tests");
+      const results = testData.find(t => (t.key as any)?.trim().startsWith("Automated Test Session Result Details"));
+
+      this.events = [];
+
+      this.jobId = session?.jobId;
+      this.stepId = session?.stepId;
+
+      if (!this.jobId || !this.stepId) {
+         console.error("Couldn't get job or step id ");
+         return;
+      }
+
+      if (results) {
+
+         const data = (results.data as Record<string, any>)[suiteTest.uid];
+         if (data.Events) {
+            for (let e of (data.Events as TestEvent[])) {
+               this.events.push(e);
             }
          }
-      });
 
-      let needArtifacts = !!this.events.find(e => e.Tag === EventTag.ImageComparison);
+         let needArtifacts = !!this.events.find(e => e.Tag === EventTag.ImageComparison);
 
-      if (needArtifacts) {
-         this.jobArtifacts = await backend.getJobArtifacts(this.jobId, this.stepId);
+         if (needArtifacts) {
+            const key = this.jobId + this.stepId;
+            if (!this.artifacts.get(key)) {
+               const artifacts = await backend.getJobArtifacts(this.jobId, this.stepId);
+               this.artifacts.set(key, artifacts);
+            }
+         }
       }
 
-      this.loaded = true;
-
+      this.loading = false;
       this.setUpdated();
    }
 
+   private async initialize(srcMetaData: GetTestMetaResponse, srcRef: GetTestDataRefResponse) {
+
+      const instance = TestDataHandler.instance;
+
+      const suiteMetaIds = Array.from(instance.metaData.keys()).filter(meta => !!this.suite.metadata.find(m => m === meta));
+      suiteMetaIds.forEach(metaId => {
+         const metaData = instance.metaData.get(metaId)!;
+         this.metaData.set(metaData.id, metaData);
+         this.metaNames.set(metaId, `${metaData.platforms.join(" - ")} / ${metaData.configurations.join(" - ")} / ${metaData.buildTargets.join(" - ")} / ${metaData.rhi === "default" ? "Default" : metaData.rhi?.toUpperCase()} / ${metaData.variation === "default" ? "Default" : metaData.variation?.toUpperCase()}`);
+      })
+
+      let altRefs: GetTestDataRefResponse[] = [];
+      const altMetaIds = suiteMetaIds.filter(id => id !== srcMetaData.id);
+      if (altMetaIds.length) {
+         altRefs = await backend.getTestRefs([srcRef.streamId], altMetaIds, undefined, undefined, srcRef.buildChangeList, srcRef.buildChangeList, undefined, [this.suite.id]);
+         altRefs.forEach(r => {
+            this.refs.set(r.id, r);
+         });
+      }
+
+      // want details here, so we can put the state in the dropdown, otherwise annoying to have to cycle to each 
+      if (altRefs.length) {
+         const altDetails = await backend.getTestDetails(altRefs.map(v => v.id));
+         altDetails.forEach(d => {
+            this.details.set(d.id, d);
+         });
+      }
+
+      this.initialized = true;
+
+      await this.set(srcRef, true);
+   }
+
    @action
    setUpdated() {
       this.updated++;
@@ -110,21 +214,36 @@ class SuiteTestHandler {
    @observable
    updated: number = 0;
 
-   events: TestEvent[] = [];
-
    suite: GetTestSuiteResponse;
+
+   cref?: GetTestDataRefResponse;
+   cdetails?: GetTestDataDetailsResponse;
+   cmetaData?: GetTestMetaResponse;
+   csuiteTest?: GetSuiteTestDataResponse;
+
+   // test ref id => refs
+   refs = new Map<string, GetTestDataRefResponse>();
+
+   // test ref id => details
+   details = new Map<string, GetTestDataDetailsResponse>();
+
+   // meta data id => meta data
+   metaData = new Map<string, GetTestMetaResponse>();
+
+   testData = new Map<string, TestData>();
+
+   metaNames = new Map<string, string>();
+
    test: GetTestResponse;
-   suiteTest?: GetSuiteTestDataResponse;
-   details: GetTestDataDetailsResponse;
-   testRef: GetTestDataRefResponse;
-   metaData: GetTestMetaResponse;
-   metaName: string;
+
+   initialized = false;
+
+   loading = false;
 
    jobId?: string;
    stepId?: string;
-   jobArtifacts?: ArtifactData[];
-
-   loaded = false;
+   artifacts: Map<string, ArtifactData[]> = new Map();
+   events: TestEvent[] = [];
 
 }
 
@@ -134,6 +253,18 @@ const EventPanel: React.FC<{ handler: SuiteTestHandler }> = observer(({ handler
 
    if (handler.updated) { }
 
+   if (handler.loading) {
+
+      return <Stack grow style={{ paddingTop: 16 }} horizontalAlign="center" tokens={{ childrenGap: 18 }}>
+         <Stack>
+            <Text variant="mediumPlus">Loading Events</Text>
+         </Stack>
+         <Stack>
+            <Spinner size={SpinnerSize.large} />
+         </Stack>
+      </Stack>
+   }
+
    const events = handler.events.map(e => {
 
       const artifactMap = new Map<string, string>();
@@ -190,7 +321,7 @@ const EventPanel: React.FC<{ handler: SuiteTestHandler }> = observer(({ handler
          }
 
          return <Stack horizontalAlign="center" style={{ paddingBottom: 12 }}>
-            <Image width="100%" style={{ minWidth: 200, maxWidth: 400, minHeight: 120 }} src={url} />
+            <Image width="340px" height="190px" src={url} />
             <Text>{tag}</Text>
          </Stack>
       });
@@ -231,12 +362,111 @@ const EventPanel: React.FC<{ handler: SuiteTestHandler }> = observer(({ handler
 });
 
 
+const MetaChooser: React.FC<{ handler: SuiteTestHandler }> = ({ handler }) => {
+
+   const metaIds = new Set<string>(Array.from(handler.refs.values()).map(r => r.metaId));
+   const metaData: GetTestMetaResponse[] = [];
+   metaIds.forEach(id => {
+      metaData.push(handler.metaData.get(id)!);
+   });
+
+   const platforms = Array.from(new Set<string>(metaData.map(d => d.platforms).flat())).sort((a, b) => a.localeCompare(b));
+
+   const options: IContextualMenuItem[] = [];
+
+   platforms.forEach(platform => {
+
+      const metaItems: IContextualMenuItem[] = [];
+
+      const pmeta = metaData.filter(m => m.platforms.indexOf(platform) !== -1).sort((a, b) => handler.metaNames.get(a.id)!.localeCompare(handler.metaNames.get(b.id)!));
+
+      pmeta.forEach(pm => {
+
+         const suiteTest = handler.getSuiteTest(pm.id);
+
+         metaItems.push({
+            key: `meta_select_${platform}_${pm.id}}`, text: handler.metaNames.get(pm.id), onClick: (ev, item) => {
+
+               if (!item) {
+                  return;
+               }
+
+               const ref = Array.from(handler.refs.values()).find(r => r.metaId === pm.id);
+
+               if (ref) {
+                  handler.set(ref);
+               }
+
+               // don't close
+               //ev?.preventDefault()
+            },
+            data: suiteTest
+         });
+
+      });
+
+      const renderMeta = (props: IContextualMenuItemProps) => {
+
+         const test = props.item.data as GetSuiteTestDataResponse | undefined;
+
+         const scolors = dashboard.getStatusColors();
+
+         let color = scolors.get(StatusColor.Unspecified);
+
+         if (test?.outcome === TestOutcome.Success) {
+            color = scolors.get(StatusColor.Success)!;
+         }
+
+         if (test?.outcome === TestOutcome.Failure) {
+            color = scolors.get(StatusColor.Failure)!;
+         }
+
+         if (test?.outcome === TestOutcome.Skipped) {
+            color = scolors.get(StatusColor.Skipped)!;
+         }
+
+         return <Stack style={{ paddingLeft: 4 }}><Stack horizontal tokens={{ childrenGap: 8 }}>
+            <Stack>
+               <FontIcon style={{ color: color, fontSize: 11 }} iconName="Square" />
+            </Stack>
+            <Text style={{ fontSize: 11 }} >{props.item.text}</Text>
+         </Stack>
+         </Stack>;
+      };
+
+      const subMenuProps: IContextualMenuProps = {
+         contextualMenuItemAs: renderMeta,
+         shouldFocusOnMount: true,
+         subMenuHoverDelay: 0,
+         items: metaItems,
+      };
+
+      options.push({ key: `test_platform_select_${platform}`, text: platform, subMenuProps: subMenuProps });
+
+   });
+
+   const menuProps: IContextualMenuProps = {
+      shouldFocusOnMount: true,
+      subMenuHoverDelay: 0,
+      items: options,
+   };
+
+   return <DefaultButton style={{ width: 280, textAlign: "left" }} menuProps={menuProps} text={"Alternate Platform"} />
+
+}
+
 const TestInfoBox: React.FC<{ handler: SuiteTestHandler }> = observer(({ handler }) => {
 
 
    const labelWidth = 120;
    const textWidth = 120;
 
+   if (handler.updated) { }
+
+   if (!handler.cref) {
+      return null;
+   }
+
    const infoItem = (label: string, textIn: any, link?: string) => {
       const text = textIn.toString();
       return <Stack key={`test_info_item_${idCounter++}`} horizontal style={{ width: labelWidth + textWidth + 12 }}>
@@ -260,11 +490,11 @@ const TestInfoBox: React.FC<{ handler: SuiteTestHandler }> = observer(({ handler
       infoItems.push(infoItem("Horde", "Job Step", `/job/${handler.jobId}/?step=${handler.stepId}`));
    }
 
-   const timestamp = handler.testRef.id.substring(0, 8)
+   const timestamp = handler.cref.id.substring(0, 8)
    const time = getShortNiceTime(new Date(parseInt(timestamp, 16) * 1000), true);
 
 
-   infoItems.push(infoItem("Change", handler.testRef.buildChangeList));
+   infoItems.push(infoItem("Change", handler.cref.buildChangeList));
    infoItems.push(infoItem("Date", time));
    infoItems.push(infoItem("Errors", errorCount));
    infoItems.push(infoItem("Warnings", warningCount));
@@ -279,7 +509,7 @@ const AutomationSuiteTestInner: React.FC<{ handler: SuiteTestHandler }> = observ
 
    if (handler.updated) { }
 
-   if (!handler.loaded) {
+   if (!handler.initialized) {
       return <Stack style={{ paddingTop: 16 }} horizontalAlign="center" tokens={{ childrenGap: 18 }}>
          <Stack>
             <Text variant="mediumPlus">Loading Events</Text>
@@ -291,7 +521,7 @@ const AutomationSuiteTestInner: React.FC<{ handler: SuiteTestHandler }> = observ
    }
 
    const test = handler.test;
-   const suiteTest = handler.suiteTest;
+   const suiteTest = handler.csuiteTest;
 
    if (!suiteTest) {
       return null;
@@ -326,17 +556,27 @@ const AutomationSuiteTestInner: React.FC<{ handler: SuiteTestHandler }> = observ
 
    return <Stack style={{ height: "100%" }} >
       <Stack horizontal tokens={{ childrenGap: 4 }} verticalAlign="center" style={{ paddingBottom: 12 }}>
-         <Stack className="horde-no-darktheme" style={{ paddingTop: 3, paddingRight: 4 }}>
+         <Stack className="horde-no-darktheme" style={{ paddingTop: 2, paddingRight: 4 }}>
             <FontIcon style={{ color: color }} iconName="Square" />
          </Stack>
-         <Stack>
-            <Text style={{ fontSize: 14, fontFamily: "Horde Open Sans Semibold" }}>{testName}</Text>
+         <Stack tokens={{childrenGap: 4}}>
+            <Stack>
+               <Text style={{ fontSize: 13, fontFamily: "Horde Open Sans Semibold" }}>{testName}</Text>
+            </Stack>
+            <Stack>
+               <Text style={{ fontSize: 11, fontFamily: "Horde Open Sans Semibold" }}>{handler.metaNames.get(handler.cmetaData!.id)}</Text>
+            </Stack>
          </Stack>
       </Stack>
 
       <Stack grow>
          <Stack horizontal style={{ height: "100%" }}>
-            <TestInfoBox handler={handler} />
+            <Stack>
+               <Stack style={{ paddingTop: 8, paddingBottom: 18 }}>
+                  <MetaChooser handler={handler} />
+               </Stack>
+               <TestInfoBox handler={handler} />
+            </Stack>
             <EventPanel handler={handler} />
          </Stack>
       </Stack>
diff --git a/Engine/Source/Programs/Horde/HordeDashboard/src/components/AutomationView.tsx b/Engine/Source/Programs/Horde/HordeDashboard/src/components/AutomationView.tsx
--- a/Engine/Source/Programs/Horde/HordeDashboard/src/components/AutomationView.tsx
+++ b/Engine/Source/Programs/Horde/HordeDashboard/src/components/AutomationView.tsx
@@ -1,15 +1,14 @@
 // Copyright Epic Games, Inc. All Rights Reserved.
 
-import { isNumber } from '@datadog/browser-core';
 import { ComboBox, DefaultButton, Dropdown, FocusZone, FocusZoneDirection, FontIcon, IComboBox, IComboBoxOption, IComboBoxStyles, IContextualMenuItem, IContextualMenuProps, IDropdownOption, Label, PrimaryButton, ScrollablePane, ScrollbarVisibility, SelectableOptionMenuItemType, Spinner, SpinnerSize, Stack, Text } from '@fluentui/react';
 import * as d3 from "d3";
-import { action, makeObservable, observable } from 'mobx';
 import { observer } from 'mobx-react-lite';
 import moment from 'moment';
 import React, { useState } from 'react';
 import { useSearchParams } from 'react-router-dom';
 import backend from '../backend';
-import { GetTestDataRefResponse, GetTestMetaResponse, GetTestResponse, GetTestStreamResponse, GetTestSuiteResponse, TestOutcome } from '../backend/Api';
+import { GetTestDataRefResponse, GetTestMetaResponse, GetTestResponse, GetTestSuiteResponse, TestOutcome } from '../backend/Api';
+import { FilterState, TestDataHandler } from '../backend/AutomationTestData';
 import dashboard, { StatusColor } from '../backend/Dashboard';
 import { projectStore } from '../backend/ProjectStore';
 import { useWindowSize } from '../base/utilities/hooks';
@@ -20,1327 +19,13 @@ import { Breadcrumbs } from './Breadcrumbs';
 import ErrorBoundary from './ErrorBoundary';
 import { TopNav } from './TopNav';
 
-// some aliases
-type StreamId = string;
-type TestId = string;
-type SuiteId = string;
-type MetaId = string;
-
-type TestStatus = {
-   suite: boolean;
-   refs: GetTestDataRefResponse[];
-   success?: number;
-   error?: number;
-   skip?: number;
-}
 
 type SelectionType = d3.Selection<SVGGElement, unknown, null, undefined>;
 type DivSelectionType = d3.Selection<HTMLDivElement, unknown, null, undefined>;
 
-const defaultQueryWeeks = 2;
-const minuteInWeek = 10080;
-const queryRefreshMS = 2500;
-
-// query states
-type TestDataState = {
-   streams?: string[];
-   automation?: string;
-   platforms?: string[];
-   configurations?: string[];
-   targets?: string[];
-   rhi?: string[];
-   variation?: string[];
-   tests?: string[];
-   suites?: string[];
-   weeks?: number;
-}
-
-type FilterState = "Success" | "Failed" | "Consecutive Failures" | "Skipped";
-
 // used to regenerate combo boxes upon automation project switches
 let multiComboBoxId = 0;
 
-class TestDataHandler {
-
-   constructor(search: URLSearchParams) {
-      makeObservable(this);
-      this.state = this.stateFromSearch(search);
-      this.load();
-   }
-
-   setAutomation(automation: string) {
-      this.state.automation = automation;
-
-      if (automation) {
-         const streams = this.getAutomationStreams(automation);
-         this.state.streams = this.state.streams?.filter(s => streams.indexOf(s) !== -1);
-         if (!this.state.streams?.length) {
-            this.state.streams = undefined;
-         }
-
-         //this.state.tests = undefined;
-         //this.state.suites = undefined;
-
-         const tests = this.streamTests;
-         this.state.tests = this.state.tests?.filter(test => !!tests.find(t => t.name === test));
-         if (!this.state.tests?.length) {
-            this.state.tests = undefined;
-         }
-
-         const suites = this.streamSuites;
-         this.state.suites = this.state.suites?.filter(suite => !!suites.find(s => s.name === suite));
-         if (!this.state.suites?.length) {
-            this.state.suites = undefined;
-         }
-
-      }
-
-      if (this.updateSearch()) {
-         multiComboBoxId++;
-         this.setUpdated();
-      }
-   }
-
-   setQueryWeeks(weeks: number) {
-      this.state.weeks = weeks;
-      if (this.updateSearch()) {
-         this.setUpdated();
-      }
-   }
-
-   setSuiteRef(suiteRef?: string) {
-      if (this.suiteRef === suiteRef) {
-         return;
-      }
-      this.suiteRef = suiteRef;
-      this.setUpdated();
-   }
-
-   addPlatform(platform: string) {
-
-      const state = this.state;
-
-      if (state.platforms) {
-         if (state.platforms.indexOf(platform) !== -1) {
-            return;
-         }
-      }
-
-      if (!state.platforms) {
-         state.platforms = [];
-      }
-
-      state.platforms.push(platform);
-
-      if (this.updateSearch()) {
-         this.setUpdated();
-      }
-   }
-
-   removePlatform(platform: string) {
-
-      const state = this.state;
-
-      if (state.platforms?.indexOf(platform) === -1) {
-         return;
-      }
-
-      const idx = state.platforms?.indexOf(platform);
-      if (!isNumber(idx) || idx < 0) {
-         return;
-      }
-
-      state.platforms!.splice(idx, 1);
-
-      if (!state.platforms!.length) {
-         state.platforms = undefined;
-      }
-
-      if (this.updateSearch()) {
-         this.setUpdated();
-      }
-   }
-
-   addConfiguration(config: string) {
-
-      const state = this.state;
-
-      if (state.configurations) {
-         if (state.configurations.indexOf(config) !== -1) {
-            return;
-         }
-      }
-
-      if (!state.configurations) {
-         state.configurations = [];
-      }
-
-      state.configurations.push(config);
-
-      if (this.updateSearch()) {
-         this.setUpdated();
-      }
-   }
-
-   removeConfiguration(config: string) {
-
-      const state = this.state;
-
-      if (state.configurations?.indexOf(config) === -1) {
-         return;
-      }
-
-      const idx = state.configurations?.indexOf(config);
-      if (!isNumber(idx) || idx < 0) {
-         return;
-      }
-
-      state.configurations!.splice(idx, 1);
-
-      if (!state.configurations!.length) {
-         state.configurations = undefined;
-      }
-
-      if (this.updateSearch()) {
-         this.setUpdated();
-      }
-   }
-
-   addTarget(target: string) {
-
-      const state = this.state;
-
-      if (state.targets) {
-         if (state.targets.indexOf(target) !== -1) {
-            return;
-         }
-      }
-
-      if (!state.targets) {
-         state.targets = [];
-      }
-
-      state.targets.push(target);
-
-      if (this.updateSearch()) {
-         this.setUpdated();
-      }
-   }
-
-   removeTarget(target: string) {
-
-      const state = this.state;
-
-      if (state.targets?.indexOf(target) === -1) {
-         return;
-      }
-
-      const idx = state.targets?.indexOf(target);
-      if (!isNumber(idx) || idx < 0) {
-         return;
-      }
-
-      state.targets!.splice(idx, 1);
-
-      if (!state.targets!.length) {
-         state.targets = undefined;
-      }
-
-      if (this.updateSearch()) {
-         this.setUpdated();
-      }
-   }
-
-   addRHI(rhi: string) {
-
-      const state = this.state;
-
-      if (state.rhi) {
-         if (state.rhi.indexOf(rhi) !== -1) {
-            return;
-         }
-      }
-
-      if (!state.rhi) {
-         state.rhi = [];
-      }
-
-      state.rhi.push(rhi);
-
-      if (this.updateSearch()) {
-         this.setUpdated();
-      }
-   }
-
-   removeRHI(rhi: string) {
-
-      const state = this.state;
-
-      if (state.rhi?.indexOf(rhi) === -1) {
-         return;
-      }
-
-      const idx = state.rhi?.indexOf(rhi);
-      if (!isNumber(idx) || idx < 0) {
-         return;
-      }
-
-      state.rhi!.splice(idx, 1);
-
-      if (!state.rhi!.length) {
-         state.rhi = undefined;
-      }
-
-      if (this.updateSearch()) {
-         this.setUpdated();
-      }
-   }
-
-   addVariation(variation: string) {
-
-      const state = this.state;
-
-      if (state.variation) {
-         if (state.variation.indexOf(variation) !== -1) {
-            return;
-         }
-      }
-
-      if (!state.variation) {
-         state.variation = [];
-      }
-
-      state.variation.push(variation);
-
-      if (this.updateSearch()) {
-         this.setUpdated();
-      }
-   }
-
-   removeVariation(variation: string) {
-
-      const state = this.state;
-
-      if (state.variation?.indexOf(variation) === -1) {
-         return;
-      }
-
-      const idx = state.variation?.indexOf(variation);
-      if (!isNumber(idx) || idx < 0) {
-         return;
-      }
-
-      state.variation!.splice(idx, 1);
-
-      if (!state.variation!.length) {
-         state.variation = undefined;
-      }
-
-      if (this.updateSearch()) {
-         this.setUpdated();
-      }
-   }
-
-
-   addTest(test: string) {
-
-      const state = this.state;
-
-      if (state.tests) {
-         if (state.tests.indexOf(test) !== -1) {
-            return;
-         }
-      }
-
-      if (!state.tests) {
-         state.tests = [];
-      }
-
-      state.tests.push(test);
-
-      if (this.updateSearch()) {
-         this.setUpdated();
-      }
-   }
-
-   removeTest(test: string) {
-
-      const state = this.state;
-
-      if (state.tests?.indexOf(test) === -1) {
-         return;
-      }
-
-      const idx = state.tests?.indexOf(test);
-      if (!isNumber(idx) || idx < 0) {
-         return;
-      }
-
-      state.tests!.splice(idx, 1);
-
-      if (!state.tests!.length) {
-         state.tests = undefined;
-      }
-
-      if (this.updateSearch()) {
-         this.setUpdated();
-      }
-   }
-
-   addStream(streamId: StreamId) {
-
-      const state = this.state;
-
-      if (state.streams) {
-         if (state.streams.indexOf(streamId) !== -1) {
-            return;
-         }
-      }
-
-      if (!state.streams) {
-         state.streams = [];
-      }
-
-      state.streams.push(streamId);
-
-      if (this.updateSearch()) {
-         this.setUpdated();
-      }
-   }
-
-   removeStream(streamId: string) {
-
-      const state = this.state;
-
-      if (state.streams?.indexOf(streamId) === -1) {
-         return;
-      }
-
-      const idx = state.streams?.indexOf(streamId);
-      if (!isNumber(idx) || idx < 0) {
-         return;
-      }
-
-      state.streams!.splice(idx, 1);
-
-      if (!state.streams!.length) {
-         state.streams = undefined;
-      }
-
-      if (this.updateSearch()) {
-         this.setUpdated();
-      }
-   }
-
-
-   addSuite(suite: string) {
-
-      const state = this.state;
-
-      if (state.suites) {
-         if (state.suites.indexOf(suite) !== -1) {
-            return;
-         }
-      }
-
-      if (!state.suites) {
-         state.suites = [];
-      }
-
-      state.suites.push(suite);
-
-      if (this.updateSearch()) {
-         this.setUpdated();
-      }
-   }
-
-   removeSuite(suite: string) {
-
-      const state = this.state;
-
-      if (state.suites?.indexOf(suite) === -1) {
-         return;
-      }
-
-      const idx = state.suites?.indexOf(suite);
-      if (!isNumber(idx) || idx < 0) {
-         return;
-      }
-
-      state.suites!.splice(idx, 1);
-
-      if (!state.suites!.length) {
-         state.suites = undefined;
-      }
-
-      if (this.updateSearch()) {
-         this.setUpdated();
-      }
-   }
-
-   get currentMeta(): GetTestMetaResponse[] {
-
-      const state = this.state;
-
-      if (!state.targets?.length || !state.platforms?.length || !state.configurations?.length || !state.rhi?.length || !state.variation?.length) {
-         return [];
-      }
-
-      if (!state.automation || !state.streams?.length) {
-         return [];
-      }
-
-      const meta = new Map<string, GetTestMetaResponse>();
-
-      state.streams.forEach(streamId => {
-
-         const stream = this.testStreams.get(streamId);
-
-         if (!stream) {
-            return;
-         }
-
-         stream.testMetadata.forEach(m => {
-
-            if (m.projectName !== state.automation) {
-               return;
-            }
-
-            if (state.targets?.length) {
-               if (!m.buildTargets.find(v => !!state.targets?.find(v2 => v === v2))) {
-                  return;
-               }
-            }
-
-            if (state.platforms?.length) {
-               if (!m.platforms.find(v => !!state.platforms?.find(v2 => v === v2))) {
-                  return;
-               }
-            }
-
-            if (state.configurations?.length) {
-               if (!m.configurations.find(v => !!state.configurations?.find(v2 => v === v2))) {
-                  return;
-               }
-            }
-
-            if (state.rhi?.length) {
-               if (!state.rhi.find(r => m.rhi === r)) {
-                  return;
-               }
-            }
-
-            if (state.variation?.length) {
-               if (!state.variation.find(r => m.variation === r)) {
-                  return;
-               }
-            }
-
-            meta.set(m.id, m);
-
-         });
-
-      });
-
-      return Array.from(meta.values());
-   };
-
-   /// Tests which match current stream
-   get streamTests(): GetTestResponse[] {
-
-      const state = this.state;
-
-      if (!state.automation || !state.streams?.length) {
-         return [];
-      }
-
-      const tests = new Map<string, GetTestResponse>();
-
-      state.streams.forEach(streamId => {
-
-         const stream = this.testStreams.get(streamId);
-         if (!stream) {
-            return;
-         }
-
-         stream.tests.forEach(test => {
-            if (!!test.metadata.find(m => this.metaData.get(m)?.projectName === state.automation)) {
-               tests.set(test.id, test);
-            }
-         });
-
-      });
-
-      return Array.from(tests.values()).sort((a, b) => a.name.localeCompare(b.name));
-   }
-
-   /// Tests that match current meta
-   get metaTests(): GetTestResponse[] {
-
-      const state = this.state;
-
-      if (!state.automation || !state.streams?.length) {
-         return [];
-      }
-
-      const meta = this.currentMeta;
-      const metaIds = new Set<string>();
-      meta.forEach(m => metaIds.add(m.id));
-
-      const tests = new Map<string, GetTestResponse>();
-      const streamTests = this.streamTests;
-
-      streamTests.forEach(test => {
-
-         if (test.metadata.find(id => metaIds.has(id))) {
-            tests.set(test.id, test);
-         }
-      });
-
-      return Array.from(tests.values()).sort((a, b) => a.name.localeCompare(b.name));
-   }
-
-   /// Suites which match current stream
-   get streamSuites(): GetTestSuiteResponse[] {
-
-      const state = this.state;
-
-      if (!state.automation || !state.streams?.length) {
-         return [];
-      }
-
-      const suites = new Map<string, GetTestSuiteResponse>();
-
-      state.streams.forEach(streamId => {
-
-         const stream = this.testStreams.get(streamId);
-         if (!stream) {
-            return;
-         }
-
-         stream.testSuites.forEach(suite => {
-            if (!!suite.metadata.find(m => this.metaData.get(m)?.projectName === state.automation)) {
-               suites.set(suite.id, suite);
-            }
-         });
-
-      });
-
-      return Array.from(suites.values()).sort((a, b) => a.name.localeCompare(b.name));
-   }
-
-   /// Suites that match current meta
-   get metaSuites(): GetTestSuiteResponse[] {
-
-      const state = this.state;
-
-      if (!state.automation || !state.streams?.length) {
-         return [];
-      }
-
-      const meta = this.currentMeta;
-      const metaIds = new Set<string>();
-      meta.forEach(m => metaIds.add(m.id));
-
-      const suites = new Map<string, GetTestResponse>();
-      const streamSuites = this.streamSuites;
-
-      streamSuites.forEach(suite => {
-         if (suite.metadata.find(id => metaIds.has(id))) {
-            suites.set(suite.id, suite);
-         }
-      });
-
-
-      return Array.from(suites.values()).sort((a, b) => a.name.localeCompare(b.name));
-   }
-
-
-   updateSearch(): boolean {
-
-      const state = this.state;
-
-      state.weeks = state.weeks ? state.weeks : undefined;
-
-      state.streams = state.streams?.sort((a, b) => a.localeCompare(b));
-      state.automation = state.automation?.length ? state.automation : undefined;
-      state.platforms = state.platforms?.sort((a, b) => a.localeCompare(b));
-
-      state.configurations = state.configurations?.sort((a, b) => a.localeCompare(b));
-      state.targets = state.targets?.sort((a, b) => a.localeCompare(b));
-      state.rhi = state.rhi?.sort((a, b) => a.localeCompare(b));
-      state.variation = state.variation?.sort((a, b) => a.localeCompare(b));
-
-      state.tests = state.tests?.sort((a, b) => a.localeCompare(b));
-      state.suites = state.suites?.sort((a, b) => a.localeCompare(b));
-
-      const search = new URLSearchParams();
-
-      const csearch = this.search.toString();
-
-      if (state.automation) {
-         search.append("automation", state.automation);
-      }
-
-      if (state.weeks) {
-         search.append("weeks", state.weeks.toString());
-      }
-
-      state.streams?.forEach(s => {
-         search.append("stream", s);
-      });
-
-      state.platforms?.forEach(p => {
-         search.append("platform", p);
-      });
-
-      state.configurations?.forEach(c => {
-         search.append("configurations", c);
-      });
-
-      state.targets?.forEach(t => {
-         search.append("targets", t);
-      });
-
-      state.rhi?.forEach(r => {
-         search.append("rhi", r);
-      });
-
-      state.variation?.forEach(v => {
-         search.append("var", v);
-      });
-
-      state.tests?.forEach(r => {
-         search.append("test", r);
-      });
-
-      state.suites?.forEach(r => {
-         search.append("suite", r);
-      });
-
-      if (search.toString() !== csearch) {
-
-         if (this.queryTimeoutId) {
-            clearTimeout(this.queryTimeoutId);
-         }
-
-         this.queryTimeoutId = setTimeout(() => {
-            this.queryTimeoutId = undefined;
-            this.query();
-         }, queryRefreshMS);
-
-         this.search = search;
-         return true;
-      }
-
-      return false;
-   }
-
-   private stateFromSearch(search: URLSearchParams): TestDataState {
-
-      const state: TestDataState = {};
-
-      const streams = search.getAll("stream") ?? undefined;
-      const automation = search.get("automation") ?? undefined;
-      const platforms = search.getAll("platform") ?? undefined;
-      const configurations = search.getAll("configurations") ?? undefined;
-      const targets = search.getAll("targets") ?? undefined;
-      const rhi = search.getAll("rhi") ?? undefined;
-      const variation = search.getAll("var") ?? undefined;
-
-      const tests = search.getAll("test") ?? undefined;
-      const suites = search.getAll("suite") ?? undefined;
-
-      const weeks = search.get("weeks") ?? undefined;
-
-      state.streams = streams?.sort((a, b) => a.localeCompare(b));
-      state.automation = automation?.length ? automation : undefined;
-      state.platforms = platforms?.sort((a, b) => a.localeCompare(b));
-
-      state.configurations = configurations?.sort((a, b) => a.localeCompare(b));
-      state.targets = targets?.sort((a, b) => a.localeCompare(b));
-      state.rhi = rhi?.sort((a, b) => a.localeCompare(b));
-      state.variation = variation?.sort((a, b) => a.localeCompare(b));
-
-      state.tests = tests?.sort((a, b) => a.localeCompare(b));
-      state.suites = suites?.sort((a, b) => a.localeCompare(b));
-
-      state.weeks = weeks ? parseInt(weeks) : defaultQueryWeeks;
-      if (typeof (state.weeks) !== "number") {
-         state.weeks = defaultQueryWeeks;
-      }
-
-      return state;
-   }
-
-   getMetaString(metaId: string, platform = true, config = true, target = true, rhi = true, variation = true) {
-
-      const meta = this.metaData.get(metaId);
-      if (!meta) {
-         return "";
-      }
-
-      const elements: string[] = [];
-      if (platform) {
-         elements.push(`${meta.platforms.join(" - ")}`);
-      }
-      if (config) {
-         elements.push(`${meta.configurations.join(" - ")}`);
-      }
-
-      if (target) {
-         elements.push(`${meta.buildTargets.join(" - ")}`);
-      }
-
-      if (rhi) {
-         elements.push(meta.rhi === "default" ? "Default" : meta.rhi.toUpperCase());
-      }
-
-      if (variation) {
-         elements.push(meta.variation === "default" ? "Default" : meta.variation.toUpperCase());
-      }
-
-      if (elements.length) {
-         return elements.join(" / ");
-      }
-
-      return "";
-
-   }
-
-   getCommonMeta(metaIds: string[], config = true, target = true, rhi = true, variation = true) {
-
-      let cconfigs = "";
-      let uniqueConfigs = false;
-
-      let ctargets = "";
-      let uniqueTargets = false;
-
-      let crhi = "";
-      let uniqueRHI = false;
-
-      let cvariation = "";
-      let uniqueVariation = false;
-
-      metaIds.forEach(m => {
-
-         const meta = this.metaData.get(m)!;
-
-         if (config) {
-            const configs = meta.configurations.join(" - ");
-            if (cconfigs && cconfigs !== configs) {
-               uniqueConfigs = true;
-            } else {
-               cconfigs = configs;
-            }
-         }
-
-         if (target) {
-            const targets = meta.buildTargets.join(" - ");
-            if (ctargets && ctargets !== targets) {
-               uniqueTargets = true;
-            } else {
-               ctargets = targets;
-            }
-         }
-
-         if (rhi) {
-            const rhi = meta.rhi === "default" ? "Default" : meta.rhi.toUpperCase()
-
-            if (crhi && crhi !== rhi) {
-               uniqueRHI = true;
-            } else {
-               crhi = rhi;
-            }
-         }
-
-         if (variation) {
-            const variation = meta.variation === "default" ? "Default" : meta.variation.toUpperCase()
-
-            if (cvariation && cvariation !== variation) {
-               uniqueVariation = true;
-            } else {
-               cvariation = variation;
-            }
-         }
-
-      });
-
-      const elements: string[] = [];
-      if (!uniqueConfigs || !uniqueTargets || !uniqueRHI) {
-         if (config && !uniqueConfigs) {
-            elements.push(cconfigs);
-         }
-         if (target && !uniqueTargets) {
-            elements.push(ctargets);
-         }
-         if (rhi && !uniqueRHI) {
-            elements.push(crhi);
-         }
-         if (variation && !uniqueVariation) {
-            elements.push(cvariation);
-         }
-      }
-
-      return {
-         commonConfigs: !uniqueConfigs,
-         commonTargets: !uniqueTargets,
-         commonRHI: !uniqueRHI,
-         commonVariation: !uniqueVariation,
-         commonMetaString: elements.length ? elements.join(" / ") : ""
-      }
-   }
-
-   async query() {
-
-      const state = this.state;
-
-      if (!state.automation || !state.streams?.length || (!state.tests?.length && !state.suites?.length)) {
-         return;
-      }
-
-      this.setQueryLoading(true);
-
-      const testIds: Set<string> = new Set();
-      const suiteIds: Set<string> = new Set();
-      const metaIds: Set<string> = new Set(this.currentMeta.map(m => m.id));
-
-      const streamTestMap: Map<string, GetTestResponse> = new Map();
-      this.streamTests.forEach(t => {
-         streamTestMap.set(t.id, t);
-      });
-
-      const streamSuiteMap: Map<string, GetTestResponse> = new Map();
-      this.streamSuites.forEach(s => {
-         streamSuiteMap.set(s.id, s);
-      });
-
-      state.tests?.forEach(name => {
-         streamTestMap.forEach((t, id) => {
-            if (t.name === name || t.displayName === name) {
-               testIds.add(id);
-            }
-         });
-      });
-
-      state.suites?.forEach(name => {
-         streamSuiteMap.forEach((s, id) => {
-            if (s.name === name) {
-               suiteIds.add(id);
-            }
-         });
-      });
-
-
-      // notify user?
-      if (!metaIds.size || (!testIds.size && !suiteIds.size)) {
-         this.refs = [];
-         this.streamStatus = new Map();
-         this.setQueryLoading(false);
-         this.setQueryUpdated();
-         return;
-      }
-
-      let queryWeeks = this.state.weeks;
-      if (typeof (queryWeeks) !== "number") {
-         queryWeeks = defaultQueryWeeks;
-      }
-
-      const minQueryDate = new Date(new Date().valueOf() - (minuteInWeek * queryWeeks * 60000));
-      const maxQueryDate = new Date();
-
-      let refs = await backend.getTestRefs(state.streams, Array.from(metaIds), minQueryDate.toISOString(), maxQueryDate.toISOString(), undefined, undefined, testIds.size ? Array.from(testIds) : undefined, suiteIds.size ? Array.from(suiteIds) : undefined);
-
-      // deduplicate and sort intro stream -> test -> meta -> refs descending by CL
-      const streamStatus = this.streamStatus;
-      streamStatus.clear();
-      const dedupedRefs: GetTestDataRefResponse[] = [];
-      const dupeWarnings = new Set<string>();
-
-      this.changeDates = new Map();
-
-      refs.forEach(r => {
-
-         if (!this.idDates.get(r.id)) {
-            const timestamp = r.id.substring(0, 8)
-            this.idDates.set(r.id, new Date(parseInt(timestamp, 16) * 1000));
-         }
-
-         const date = this.changeDates.get(r.buildChangeList);
-         const rdate = this.idDates.get(r.id)!;
-         if (!date || date.getTime() > rdate.getTime()) {
-            this.changeDates.set(r.buildChangeList, rdate);
-         }
-
-         let streamRefs = streamStatus.get(r.streamId);
-         if (!streamRefs) {
-            streamRefs = new Map();
-            streamStatus.set(r.streamId, streamRefs);
-         }
-
-         const id = r.testId ?? r.suiteId!;
-         let testRefs = streamRefs.get(id);
-         if (!testRefs) {
-            testRefs = new Map();
-            streamRefs.set(id, testRefs);
-         }
-
-         let metaStatus = testRefs.get(r.metaId);
-         if (!metaStatus) {
-            metaStatus = { refs: [], suite: !!r.suiteId };
-            testRefs.set(r.metaId, metaStatus);
-         }
-
-         let cref = metaStatus.refs.find(tr => r.buildChangeList === tr.buildChangeList);
-         if (cref) {
-            // todo: better replace logic, use latest, should elevate errors/skips/etc
-            dupeWarnings.add(id);
-         } else {
-            metaStatus.refs.push(r);
-            dedupedRefs.push(r);
-         }
-      });
-
-      this.minRefDate = new Date();
-      this.maxRefDate = new Date();
-
-      const timeRefs = refs.sort((a, b) => this.idDates.get(a.id)!.getTime() - this.idDates.get(b.id)!.getTime());
-      if (timeRefs.length === 1) {
-         this.minRefDate = this.maxRefDate = this.idDates.get(timeRefs[0].id)!;
-      } else if (timeRefs.length > 1) {
-         this.minRefDate = this.idDates.get(timeRefs[0].id)!;
-         this.maxRefDate = this.idDates.get(timeRefs[timeRefs.length - 1].id)!;
-      }
-
-      for (const [, testMap] of streamStatus) {
-         for (const [, metaMap] of testMap) {
-            for (const [, metaStatus] of metaMap) {
-
-               const suite = metaStatus.suite;
-               metaStatus.refs = metaStatus.refs.sort((a, b) => a.buildChangeList - b.buildChangeList).reverse();
-
-               // errors
-               for (const ref of metaStatus.refs) {
-                  const error = suite ? !!ref.suiteErrorCount : ref.outcome === "Failure";
-                  if (!error) {
-                     break;
-                  }
-                  metaStatus.error = metaStatus.error ? metaStatus.error + 1 : 1;
-               }
-
-               // skip
-               for (const ref of metaStatus.refs) {
-                  const skip = suite ? !!ref.suiteSkipCount : ref.outcome === "Skipped";
-                  if (!skip) {
-                     break;
-                  }
-                  metaStatus.skip = metaStatus.skip ? metaStatus.skip + 1 : 1;
-               }
-
-               // success
-               for (const ref of metaStatus.refs) {
-                  const success = suite ? (!ref.suiteSkipCount && !ref.suiteErrorCount) : (ref.outcome !== "Failure" && ref.outcome !== "Skipped");
-                  if (!success) {
-                     break;
-                  }
-                  metaStatus.success = metaStatus.success ? metaStatus.success + 1 : 1;
-               }
-            }
-         }
-      }
-
-      // generate warnings for non unique tests in stream
-      for (let id of dupeWarnings) {
-         if (this.testMap.get(id)) {
-            console.warn(`Test ${this.testMap.get(id)!.name} has duplicate refs for stream changelists and may require a more specific name`);
-         }
-         if (this.suiteMap.get(id)) {
-            console.warn(`Suite ${this.suiteMap.get(id)!.name} has duplicate refs for stream changelists and may require a more specific name`);
-         }
-      }
-
-      this.refs = dedupedRefs;
-
-      const changelists = this.refs.map(r => r.buildChangeList);
-
-      this.minCL = this.maxCL = 0;
-
-      this.minCL = Math.min(...changelists);
-      this.maxCL = Math.max(...changelists);
-
-      this.hasQueried = true;
-
-      this.setQueryLoading(false);
-      this.setQueryUpdated();
-
-   }
-
-   // Loads immutable data for this view session
-   async load() {
-
-      const allStreams = projectStore.projects.map(p => p.streams).flat().filter(s => !!s).map(s => s!.id);
-      //const streams = testingStreams;
-      const testStreams = await backend.getTestStreams(allStreams);
-
-      // get unique meta, tests, and suites
-      const metaMap: Map<string, GetTestMetaResponse> = new Map();
-
-      const testMap = this.testMap;
-      const suiteMap = this.suiteMap;
-
-      testMap.clear();
-      suiteMap.clear();
-
-      testStreams.forEach(stream => {
-
-         // it looks like there was a misconfiguration, need a way to remove or age out tests
-         stream.tests = stream.tests.filter(t => !t.name.startsWith("EditorAutomation") && !t.name.startsWith("TargetAutomation"));
-
-         this.testStreams.set(stream.streamId, stream);
-
-         const desktop = ["Win64", "Mac", "Linux"];
-         stream.testMetadata.forEach(m => {
-
-            m.platforms = m.platforms.sort((a, b) => {
-               if (desktop.indexOf(a) !== -1) {
-                  return 1;
-               }
-               if (desktop.indexOf(b) !== -1) {
-                  return -1;
-               }
-
-               return a.localeCompare(b);
-
-            });
-
-            if (!metaMap.has(m.id)) {
-               metaMap.set(m.id, m);
-            }
-         });
-
-         stream.tests.forEach(t => {
-            if (!testMap.has(t.id)) {
-               testMap.set(t.id, t);
-            }
-         });
-
-         stream.testSuites.forEach(s => {
-            if (!suiteMap.has(s.id)) {
-               suiteMap.set(s.id, s);
-            }
-         });
-      });
-
-      const metaData = Array.from(metaMap.values());
-      this.tests = Array.from(testMap.values()).map(t => t.displayName ?? t.name).sort((a, b) => a.localeCompare(b));
-      this.suites = Array.from(suiteMap.values()).map(s => s.name).sort((a, b) => a.localeCompare(b));
-
-      this.metaData = metaMap;
-      metaMap.forEach((meta, id) => {
-         if (this.metaNames.has(id)) {
-            return;
-         }
-
-         const metaName = `${meta.platforms.join(" - ")} / ${meta.configurations.join(" - ")} / ${meta.buildTargets.join(" - ")} / ${meta.rhi} / ${meta.variation}`;
-
-         this.metaNames.set(id, metaName);
-
-      });
-
-      const automation = new Set<string>();
-      const platforms = new Set<string>();
-      const configurations = new Set<string>();
-      const targets = new Set<string>();
-      const rhi = new Set<string>();
-      const variation = new Set<string>();
-
-      metaData.forEach(m => {
-         automation.add(m.projectName);
-         m.platforms.forEach(p => platforms.add(p));
-         m.buildTargets.forEach(t => targets.add(t));
-         m.configurations.forEach(c => configurations.add(c));
-         rhi.add(m.rhi);
-         variation.add(m.variation)
-      });
-
-      this.automation = Array.from(automation).sort((a, b) => a.localeCompare(b));
-      this.platforms = Array.from(platforms).sort((a, b) => a.localeCompare(b));
-      this.configurations = Array.from(configurations).sort((a, b) => a.localeCompare(b));
-      this.targets = Array.from(targets).sort((a, b) => a.localeCompare(b));
-      this.rhi = Array.from(rhi).sort((a, b) => a.localeCompare(b));
-      this.variation = Array.from(variation).sort((a, b) => a.localeCompare(b));
-
-      this.loaded = true;
-
-      const state = this.state;
-      if (!state.platforms?.length && this.platforms.length) {
-         this.platforms.forEach(p => this.addPlatform(p));
-      }
-      if (!state.configurations?.length && this.configurations.length) {
-         this.configurations.forEach(c => this.addConfiguration(c));
-      }
-      if (!state.targets?.length && this.targets.length) {
-         this.targets.forEach(t => this.addTarget(t));
-      }
-      if (!state.rhi?.length && this.rhi.length) {
-         this.rhi.forEach(r => this.addRHI(r));
-      }
-      if (!state.variation?.length && this.variation.length) {
-         this.variation.forEach(v => this.addVariation(v));
-      }
-
-      this.setUpdated();
-
-      this.query()
-   }
-
-   @action
-   setUpdated() {
-      this.updated++;
-   }
-
-   @action
-   setQueryUpdated() {
-      this.queryUpdated++;
-   }
-
-   @action
-   setQueryLoading(value: boolean) {
-      this.queryLoading = value;
-   }
-
-
-   setFilterState(nstate?: FilterState) {
-      if (this.filterState !== nstate) {
-         this.filterState = nstate;
-         this.setUpdated();
-         // need to update the center view
-         this.setQueryUpdated();
-      }
-   }
-
-   getAutomationStreams(automation: string): StreamId[] {
-
-      const streams: Set<StreamId> = new Set();
-
-      this.testStreams.forEach(testStream => {
-         if (testStream.testMetadata.find(t => t.projectName === automation)) {
-            streams.add(testStream.streamId);
-         }
-
-      });
-
-      return Array.from(streams);
-   }
-
-   getFilterState(): FilterState | undefined {
-      return this.filterState;
-   }
-
-   getRef(refId: string) {
-      return this.refs.find(r => r.id === refId);
-   }
-
-   getSuiteRefs(suiteId: string, metaId: string) {
-      return this.refs.filter(r => r.suiteId === suiteId && r.metaId === metaId);
-   }
-
-   getFilteredRefs(streamIdIn?: StreamId, idIn?: TestId | SuiteId, metaIdIn?: MetaId): GetTestDataRefResponse[] {
-
-      let status: TestStatus[] = [];
-
-      for (const [streamId, testMap] of this.streamStatus) {
-         if (streamIdIn && streamId !== streamIdIn) {
-            continue;
-         }
-         for (const [id, metaMap] of testMap) {
-            if (idIn && id !== idIn) {
-               continue;
-            }
-            for (const [metaId, metaStatus] of metaMap) {
-               if (metaIdIn && metaId !== metaIdIn) {
-                  continue;
-               }
-               status.push(metaStatus);
-            }
-         }
-      }
-
-      if (!this.filterState) {
-         return status.map(m => m.refs).flat();
-      }
-
-      if (this.filterState === "Failed") {
-         return status.filter(s => !!s.error).map(m => m.refs).flat();
-      }
-
-      if (this.filterState === "Consecutive Failures") {
-         return status.filter(s => s.error && s.error > 1).map(m => m.refs).flat();
-      }
-
-      if (this.filterState === "Skipped") {
-         return status.filter(s => !!s.skip).map(m => m.refs).flat();
-      }
-
-      if (this.filterState === "Success") {
-         return status.filter(s => !!s.success).map(m => m.refs).flat();
-      }
-
-      return [];
-
-   }
-
-   @observable
-   updated: number = 0;
-
-   @observable
-   queryUpdated: number = 0;
-
-   @observable
-   queryLoading: boolean = false;
-
-   state: TestDataState = {};
-   search: URLSearchParams = new URLSearchParams();
-
-   testStreams: Map<string, GetTestStreamResponse> = new Map();
-
-   metaData: Map<MetaId, GetTestMetaResponse> = new Map();
-   metaNames: Map<MetaId, string> = new Map();
-
-   testMap: Map<TestId, GetTestResponse> = new Map();
-   suiteMap: Map<SuiteId, GetTestSuiteResponse> = new Map();
-
-   automation: string[] = [];
-   platforms: string[] = [];
-   configurations: string[] = [];
-   targets: string[] = [];
-   rhi: string[] = [];
-   variation: string[] = [];
-
-   tests: string[] = [];
-   suites: string[] = [];
-
-   private refs: GetTestDataRefResponse[] = [];
-   minCL: number = 0;
-   maxCL: number = 0;
-   minRefDate: Date = new Date();
-   maxRefDate: Date = new Date();
-   idDates = new Map<string, Date>();
-   // quantized so common CL -> lands at same time in views
-   changeDates = new Map<number, Date>();
-
-   private filterState?: FilterState;
-   private streamStatus: Map<StreamId, Map<TestId | SuiteId, Map<MetaId, TestStatus>>> = new Map();
-
-   loaded = false;
-
-   hasQueried = false;
-
-   suiteRef?: string;
-
-   private queryTimeoutId: any = undefined;
-}
-
 const StreamChooser: React.FC<{ handler: TestDataHandler }> = observer(({ handler }) => {
 
    // subscribe
@@ -1422,7 +107,7 @@ const AutomationChooser: React.FC<{ handler: TestDataHandler }> = observer(({ ha
    const options: IContextualMenuItem[] = [];
 
    handler.automation.forEach(a => {
-      options.push({ key: `automation_${a}`, text: a, onClick: () => handler.setAutomation(a) });
+      options.push({ key: `automation_${a}`, text: a, onClick: () => { multiComboBoxId++; handler.setAutomation(a); } });
    });
 
    const menuProps: IContextualMenuProps = {
@@ -1486,7 +171,7 @@ const MultiOptionChooser: React.FC<{ options: IComboBoxOption[], initialKeysIn:
 
    const comboBoxStyles: Partial<IComboBoxStyles> = { root: { width: 270 } };
 
-   return <ComboBox key={`multi_option_id_${multiComboBoxId}` } placeholder="None" defaultSelectedKey={initialKeys} multiSelect options={options} onChange={onChange} styles={comboBoxStyles} />
+   return <ComboBox key={`multi_option_id_${multiComboBoxId}`} placeholder="None" defaultSelectedKey={initialKeys} multiSelect options={options} onChange={onChange} styles={comboBoxStyles} />
 };
 
 const TestChooser: React.FC<{ handler: TestDataHandler }> = observer(({ handler }) => {
diff --git a/Engine/Source/Programs/Horde/HordeDashboard/src/components/jobDetailsV2/JobDetailViewV2.tsx b/Engine/Source/Programs/Horde/HordeDashboard/src/components/jobDetailsV2/JobDetailViewV2.tsx
--- a/Engine/Source/Programs/Horde/HordeDashboard/src/components/jobDetailsV2/JobDetailViewV2.tsx
+++ b/Engine/Source/Programs/Horde/HordeDashboard/src/components/jobDetailsV2/JobDetailViewV2.tsx
@@ -1,4 +1,4 @@
-import { mergeStyleSets, Pivot, PivotItem, Stack, Text } from "@fluentui/react";
+import { mergeStyleSets, Pivot, PivotItem, Stack } from "@fluentui/react";
 import { observer } from "mobx-react-lite";
 import { useEffect, useRef, useState } from "react";
 import { useLocation, useNavigate, useParams } from "react-router-dom";
