Horde: Fix issue with updated Markdown component
#rnx
#skipci

diff --git a/Engine/Source/Programs/Horde/HordeDashboard/src/base/components/Markdown/Markdown.tsx b/Engine/Source/Programs/Horde/HordeDashboard/src/base/components/Markdown/Markdown.tsx
--- a/Engine/Source/Programs/Horde/HordeDashboard/src/base/components/Markdown/Markdown.tsx
+++ b/Engine/Source/Programs/Horde/HordeDashboard/src/base/components/Markdown/Markdown.tsx
@@ -28,8 +28,8 @@ const getStyles: IStyleFunction<IMarkdownStyleProps, IMarkdownStyles> = () => {
   };
 
   return {
-    root: 'ms-Markdown',
-    subComponentStyles: { image: imageStyles } as IMarkdownSubComponentStyles,
+     root: { fontSize: "13px", fontFamily: "Horde Open Sans Regular" },
+     subComponentStyles: { image: imageStyles } as IMarkdownSubComponentStyles,
   };
 };
 
