Horde: Remove trends view x-ray in favor of new timeline view
#skipci
#rnx

diff --git a/Engine/Source/Programs/Horde/HordeDashboard/src/components/jobDetailsV2/JobDetailStepTrends.tsx b/Engine/Source/Programs/Horde/HordeDashboard/src/components/jobDetailsV2/JobDetailStepTrends.tsx
--- a/Engine/Source/Programs/Horde/HordeDashboard/src/components/jobDetailsV2/JobDetailStepTrends.tsx
+++ b/Engine/Source/Programs/Horde/HordeDashboard/src/components/jobDetailsV2/JobDetailStepTrends.tsx
@@ -1,6 +1,6 @@
 // Copyright Epic Games, Inc. All Rights Reserved.
 
-import { ContextualMenu, Dialog, DialogType, IconButton, IIconProps, Spinner, SpinnerSize, Stack, Text, Toggle } from '@fluentui/react';
+import { IconButton, IIconProps, Spinner, SpinnerSize, Stack, Text, Toggle } from '@fluentui/react';
 import { IChartProps, ILineChartDataPoint } from '@fluentui/react-charting';
 import { observer } from 'mobx-react-lite';
 import moment from 'moment';
@@ -365,42 +365,12 @@ type ChartDataCache = {
    chartData?: ChartData;
 }
 
-type XRayBatch = {
-   dependencyWaitTime: number;
-   waitTime: number;
-   setupTime: number;
-   reserveTime: moment.Moment;
-   startTime: moment.Moment;
-   finishTime: moment.Moment;
-   steps: GraphHelperStep[];
-}
-
-type XRayColumn = {
-   minX: number;
-   maxX: number;
-   batches: XRayBatch[];
-}
-
-type XRayCanvasData = {
-   jobStartTime: moment.Moment;
-   jobEndTime: moment.Moment;
-   minX: number;
-   maxX: number;
-   minY: number;
-   maxY: number;
-   canvasWidth: number;
-   canvasHeight: number;
-   columns: XRayColumn[];
-   change: number;
-}
-
 type PageInfo = {
    pageNumber: number;
    zoomLevel: number;
    pageSize: number;
 }
 
-let xRayHoveredJobStep: GraphHelperStep | undefined = undefined;
 function preventChartScroll(this: Document, e: WheelEvent) {
    e = e || window.event
    if (e.preventDefault) {
@@ -460,12 +430,8 @@ const StepTrendsPanelInner: React.FC<{ jobDetails: JobDetailsV2; dataView: Trend
    const [includeWaitTimes, setIncludeWaitTimes] = useState<boolean>(false);
    const [includeStepDependencies, setIncludeStepDependencies] = useState<boolean>(false);
 
-   const [xRayData, setXRayData] = useState<XRayCanvasData | undefined>(undefined);
-   const [canvas, setCanvas] = useState<HTMLCanvasElement | null>(null);
-
    const [pageInfo, setPageInfo] = useState<PageInfo>({ pageNumber: 1, zoomLevel: 1, pageSize: 10 });
 
-
    dataView.subscribe();
 
    const findTemplate = useCallback(async () => {
@@ -483,7 +449,6 @@ const StepTrendsPanelInner: React.FC<{ jobDetails: JobDetailsV2; dataView: Trend
 
 
    if (!selectedTemplateRef) {
-      xRayHoveredJobStep = undefined;
       findTemplate();
       return null;
    }
@@ -494,411 +459,7 @@ const StepTrendsPanelInner: React.FC<{ jobDetails: JobDetailsV2; dataView: Trend
       };
       dataView.populateChartDataCache(query);
    }
-
-   const selectDataPoint = (label: GraphHelperLabel | undefined, steps: GraphHelperStep[] | undefined) => {
-      if (!label && !steps) {
-         setXRayData(undefined);
-         return;
-      }
-      function tryInsertIntoColumn(column: XRayColumn, batch: XRayBatch) {
-         let idx = 0;
-         for (; idx < column.batches.length; idx++) {
-            let otherBatch = column.batches[idx];
-            if (batch.reserveTime <= otherBatch.finishTime) {
-               if (batch.finishTime < otherBatch.reserveTime) {
-                  break;
-               }
-               else {
-                  return false;
-               }
-            }
-         }
-
-         column.batches.splice(idx, 0, batch);
-         return true;
-      }
-      const canvasData: XRayCanvasData = { jobStartTime: moment(), jobEndTime: moment(), minX: 0, maxX: 0, minY: 0, maxY: 0, canvasHeight: 0, canvasWidth: 0, columns: [], change: 0 };
-      //const stepStartTimes: moment.Moment[] = [];
-      const stepEndTimes: moment.Moment[] = [];
-      if (steps) {
-         canvasData.change = steps[0].job.change!;
-         //each batch should be a column. make batch -> step dictionary
-         let batches: any = {};
-         steps.forEach(step => {
-            if (!batches[step.batch.id]) {
-               batches[step.batch.id] = { batch: step.batch, steps: [] };
-            }
-            batches[step.batch.id].steps.push(step);
-
-            //stepStartTimes.push(moment.utc(step.step.startTime));
-            stepEndTimes.push(moment.utc(step.step.finishTime));
-         });
-
-         //const jobStartTime = moment.min(stepStartTimes);
-         const jobEndTime = moment.max(stepEndTimes);
-
-         canvasData.jobStartTime = moment.utc(steps[0].job.createTime);
-         canvasData.jobEndTime = jobEndTime;
-
-
-         // order the columns by groupidx, then the steps by nodeidx
-         let sortedBatches = Object.entries(batches).sort(([, a], [, b]) => (a as any).batch.groupIdx - (b as any).batch.groupIdx);
-         const defaultColumnWidth = 2000 / sortedBatches.length;
-         sortedBatches.forEach((sortedBatch, idx) => {
-            let width = Math.max(defaultColumnWidth, 210);
-            width = Math.min(defaultColumnWidth, 350);
-
-            let minX = 100 + idx * width;
-            let maxX = minX + (width * (200 / 210));
-            const sortedSteps: GraphHelperStep[] = Object.values((sortedBatch[1] as any).steps.sort((a: any, b: any) => a.step.nodeIdx - b.step.nodeIdx));
-
-            let totalDependencyTime = 0;
-            let waitTime = 0;
-            let setupTime = 0;
-            sortedSteps.forEach(step => {
-               let thisStep = step;
-               if (thisStep.setupTime[0] !== 0) {
-                  thisStep.dependencies?.forEach((dependency: any) => {
-                     totalDependencyTime += dependency.runTime;
-                     if (dependency.setupTime[0] !== 0) {
-                        totalDependencyTime += dependency.setupTime[0];
-                     }
-                     if (dependency.waitTime[0] !== 0) {
-                        totalDependencyTime += dependency.waitTime[0];
-                     }
-                  });
-               }
-
-               if (thisStep.waitTime[0] !== 0) {
-                  waitTime = thisStep.waitTime[0]!;
-               }
-               if (thisStep.setupTime[0] !== 0) {
-                  setupTime = thisStep.setupTime[0]!;
-               }
-            });
-
-            let reserveTime = moment.utc((sortedBatch[1] as any).batch.startTime);
-            let startTime = moment.utc(reserveTime).add(setupTime, 'seconds');
-
-            let newBatch: XRayBatch = {
-               reserveTime: reserveTime,
-               startTime: startTime,
-               finishTime: moment.utc((sortedBatch[1] as any).batch.finishTime),
-               setupTime: setupTime,
-               dependencyWaitTime: totalDependencyTime,
-               waitTime: waitTime,
-               steps: sortedSteps
-            };
-
-            for (let i = 0; ; i++) {
-               if (i === canvasData.columns.length) {
-                  canvasData.columns.push({ minX: minX, maxX: maxX, batches: [] });
-               }
-               if (tryInsertIntoColumn(canvasData.columns[i], newBatch)) {
-                  break;
-               }
-            }
-         });
-
-         canvasData.minX = 60.5;
-         canvasData.minY = 5;
-         canvasData.maxX = Math.max(canvasData.columns[canvasData.columns.length - 1].maxX, 1800);
-         canvasData.maxY = canvasData.minY + 50 + canvasData.jobEndTime.diff(canvasData.jobStartTime, 'seconds', true);
-
-         canvasData.canvasWidth = (canvasData.maxX + canvasData.minX);
-         canvasData.canvasHeight = (canvasData.maxY + canvasData.minY);
-      }
-
-      setXRayData(canvasData);
-   }
-
-   const onCanvasMouseMove = (ev: React.MouseEvent<HTMLCanvasElement, MouseEvent>) => {
-      if (xRayData && canvas) {
-         let rect = canvas.getBoundingClientRect();
-         let scaleX = canvas.width / rect.width;
-         let scaleY = canvas.height / rect.height;
-
-         let xPos = (ev.clientX - rect.left) * scaleX;
-         let yPos = (ev.clientY - rect.top) * scaleY;
-
-         xRayHoveredJobStep = undefined;
-         for (let colIdx = 0; colIdx < xRayData.columns.length; colIdx++) {
-            let column = xRayData.columns[colIdx];
-            if (xPos >= column.minX && xPos <= column.maxX) {
-               for (let batchIdx = 0; batchIdx < column.batches.length; batchIdx++) {
-                  let batch = column.batches[batchIdx];
-                  for (let stepIdx = 0; stepIdx < batch.steps.length; stepIdx++) {
-                     let step = batch.steps[stepIdx];
-                     let stepMinY = xRayData.minY + moment.utc(step.step.startTime).diff(xRayData.jobStartTime, 'seconds', true) + 1.0;
-                     let stepMaxY = xRayData.minY + moment.utc(step.step.finishTime).diff(xRayData.jobStartTime, 'seconds', true) - 1.0;
-                     if (yPos >= stepMinY && yPos <= stepMaxY) {
-                        if (xRayHoveredJobStep !== step) {
-                           xRayHoveredJobStep = step;
-                           break;
-                        }
-                     }
-                  }
-               }
-            }
-         }
-      }
-      drawXRayChart();
-   };
-
-   const drawXRayChart = () => {
-      function drawRoundedRect(context: CanvasRenderingContext2D, minX: number, minY: number, maxX: number, maxY: number, radius: number) {
-         if (maxX > minX && maxY > minY) {
-            radius = Math.min(radius, Math.min(maxY - minY, maxX - minX));
-
-            context.beginPath();
-            context.moveTo(minX + radius, minY);
-            context.arc(maxX - radius, minY + radius, radius, -Math.PI * 0.5, 0.0);
-            context.arc(maxX - radius, maxY - radius, radius, 0.0, +Math.PI * 0.5);
-            context.arc(minX + radius, maxY - radius, radius, +Math.PI * 0.5, +Math.PI * 1.0);
-            context.arc(minX + radius, minY + radius, radius, -Math.PI * 1.0, -Math.PI * 0.5);
-            context.closePath();
-            context.fill();
-         }
-      }
-      function drawTimeBox(context: CanvasRenderingContext2D, minX: number, minY: number, maxX: number, maxY: number, color: string, text: string, textColor: string, underline: boolean) {
-         context.fillStyle = color;
-         drawRoundedRect(context, minX, minY, maxX, maxY, 5.0);
-
-         if ((maxY - minY) > 20) {
-            context.font = "11 px sans-serif";
-
-            let metrics: TextMetrics | undefined = undefined;
-            for (; ;) {
-               metrics = context.measureText(text);
-               if (metrics.width + 30 < maxX - minX) {
-                  break;
-               }
-
-               let idx = text.length;
-               while (idx >= 0 && text.charAt(idx) !== ' ') {
-                  idx--;
-               }
-               if (idx === 0) {
-                  break;
-               }
-               text = text.substring(0, idx) + "...";
-            }
-
-            context.fillStyle = textColor;
-            context.textAlign = 'center';
-            context.textBaseline = 'middle';
-            context.fillText(text, (minX + maxX) * 0.5, (minY + maxY) * 0.5);
-
-            if (underline) {
-               context.beginPath();
-               context.strokeStyle = "white";
-               context.moveTo((minX + maxX - metrics.width) * 0.5, (minY + maxY) * 0.5 + 6.0);
-               context.lineTo((minX + maxX + metrics.width) * 0.5, (minY + maxY) * 0.5 + 6.0);
-               context.stroke();
-            }
-         }
-      }
-      function drawAnnotation(context: CanvasRenderingContext2D, x: number, y: number, align: string, color: string, text: string, textColor: string) {
-         context.font = "11 px sans-serif";
-         let metrics = context.measureText(text);
-
-         let padding = 15;
-         let width = padding * 2 + metrics.width;
-         let height = 20;
-
-         let minX;
-         if (align === 'right') {
-            minX = x - width;
-         }
-         else if (align === 'left') {
-            minX = x;
-         }
-         else {
-            minX = x - width * 0.5;
-         }
-
-         let minY = y - height * 0.5;
-
-         context.fillStyle = "rgba(255,255,255,0.3)";
-         let border = 2.0;
-         drawRoundedRect(context, minX - border, minY - border, minX + width + border, minY + height + border, 10.0 + border);
-
-         context.fillStyle = color;
-         drawRoundedRect(context, minX, minY, minX + width, minY + height, 10.0);
-
-         context.fillStyle = textColor;
-         context.textAlign = 'center';
-         context.textBaseline = 'middle';
-         context.fillText(text, minX + width * 0.5, minY + height * 0.5);
-      }
-      if (canvas && xRayData) {
-         const context = canvas.getContext("2d")!;
-         canvas.width = xRayData.canvasWidth;
-         canvas.height = xRayData.canvasHeight;
-         context.clearRect(0, 0, canvas.width, canvas.height);
-
-
-
-         // draw left arrow
-         context.beginPath();
-         context.strokeStyle = "#FFFFFF";
-         context.lineWidth = 1;
-
-         context.moveTo(xRayData.minX - 5, xRayData.minY);
-         context.lineTo(xRayData.minX + 5, xRayData.minY);
-         context.moveTo(xRayData.minX, xRayData.minY);
-         context.lineTo(xRayData.minX, xRayData.maxY);
-         context.moveTo(xRayData.minX - 5, xRayData.maxY - 5);
-         context.lineTo(xRayData.minX, xRayData.maxY);
-         context.lineTo(xRayData.minX + 5, xRayData.maxY - 5);
-         context.stroke();
-
-         let hoverJobStepY = 0;
-         if (xRayHoveredJobStep !== undefined) {
-            let runTime = moment.utc(xRayHoveredJobStep.step.finishTime).diff(xRayData.jobStartTime, 'seconds', true);
-            hoverJobStepY = xRayData.minY + runTime - 1.0;
-
-            context.beginPath();
-            context.strokeStyle = "#000000";
-            context.moveTo(xRayData.minX - 5, hoverJobStepY);
-            context.lineTo(xRayData.minX, hoverJobStepY);
-            context.stroke();
-            context.closePath();
-
-            context.beginPath();
-            context.strokeStyle = "#efefef";
-            context.moveTo(xRayData.minX, hoverJobStepY);
-            context.lineTo(xRayData.maxX, hoverJobStepY);
-            context.stroke();
-            context.closePath();
-
-            let elapsedTime = moment.utc(runTime * 1000).format('HH:mm:ss');
-
-            context.font = "10 px sans-serif";
-            context.fillStyle = "white";
-            context.textAlign = 'center';
-            context.textBaseline = 'middle';
-            context.fillText(elapsedTime, xRayData.minX - 8, hoverJobStepY);
-         }
-
-         // Draw interval stuff?
-         for (let intervalTime = 0; ; intervalTime += 10 * 60) {
-            const y = (xRayData.minY + (intervalTime)) + 0.5;
-            if (y > xRayData.maxY - 20) {
-               break;
-            }
-
-            context.beginPath();
-            context.strokeStyle = "#e0e0e0";
-            context.moveTo(xRayData.minX, y);
-            context.lineTo(xRayData.maxX, y);
-            context.stroke();
-
-            if (intervalTime > 0 && (xRayHoveredJobStep !== undefined || Math.abs(y - hoverJobStepY) > 15)) {
-               context.beginPath();
-               context.strokeStyle = "#FFFFFF";
-               context.moveTo(xRayData.minX - 5, y);
-               context.lineTo(xRayData.minX, y);
-               context.stroke();
-
-               let elapsedTime = moment.utc(intervalTime * 1000).format('HH:mm:ss');
-
-               context.font = "10 px sans-serif";
-               context.fillStyle = "white";
-               context.textAlign = 'right';
-               context.textBaseline = 'middle';
-               context.fillText(elapsedTime, xRayData.minX - 8, y);
-            }
-         }
-
-         // Figure out all the dependencies of the steps we're hovering over
-         let selectedDependencies: GraphHelperStep[] | undefined = undefined;
-         if (xRayHoveredJobStep !== undefined) {
-            selectedDependencies = [];
-            selectedDependencies.push(xRayHoveredJobStep);
-            for (let idx = 0; idx < selectedDependencies.length; idx++) {
-               let jobStep = selectedDependencies[idx];
-               jobStep.dependencies?.forEach(dependency => {
-                  if (!selectedDependencies!.includes(dependency)) {
-                     selectedDependencies!.push(dependency);
-                  }
-               });
-
-            }
-         }
-
-         const colors = ["35AAF2", "DE6A10", "F6C000", "00882B", "C82506", "773F9B"];
-
-         // TODO: Hover junk
-         let hoverColumn: XRayColumn | undefined = undefined;
-         let hoverColor: string | undefined = undefined;
-
-         // Draw them all
-         let colorIdx = 0;
-         for (let columnIdx = xRayData.columns.length - 1; columnIdx >= 0; columnIdx--) {
-            const column = xRayData.columns[columnIdx];
-            for (let batchIdx = column.batches.length - 1; batchIdx >= 0; batchIdx--) {
-
-               const batch = column.batches[batchIdx];
-
-               let setupMinY = xRayData.minY + batch.reserveTime.diff(xRayData.jobStartTime, 'seconds', true) + 1;//setupMaxY - batch.setupTime + 1;
-               let setupMaxY = xRayData.minY + moment.utc(batch.steps[0].step.startTime).diff(xRayData.jobStartTime, 'seconds', true) - 1;
-
-               drawTimeBox(context, column.minX, setupMinY, column.maxX, setupMaxY, "#ecedf1", "Setup Time", "#707070", false);
-
-               let baseColor = colors[colorIdx % colors.length];
-               for (let batchStep = batch.steps.length - 1; batchStep >= 0; batchStep--) {
-                  let step = batch.steps[batchStep];
-                  if (step === xRayHoveredJobStep) {
-                     hoverColumn = column;
-                     hoverColor = `#${baseColor}FF`;
-                  }
-                  else {
-                     let stepMinY = xRayData.minY + moment.utc(step.step.startTime).diff(xRayData.jobStartTime, 'seconds', true) + 1.0;
-                     let stepMaxY = xRayData.minY + moment.utc(step.step.finishTime).diff(xRayData.jobStartTime, 'seconds', true) - 1.0;
-
-                     let alpha = 'E6';
-                     if (selectedDependencies != null) {
-                        alpha = selectedDependencies.includes(step) ? 'E6' : '26';
-                     }
-
-
-                     let color = `#${baseColor}${alpha}`;
-                     drawTimeBox(context, column.minX, stepMinY, column.maxX, stepMaxY, color, step.name, "white", false);
-                  }
-               }
-               colorIdx++;
-            }
-
-         }
-
-         if (xRayHoveredJobStep !== undefined) {
-            let stepMinX = hoverColumn!.minX;
-            let stepMaxX = hoverColumn!.maxX;
-
-            let stepMinY = xRayData.minY + moment.utc(xRayHoveredJobStep.step.startTime).diff(xRayData.jobStartTime, 'seconds', true) + 1.0;
-            let stepMaxY = xRayData.minY + moment.utc(xRayHoveredJobStep.step.finishTime).diff(xRayData.jobStartTime, 'seconds', true) - 1.0;
-
-            let metrics = context.measureText(xRayHoveredJobStep.name);
-            let labelWidth = metrics.width + 50;
-            if (stepMaxX - stepMinX < labelWidth) {
-               let stepMidX = (stepMinX + stepMaxX) * 0.5;
-               stepMinX = stepMidX - labelWidth * 0.5;
-               stepMaxX = stepMidX + labelWidth * 0.5;
-            }
-
-            stepMaxY = Math.max(stepMinY + 16, stepMaxY);
-            context.fillStyle = 'white';
-            let overlap = 3.0;
-            drawRoundedRect(context, stepMinX - overlap, stepMinY - overlap, stepMaxX + overlap, stepMaxY + overlap, 5.0);
-
-            drawTimeBox(context, stepMinX, stepMinY, stepMaxX, stepMaxY, hoverColor!, xRayHoveredJobStep!.name, "white", true);
-
-            let stepDuration = moment(xRayHoveredJobStep.step.finishTime).diff(moment(xRayHoveredJobStep.step.startTime), 'seconds', true);
-            drawAnnotation(context, stepMaxX + 5, (stepMinY + stepMaxY) * 0.5, 'left', "#e0e0e0", moment.utc(stepDuration * 1000).format("HH:mm:ss"), "black");
-         }
-      }
-   };
+   
 
    const filterChartData = (chartDataCache?: ChartDataCache) => {
       if (!chartDataCache) {
@@ -1135,7 +696,6 @@ const StepTrendsPanelInner: React.FC<{ jobDetails: JobDetailsV2; dataView: Trend
                               y: labelRunTime,
                               yAxisCalloutData: calloutData,
                               xAxisCalloutData: `CL ${change.change}`,
-                              onDataPointClick: () => selectDataPoint(change, Array.from(changesToAllStepsMap[change.change]))
                            });
                         }
                      }
@@ -1211,9 +771,6 @@ const StepTrendsPanelInner: React.FC<{ jobDetails: JobDetailsV2; dataView: Trend
 
    let filteredChartState: FilteredChartData | undefined = filterChartData(dataView.cacheData);
 
-   if (canvas && xRayData) {
-      drawXRayChart();
-   }
 
    let placeholderDiv = <div></div>;
    let titleText = 'Trends';
@@ -1317,47 +874,14 @@ const StepTrendsPanelInner: React.FC<{ jobDetails: JobDetailsV2; dataView: Trend
          <Stack>
             <Stack horizontal horizontalAlign="space-between" styles={{ root: { minHeight: 32, paddingRight: 8 } }}>
                <Text variant="mediumPlus" styles={{ root: { fontFamily: "Horde Open Sans SemiBold" } }}>{titleText}</Text>
-               {step && !xRayData && <Toggle label="Include dependencies" checked={includeStepDependencies} onText="On" offText="Off" onChange={() => setIncludeStepDependencies(!includeStepDependencies)} />}
-               {!xRayData && <Toggle label="Include wait times" checked={includeWaitTimes} onText="On" offText="Off" onChange={() => setIncludeWaitTimes(!includeWaitTimes)} />}
-               {/* { !xRayData && <Toggle label="Include cost" checked={includeCost} onText="On" offText="Off" onChange={() => setIncludeCost(!includeCost)}/> } */}
+               {!!step && <Toggle label="Include dependencies" checked={includeStepDependencies} onText="On" offText="Off" onChange={() => setIncludeStepDependencies(!includeStepDependencies)} />}
+               <Toggle label="Include wait times" checked={includeWaitTimes} onText="On" offText="Off" onChange={() => setIncludeWaitTimes(!includeWaitTimes)} />               
             </Stack>
             <Stack styles={{ root: { position: 'relative', top: -10 } }}>
                {statusDiv}
             </Stack>
             <Stack styles={{ root: { width: 1720, minHeight: 400, height: 400 } }} onWheel={onChartScroll} onMouseEnter={onChartWheelEnter} onMouseLeave={onChartWheelLeave}>
                {placeholderDiv}
-               <Dialog
-                  modalProps={{
-                     isBlocking: false,
-                     dragOptions: {
-                        closeMenuItemText: "Close",
-                        moveMenuItemText: "Move",
-                        menu: ContextualMenu
-                     },
-                     styles: {
-                        root: {
-                           selectors: {
-                              ".ms-Dialog-title": {
-                                 paddingTop: '24px',
-                                 paddingLeft: '32px'
-                              }
-                           }
-                        }
-                     }
-                  }}
-                  onDismiss={() => { selectDataPoint(undefined, undefined); xRayHoveredJobStep = undefined; }}
-                  //className={historyStyles.dialog}
-                  hidden={xRayData === undefined}
-                  minWidth={1920}
-                  dialogContentProps={{
-                     type: DialogType.close,
-                     onDismiss: () => { selectDataPoint(undefined, undefined); xRayHoveredJobStep = undefined; },
-                     title: `X-Ray - CL ${xRayData ? xRayData.change : 0} `,
-                  }}
-               >
-                  <canvas ref={newRef => setCanvas(newRef)} width={1920} height={0} onMouseMove={(ev) => onCanvasMouseMove(ev)}>
-                  </canvas>
-               </Dialog>
             </Stack>
 
          </Stack>
