Horde: Fix exception sorting Slack reports to send.

diff --git a/Engine/Source/Programs/Horde/Horde.Build/Notifications/Sinks/SlackNotificationSink.cs b/Engine/Source/Programs/Horde/Horde.Build/Notifications/Sinks/SlackNotificationSink.cs
--- a/Engine/Source/Programs/Horde/Horde.Build/Notifications/Sinks/SlackNotificationSink.cs
+++ b/Engine/Source/Programs/Horde/Horde.Build/Notifications/Sinks/SlackNotificationSink.cs
@@ -1551,7 +1551,7 @@ namespace Horde.Build.Notifications.Sinks
 		{
 			GlobalConfig globalConfig = _globalConfig.CurrentValue;
 
-			foreach (IssueReport report in group.Reports.OrderBy(x => x.WorkflowId).ThenBy(x => x.StreamId))
+			foreach (IssueReport report in group.Reports.OrderBy(x => x.WorkflowId.Id.Text).ThenBy(x => x.StreamId.Id.Text))
 			{
 				await SendIssueReportForStreamAsync(globalConfig, group.Channel, group.Time, report);
 			}
