HordeAgent: Fixing Workspace command argument

WorkspaceBase commands class has a -Server argument which clashes with the base Command class -Server argument, renaming the WorkspaceBase to -ServerAndPort

#rnx

diff --git a/Engine/Source/Programs/Horde/Horde.Agent/Commands/Workspace/WorkspaceBase.cs b/Engine/Source/Programs/Horde/Horde.Agent/Commands/Workspace/WorkspaceBase.cs
--- a/Engine/Source/Programs/Horde/Horde.Agent/Commands/Workspace/WorkspaceBase.cs
+++ b/Engine/Source/Programs/Horde/Horde.Agent/Commands/Workspace/WorkspaceBase.cs
@@ -15,7 +15,7 @@ namespace Horde.Agent.Commands.Workspace
 {
 	abstract class WorkspaceBase : Command
 	{
-		[CommandLine("-Server")]
+		[CommandLine("-ServerAndPort")]
 		[Description("Specifies the Perforce server and port")]
 		protected string ServerAndPort { get; set; } = PerforceSettings.Default.ServerAndPort;
 
