Horde: Fix tests.

diff --git a/Engine/Source/Programs/Horde/Horde.Build.Tests/TestSetup.cs b/Engine/Source/Programs/Horde/Horde.Build.Tests/TestSetup.cs
--- a/Engine/Source/Programs/Horde/Horde.Build.Tests/TestSetup.cs
+++ b/Engine/Source/Programs/Horde/Horde.Build.Tests/TestSetup.cs
@@ -262,6 +262,8 @@ namespace Horde.Build.Tests
 			services.AddSingleton<DeviceService>();
 			services.AddSingleton<TestDataService>();
 
+			services.AddScoped<JobRpcCommon>();
+
 			services.AddSingleton<ConformTaskSource>();
 			services.AddSingleton<ICommitService, CommitService>();
 
diff --git a/Engine/Source/Programs/Horde/Horde.Build/Startup.cs b/Engine/Source/Programs/Horde/Horde.Build/Startup.cs
--- a/Engine/Source/Programs/Horde/Horde.Build/Startup.cs
+++ b/Engine/Source/Programs/Horde/Horde.Build/Startup.cs
@@ -331,7 +331,7 @@ namespace Horde.Build
 			services.AddGrpcReflection();
 
 			services.AddHttpClient<JobRpcCommon>();
-			services.AddSingleton<JobRpcCommon>();
+			services.AddScoped<JobRpcCommon>();
 
 			services.AddSingleton<IAgentCollection, AgentCollection>();
 			services.AddSingleton<IAgentSoftwareCollection, AgentSoftwareCollection>();
