Horde: Agent test fixes.

diff --git a/Engine/Source/Programs/Horde/Horde.Agent.Tests/Helper.cs b/Engine/Source/Programs/Horde/Horde.Agent.Tests/Helper.cs
--- a/Engine/Source/Programs/Horde/Horde.Agent.Tests/Helper.cs
+++ b/Engine/Source/Programs/Horde/Horde.Agent.Tests/Helper.cs
@@ -10,6 +10,7 @@ using Grpc.Net.Client;
 using Horde.Agent.Execution;
 using Horde.Agent.Services;
 using Horde.Agent.Utility;
+using Horde.Common.Rpc;
 using HordeCommon;
 using HordeCommon.Rpc;
 using HordeCommon.Rpc.Messages;
@@ -55,13 +56,15 @@ namespace Horde.Agent.Tests
 	{
 		private readonly GrpcChannel _grpcChannel;
 		private readonly HordeRpc.HordeRpcClient _hordeRpcClient;
+		private readonly JobRpc.JobRpcClient _jobRpcClient;
 
 		public ILogger Logger => NullLogger.Instance;
 
-		public RpcConnectionStub(GrpcChannel grpcChannel, HordeRpc.HordeRpcClient hordeRpcClient)
+		public RpcConnectionStub(GrpcChannel grpcChannel, HordeRpc.HordeRpcClient hordeRpcClient, JobRpc.JobRpcClient jobRpcClient)
 		{
 			_grpcChannel = grpcChannel;
 			_hordeRpcClient = hordeRpcClient;
+			_jobRpcClient = jobRpcClient;
 		}
 
 		public IRpcClientRef<TClient>? TryGetClientRef<TClient>() where TClient : ClientBase<TClient>
@@ -71,7 +74,19 @@ namespace Horde.Agent.Tests
 
 		public Task<IRpcClientRef<TClient>> GetClientRefAsync<TClient>(CancellationToken cancellationToken) where TClient : ClientBase<TClient>
 		{
-			IRpcClientRef<TClient> rpcClientRefStub = (IRpcClientRef<TClient>)(object)new RpcClientRefStub<HordeRpc.HordeRpcClient>(_grpcChannel, _hordeRpcClient);
+			IRpcClientRef<TClient> rpcClientRefStub;
+			if (typeof(TClient) == typeof(HordeRpc.HordeRpcClient))
+			{
+				rpcClientRefStub = (IRpcClientRef<TClient>)(object)new RpcClientRefStub<HordeRpc.HordeRpcClient>(_grpcChannel, _hordeRpcClient);
+			}
+			else if (typeof(TClient) == typeof(JobRpc.JobRpcClient))
+			{
+				rpcClientRefStub = (IRpcClientRef<TClient>)(object)new RpcClientRefStub<JobRpc.JobRpcClient>(_grpcChannel, _jobRpcClient);
+			}
+			else
+			{
+				throw new NotImplementedException();
+			}
 			return Task.FromResult(rpcClientRefStub);
 		}
 
@@ -81,7 +96,7 @@ namespace Horde.Agent.Tests
 		}
 	}
 
-	class HordeRpcClientStub : HordeRpc.HordeRpcClient
+	class JobRpcClientStub : JobRpc.JobRpcClient
 	{
 		public readonly Queue<BeginStepResponse> BeginStepResponses = new Queue<BeginStepResponse>();
 		public readonly List<UpdateStepRequest> UpdateStepRequests = new List<UpdateStepRequest>();
@@ -89,7 +104,7 @@ namespace Horde.Agent.Tests
 		public Func<GetStepRequest, GetStepResponse>? _getStepFunc = null;
 		private readonly ILogger _logger;
 
-		public HordeRpcClientStub(ILogger logger)
+		public JobRpcClientStub(ILogger logger)
 		{
 			_logger = logger;
 		}
diff --git a/Engine/Source/Programs/Horde/Horde.Agent.Tests/WorkerServiceTest.cs b/Engine/Source/Programs/Horde/Horde.Agent.Tests/WorkerServiceTest.cs
--- a/Engine/Source/Programs/Horde/Horde.Agent.Tests/WorkerServiceTest.cs
+++ b/Engine/Source/Programs/Horde/Horde.Agent.Tests/WorkerServiceTest.cs
@@ -21,6 +21,7 @@ using Horde.Agent.Parser;
 using Horde.Agent.Services;
 using Horde.Agent.Utility;
 using Horde.Common;
+using Horde.Common.Rpc;
 using HordeCommon;
 using HordeCommon.Rpc;
 using HordeCommon.Rpc.Messages;
@@ -147,8 +148,8 @@ namespace Horde.Agent.Tests
 			executeJobTask.AutoSdkWorkspace = new AgentWorkspace();
 			executeJobTask.Workspace = new AgentWorkspace();
 
-			HordeRpcClientStub client = new HordeRpcClientStub(NullLogger.Instance);
-			await using RpcConnectionStub rpcConnection = new RpcConnectionStub(null!, client);
+			JobRpcClientStub client = new JobRpcClientStub(NullLogger.Instance);
+			await using RpcConnectionStub rpcConnection = new RpcConnectionStub(null!, null!, client);
 
 			await using ISession session = FakeServerSessionFactory.CreateSession(rpcConnection);
 
@@ -200,8 +201,8 @@ namespace Horde.Agent.Tests
 			JobHandler jobHandler = serviceProvider.GetRequiredService<JobHandler>();
 			jobHandler._stepAbortPollInterval = TimeSpan.FromMilliseconds(5);
 
-			HordeRpcClientStub client = new HordeRpcClientStub(NullLogger.Instance);
-			await using RpcConnectionStub rpcConnection = new RpcConnectionStub(null!, client);
+			JobRpcClientStub client = new JobRpcClientStub(NullLogger.Instance);
+			await using RpcConnectionStub rpcConnection = new RpcConnectionStub(null!, null!, client);
 
 			int c = 0;
 			client._getStepFunc = (request) =>
@@ -295,7 +296,7 @@ namespace Horde.Agent.Tests
 		public readonly TaskCompletionSource<bool> UpdateSessionReceived = new();
 
 		private readonly RpcConnectionStub _connection;
-		private readonly FakeHordeRpcClient _client;
+		private readonly FakeJobRpcClient _client;
 
 		private class FakeHordeRpcClient : HordeRpc.HordeRpcClient
 		{
@@ -306,11 +307,26 @@ namespace Horde.Agent.Tests
 				_outer = outer;
 			}
 
+			public override AsyncDuplexStreamingCall<UpdateSessionRequest, UpdateSessionResponse> UpdateSession(Metadata headers = null!, DateTime? deadline = null, CancellationToken cancellationToken = default)
+			{
+				return _outer.GetUpdateSessionCall(CancellationToken.None);
+			}
+		}
+
+		private class FakeJobRpcClient : JobRpc.JobRpcClient
+		{
+			private readonly FakeHordeRpcServer _outer;
+
+			public FakeJobRpcClient(FakeHordeRpcServer outer)
+			{
+				_outer = outer;
+			}
+
 			public override AsyncUnaryCall<GetStreamResponse> GetStreamAsync(GetStreamRequest request, CallOptions options)
 			{
 				if (_outer._streamIdToStreamResponse.TryGetValue(request.StreamId, out GetStreamResponse? streamResponse))
 				{
-					return HordeRpcClientStub.Wrap(streamResponse);
+					return JobRpcClientStub.Wrap(streamResponse);
 				}
 
 				throw new RpcException(new Status(StatusCode.NotFound, $"Stream ID {request.StreamId} not found"));
@@ -320,16 +336,11 @@ namespace Horde.Agent.Tests
 			{
 				if (_outer._jobIdToJobResponse.TryGetValue(request.JobId, out GetJobResponse? jobResponse))
 				{
-					return HordeRpcClientStub.Wrap(jobResponse);
+					return JobRpcClientStub.Wrap(jobResponse);
 				}
 
 				throw new RpcException(new Status(StatusCode.NotFound, $"Job ID {request.JobId} not found"));
 			}
-
-			public override AsyncDuplexStreamingCall<UpdateSessionRequest, UpdateSessionResponse> UpdateSession(Metadata headers = null!, DateTime? deadline = null, CancellationToken cancellationToken = default)
-			{
-				return _outer.GetUpdateSessionCall(CancellationToken.None);
-			}
 		}
 		
 		public FakeHordeRpcServer()
@@ -337,13 +348,14 @@ namespace Horde.Agent.Tests
 			_serverName = "FakeServer";
 			_mockClient = new (MockBehavior.Strict);
 			_logger = NullLogger<FakeHordeRpcServer>.Instance;
-			_client = new FakeHordeRpcClient(this);
-			_connection = new RpcConnectionStub(null!, _client);
+			FakeHordeRpcClient hordeClient = new FakeHordeRpcClient(this);
+			_client = new FakeJobRpcClient(this);
+			_connection = new RpcConnectionStub(null!, hordeClient, _client);
 
 			_mockClientRef = new Mock<IRpcClientRef<HordeRpc.HordeRpcClient>>();
 			_mockClientRef
 				.Setup(m => m.Client)
-				.Returns(() => _client);
+				.Returns(() => hordeClient);
 
 			_mockConnection = new(MockBehavior.Strict);
 			_mockConnection
