Horde: Add methods to allow reading ref and blobs for an artifact from storage without having to be given permissions to the underlying namespace.

diff --git a/Engine/Source/Programs/Horde/Horde.Build/Artifacts/ArtifactsController.cs b/Engine/Source/Programs/Horde/Horde.Build/Artifacts/ArtifactsController.cs
--- a/Engine/Source/Programs/Horde/Horde.Build/Artifacts/ArtifactsController.cs
+++ b/Engine/Source/Programs/Horde/Horde.Build/Artifacts/ArtifactsController.cs
@@ -154,6 +154,59 @@ namespace Horde.Build.Artifacts
 			return PropertyFilter.Apply(new GetArtifactResponse(artifact), filter);
 		}
 
+		/// <summary>
+		/// Retrieves blobs for a particular artifact
+		/// </summary>
+		/// <param name="id">Identifier of the artifact to retrieve</param>
+		/// <param name="locator">The blob locator</param>
+		/// <param name="length">Length of data to return</param>
+		/// <param name="offset">Offset of the data to return</param>
+		/// <param name="cancellationToken">Cancellation token for the operation</param>
+		/// <returns>Information about all the artifacts</returns>
+		[HttpGet]
+		[Route("/api/v2/artifacts/{id}/blobs/{*locator}")]
+		public async Task<ActionResult<object>> ReadArtifactBlobAsync(ArtifactId id, BlobLocator locator, [FromQuery] int? offset = null, [FromQuery] int? length = null, CancellationToken cancellationToken = default)
+		{
+			IArtifact? artifact = await _artifactCollection.GetAsync(id, cancellationToken);
+			if (artifact == null)
+			{
+				return NotFound(id);
+			}
+			if (!_globalConfig.Authorize(artifact.AclScope, AclAction.ReadArtifact, User))
+			{
+				return Forbid(AclAction.ReadArtifact, artifact.AclScope);
+			}
+			if (!locator.BlobId.WithinFolder(artifact.RefName.Text))
+			{
+				return BadRequest("Invalid blob id for artifact");
+			}
+
+			return StorageController.ReadBlobInternalAsync(_storageService, artifact.NamespaceId, locator, offset, length, cancellationToken);
+		}
+
+		/// <summary>
+		/// Retrieves the root blob for an artifact
+		/// </summary>
+		/// <param name="id">Identifier of the artifact to retrieve</param>
+		/// <param name="cancellationToken">Cancellation token for the operation</param>
+		/// <returns>Information about all the artifacts</returns>
+		[HttpGet]
+		[Route("/api/v2/artifacts/{id}/ref")]
+		public async Task<ActionResult<object>> ReadArtifactRefAsync(ArtifactId id, CancellationToken cancellationToken = default)
+		{
+			IArtifact? artifact = await _artifactCollection.GetAsync(id, cancellationToken);
+			if (artifact == null)
+			{
+				return NotFound(id);
+			}
+			if (!_globalConfig.Authorize(artifact.AclScope, AclAction.ReadArtifact, User))
+			{
+				return Forbid(AclAction.ReadArtifact, artifact.AclScope);
+			}
+
+			return StorageController.ReadRefInternalAsync(_storageService, artifact.NamespaceId, artifact.RefName, cancellationToken);
+		}
+
 		/// <summary>
 		/// Gets metadata about an artifact object
 		/// </summary>
diff --git a/Engine/Source/Programs/Horde/Horde.Build/Storage/StorageController.cs b/Engine/Source/Programs/Horde/Horde.Build/Storage/StorageController.cs
--- a/Engine/Source/Programs/Horde/Horde.Build/Storage/StorageController.cs
+++ b/Engine/Source/Programs/Horde/Horde.Build/Storage/StorageController.cs
@@ -261,12 +261,20 @@ namespace Horde.Build.Storage
 				return Forbid(AclAction.ReadBlobs, namespaceId);
 			}
 
-			IStorageClientImpl client = await _storageService.GetClientAsync(namespaceId, cancellationToken);
+			return await ReadBlobInternalAsync(_storageService, namespaceId, locator, offset, length, cancellationToken);
+		}
+
+		/// <summary>
+		/// Reads a blob from storage, without performing namespace access checks.
+		/// </summary>
+		internal static async Task<ActionResult> ReadBlobInternalAsync(StorageService storageService, NamespaceId namespaceId, BlobLocator locator, int? offset, int? length, CancellationToken cancellationToken)
+		{
+			IStorageClientImpl client = await storageService.GetClientAsync(namespaceId, cancellationToken);
 
 			Uri? redirectUrl = await client.GetReadRedirectAsync(locator, cancellationToken);
 			if (redirectUrl != null)
 			{
-				return Redirect(redirectUrl.ToString());
+				return new RedirectResult(redirectUrl.ToString());
 			}
 
 #pragma warning disable CA2000 // Dispose objects before losing scope
@@ -282,9 +290,9 @@ namespace Horde.Build.Storage
 			}
 			else
 			{
-				return BadRequest("Offset and length must both be specified as query parameters for ranged reads");
+				return new BadRequestObjectResult("Offset and length must both be specified as query parameters for ranged reads");
 			}
-			return File(stream, "application/octet-stream");
+			return new FileStreamResult(stream, "application/octet-stream");
 #pragma warning restore CA2000 // Dispose objects before losing scope
 		}
 
@@ -372,12 +380,20 @@ namespace Horde.Build.Storage
 				return Forbid(AclAction.ReadRefs, namespaceId);
 			}
 
-			IStorageClient client = await _storageService.GetClientAsync(namespaceId, cancellationToken);
+			return await ReadRefInternalAsync(_storageService, namespaceId, refName, cancellationToken);
+		}
+
+		/// <summary>
+		/// Reads a ref from storage, without performing namespace access checks.
+		/// </summary>
+		internal static async Task<ActionResult<ReadRefResponse>> ReadRefInternalAsync(StorageService storageService, NamespaceId namespaceId, RefName refName, CancellationToken cancellationToken)
+		{
+			IStorageClient client = await storageService.GetClientAsync(namespaceId, cancellationToken);
 
 			NodeHandle? target = await client.TryReadRefTargetAsync(refName, cancellationToken: cancellationToken);
 			if (target == null)
 			{
-				return NotFound();
+				return new NotFoundResult();
 			}
 
 			return new ReadRefResponse(target);
diff --git a/Engine/Source/Programs/Shared/EpicGames.Horde/Storage/BlobId.cs b/Engine/Source/Programs/Shared/EpicGames.Horde/Storage/BlobId.cs
--- a/Engine/Source/Programs/Shared/EpicGames.Horde/Storage/BlobId.cs
+++ b/Engine/Source/Programs/Shared/EpicGames.Horde/Storage/BlobId.cs
@@ -229,6 +229,16 @@ namespace EpicGames.Horde.Storage
 			return new Utf8String(output.AsMemory(0, outputIdx));
 		}
 
+		/// <summary>
+		/// Checks whether this blob is within the given folder
+		/// </summary>
+		/// <param name="folderName">Name of the folder</param>
+		/// <returns>True if the the blob id is within the given folder</returns>
+		public bool WithinFolder(Utf8String folderName)
+		{
+			return Inner.Length > folderName.Length && Inner.StartsWith(folderName) && Inner[folderName.Length] == '/';
+		}
+
 		/// <summary>
 		/// Test if a given character is valid in a store id
 		/// </summary>
