Horde: Add health checks which check connectivity to Redis and Mongo.

diff --git a/Engine/Source/Programs/Horde/Horde.Build/Server/LifetimeController.cs b/Engine/Source/Programs/Horde/Horde.Build/Server/LifetimeController.cs
--- a/Engine/Source/Programs/Horde/Horde.Build/Server/LifetimeController.cs
+++ b/Engine/Source/Programs/Horde/Horde.Build/Server/LifetimeController.cs
@@ -50,17 +50,5 @@ namespace Horde.Build.Server
 
 			return Task.FromResult<ActionResult>(new ContentResult { ContentType = "text/plain", StatusCode = statusCode, Content = content });
 		}
-		
-		/// <summary>
-		/// Liveness check for Kubernetes
-		/// If this return a non-successful HTTP response, Kubernetes will kill the pod and restart it
-		/// </summary>
-		/// <returns>Ok if app is not stopping and all databases can be reached</returns>
-		[HttpGet]
-		[Route("/health/live")]
-		public ActionResult K8SLivenessProbe()
-		{
-			return new ContentResult { ContentType = "text/plain", StatusCode = 200, Content = "ok" };
-		}
 	}
 }
diff --git a/Engine/Source/Programs/Horde/Horde.Build/Server/MongoService.cs b/Engine/Source/Programs/Horde/Horde.Build/Server/MongoService.cs
--- a/Engine/Source/Programs/Horde/Horde.Build/Server/MongoService.cs
+++ b/Engine/Source/Programs/Horde/Horde.Build/Server/MongoService.cs
@@ -22,6 +22,7 @@ using EpicGames.Redis.Utility;
 using Horde.Build.Secrets;
 using Horde.Build.Utilities;
 using Horde.Build.Utiltiies;
+using Microsoft.Extensions.Diagnostics.HealthChecks;
 using Microsoft.Extensions.Hosting;
 using Microsoft.Extensions.Logging;
 using Microsoft.Extensions.Options;
@@ -198,7 +199,7 @@ namespace Horde.Build.Server
 	/// <summary>
 	/// Singleton for accessing the database
 	/// </summary>
-	public sealed class MongoService : IDisposable
+	public sealed class MongoService : IHealthCheck, IDisposable
 	{
 		/// <summary>
 		/// The database instance
@@ -876,6 +877,20 @@ namespace Horde.Build.Server
 				}
 			}
 		}
+
+		/// <inheritdoc/>
+		public async Task<HealthCheckResult> CheckHealthAsync(HealthCheckContext context, CancellationToken cancellationToken = default)
+		{
+			try
+			{
+				await Database.RunCommandAsync((Command<BsonDocument>)"{ping:1}", cancellationToken: cancellationToken);
+				return HealthCheckResult.Healthy();
+			}
+			catch (Exception ex)
+			{
+				return HealthCheckResult.Unhealthy("Unable to ping MongoDB", ex);
+			}
+		}
 	}
 
 	/// <summary>
diff --git a/Engine/Source/Programs/Horde/Horde.Build/Server/RedisService.cs b/Engine/Source/Programs/Horde/Horde.Build/Server/RedisService.cs
--- a/Engine/Source/Programs/Horde/Horde.Build/Server/RedisService.cs
+++ b/Engine/Source/Programs/Horde/Horde.Build/Server/RedisService.cs
@@ -1,15 +1,18 @@
 // Copyright Epic Games, Inc. All Rights Reserved.
 
 using System;
+using System.Collections.Generic;
 using System.Diagnostics;
 using System.IO;
 using System.Linq;
 using System.Net;
 using System.Net.NetworkInformation;
 using System.Runtime.InteropServices;
+using System.Threading;
 using System.Threading.Tasks;
 using EpicGames.Core;
 using EpicGames.Redis;
+using Microsoft.Extensions.Diagnostics.HealthChecks;
 using Microsoft.Extensions.Logging;
 using Microsoft.Extensions.Options;
 using StackExchange.Redis;
@@ -19,7 +22,7 @@ namespace Horde.Build.Server
 	/// <summary>
 	/// Manages the lifetime of a bundled Redis instance
 	/// </summary>
-	public sealed class RedisService : IDisposable
+	public sealed class RedisService : IHealthCheck, IDisposable
 	{
 		/// <summary>
 		/// Default Redis port
@@ -111,6 +114,20 @@ namespace Horde.Build.Server
 			ConnectionPool = new RedisConnectionPool(20, connectionString, dbNum);
 		}
 
+		/// <inheritdoc/>
+		public async Task<HealthCheckResult> CheckHealthAsync(HealthCheckContext context, CancellationToken cancellationToken = default)
+		{
+			try
+			{
+				TimeSpan span = await GetDatabase().PingAsync();
+				return HealthCheckResult.Healthy(data: new Dictionary<string, object> { ["Latency"] = span.ToString() });
+			}
+			catch (Exception ex)
+			{
+				return HealthCheckResult.Unhealthy("Unable to ping Redis", ex);
+			}
+		}
+
 		/// <summary>
 		/// Get the least-loaded Redis database from the connection pool
 		/// Don't store the returned object and try to resolve this as late as possible to ensure load is balanced.
diff --git a/Engine/Source/Programs/Horde/Horde.Build/Startup.cs b/Engine/Source/Programs/Horde/Horde.Build/Startup.cs
--- a/Engine/Source/Programs/Horde/Horde.Build/Startup.cs
+++ b/Engine/Source/Programs/Horde/Horde.Build/Startup.cs
@@ -965,7 +965,10 @@ namespace Horde.Build
 
 		private static void OnAddHealthChecks(IServiceCollection services)
 		{
-			IHealthChecksBuilder healthChecks = services.AddHealthChecks().AddCheck("self", () => HealthCheckResult.Healthy(), tags: new[] { "self" });
+			services.AddHealthChecks()
+				.AddCheck("self", () => HealthCheckResult.Healthy(), tags: new[] { "self" })
+				.AddCheck<RedisService>("redis")
+				.AddCheck<MongoService>("mongo");
 		}
 
 		// This method gets called by the runtime. Use this method to configure the HTTP request pipeline.
@@ -1036,6 +1039,8 @@ namespace Horde.Build
 				endpoints.MapGrpcReflectionService();
 
 				endpoints.MapControllers();
+
+				endpoints.MapHealthChecks("/health/live");
 			});
 
 			if (DirectoryReference.Exists(dashboardDir)) 
