Horde: Enable trace logging for artifact uploads.

diff --git a/Engine/Source/Programs/Horde/Horde.Agent/Execution/JobExecutor.cs b/Engine/Source/Programs/Horde/Horde.Agent/Execution/JobExecutor.cs
--- a/Engine/Source/Programs/Horde/Horde.Agent/Execution/JobExecutor.cs
+++ b/Engine/Source/Programs/Horde/Horde.Agent/Execution/JobExecutor.cs
@@ -958,7 +958,7 @@ namespace Horde.Agent.Execution
 				RefName refName = TempStorage.GetRefNameForNode(_storagePrefix, step.Name);
 
 				TreeOptions treeOptions = new TreeOptions();
-				using TreeWriter treeWriter = new TreeWriter(storage, treeOptions, refName.Text);
+				using TreeWriter treeWriter = new TreeWriter(storage, treeOptions, refName.Text, traceLogger: logger);
 
 				DirectoryNode outputNode = new DirectoryNode();
 
