Horde: Clean up inconsistency between parameters and default server profiles in grpc service.

diff --git a/Engine/Source/Programs/Horde/Horde.Agent/Services/GrpcService.cs b/Engine/Source/Programs/Horde/Horde.Agent/Services/GrpcService.cs
--- a/Engine/Source/Programs/Horde/Horde.Agent/Services/GrpcService.cs
+++ b/Engine/Source/Programs/Horde/Horde.Agent/Services/GrpcService.cs
@@ -46,22 +46,6 @@ namespace Horde.Agent.Services
 		/// </summary>
 		/// <returns>New grpc channel</returns>
 		public GrpcChannel CreateGrpcChannel(string? bearerToken)
-		{
-			if (bearerToken == null)
-			{
-				return CreateGrpcChannel(_serverProfile.Url, null);
-			}
-			else
-			{
-				return CreateGrpcChannel(_serverProfile.Url, new AuthenticationHeaderValue("Bearer", bearerToken));
-			}
-		}
-
-		/// <summary>
-		/// Create a GRPC channel with the given auth header value
-		/// </summary>
-		/// <returns>New grpc channel</returns>
-		public GrpcChannel CreateGrpcChannel(Uri address, AuthenticationHeaderValue? authHeaderValue)
 		{
 			#pragma warning disable CA2000 // Dispose objects before losing scope
 			// HTTP client handler is disposed by GrpcChannel below
@@ -72,15 +56,15 @@ namespace Horde.Agent.Services
 
 			HttpClient httpClient = new (customCertHandler, true);
 			httpClient.DefaultRequestHeaders.Add("Accept", "application/json");
-			if (authHeaderValue != null)
+			if (bearerToken != null)
 			{
-				httpClient.DefaultRequestHeaders.Authorization = authHeaderValue;
+				httpClient.DefaultRequestHeaders.Authorization = new AuthenticationHeaderValue("Bearer", bearerToken);
 			}
 
 			httpClient.Timeout = TimeSpan.FromSeconds(210); // Need to make sure this doesn't cancel any long running gRPC streaming calls (eg. session update)
 
 			_logger.LogInformation("Connecting to rpc server {BaseUrl}", _serverProfile.Url);
-			return GrpcChannel.ForAddress(address, new GrpcChannelOptions
+			return GrpcChannel.ForAddress(_serverProfile.Url, new GrpcChannelOptions
 			{
 				// Required payloads coming from CAS service can be large
 				MaxReceiveMessageSize = 1024 * 1024 * 1024, // 1 GB
