Horde: Add an option to execute leases in separate processes, separating the host process from being destabilized by any of the heavily lifting around job execution.

diff --git a/Engine/Source/Programs/Horde/Horde.Agent/Leases/Handlers/JobHandler.cs b/Engine/Source/Programs/Horde/Horde.Agent/Leases/Handlers/JobHandler.cs
--- a/Engine/Source/Programs/Horde/Horde.Agent/Leases/Handlers/JobHandler.cs
+++ b/Engine/Source/Programs/Horde/Horde.Agent/Leases/Handlers/JobHandler.cs
@@ -6,11 +6,13 @@ using System.Diagnostics;
 using System.IO;
 using System.Linq;
 using System.Net.Http;
+using System.Reflection;
 using System.Runtime.InteropServices;
 using System.Threading;
 using System.Threading.Tasks;
 using EpicGames.Core;
 using EpicGames.Horde.Storage;
+using Google.Protobuf;
 using Grpc.Core;
 using Horde.Agent.Execution;
 using Horde.Agent.Parser;
@@ -60,6 +62,52 @@ namespace Horde.Agent.Leases.Handlers
 		{
 			executeTask.JobOptions ??= new JobOptions();
 
+			if (executeTask.JobOptions.RunInSeparateProcess ?? false)
+			{
+				using (ManagedProcessGroup processGroup = new ManagedProcessGroup())
+				{
+					string fileName = Assembly.GetExecutingAssembly().Location;
+
+					List<string> arguments = new List<string>();
+					arguments.Add($"-AgentId={session.AgentId}");
+					arguments.Add($"-SessionId={session.SessionId}");
+					arguments.Add($"-LeaseId={leaseId}");
+					arguments.Add($"-WorkingDir={session.WorkingDir}");
+					arguments.Add($"-Task={Convert.ToBase64String(executeTask.ToByteArray())}");
+
+					string commandLine = CommandLineArguments.Join(arguments);
+					_defaultLogger.LogInformation("Running child process with arguments: {CommandLine}", commandLine);
+
+					using (ManagedProcess process = new ManagedProcess(processGroup, "dotnet", fileName, commandLine, null, ProcessPriorityClass.Normal))
+					{
+						using (LogEventParser parser = new LogEventParser(_defaultLogger))
+						{
+							for (; ; )
+							{
+								string? line = await process.ReadLineAsync(cancellationToken);
+								if (line == null)
+								{
+									break;
+								}
+								parser.WriteLine(line);
+							}
+						}
+						process.WaitForExit();
+
+						if (process.ExitCode != 0)
+						{
+							return LeaseResult.Failed;
+						}
+					}
+				}
+				return LeaseResult.Success;
+			}
+
+			return await ExecuteInternalAsync(session, leaseId, executeTask, cancellationToken);
+		}
+
+		internal async Task<LeaseResult> ExecuteInternalAsync(ISession session, string leaseId, ExecuteJobTask executeTask, CancellationToken cancellationToken)
+		{
 			// Create a storage client for this session
 			JobOptions jobOptions = executeTask.JobOptions;
 			await using IServerLogger logger = _serverLoggerFactory.CreateLogger(session, executeTask.LogId, executeTask.JobId, executeTask.BatchId, null, null, jobOptions.UseNewLogStorage);
