Horde: Fix order of arguments to job child process.

diff --git a/Engine/Source/Programs/Horde/Horde.Agent/Leases/Handlers/JobHandler.cs b/Engine/Source/Programs/Horde/Horde.Agent/Leases/Handlers/JobHandler.cs
--- a/Engine/Source/Programs/Horde/Horde.Agent/Leases/Handlers/JobHandler.cs
+++ b/Engine/Source/Programs/Horde/Horde.Agent/Leases/Handlers/JobHandler.cs
@@ -66,9 +66,10 @@ namespace Horde.Agent.Leases.Handlers
 			{
 				using (ManagedProcessGroup processGroup = new ManagedProcessGroup())
 				{
-					string fileName = Assembly.GetExecutingAssembly().Location;
-
 					List<string> arguments = new List<string>();
+					arguments.Add(Assembly.GetExecutingAssembly().Location);
+					arguments.Add("execute");
+					arguments.Add("job");
 					arguments.Add($"-AgentId={session.AgentId}");
 					arguments.Add($"-SessionId={session.SessionId}");
 					arguments.Add($"-LeaseId={leaseId}");
@@ -78,7 +79,7 @@ namespace Horde.Agent.Leases.Handlers
 					string commandLine = CommandLineArguments.Join(arguments);
 					_defaultLogger.LogInformation("Running child process with arguments: {CommandLine}", commandLine);
 
-					using (ManagedProcess process = new ManagedProcess(processGroup, "dotnet", fileName, commandLine, null, ProcessPriorityClass.Normal))
+					using (ManagedProcess process = new ManagedProcess(processGroup, "dotnet", commandLine, null, null, ProcessPriorityClass.Normal))
 					{
 						using (LogEventParser parser = new LogEventParser(_defaultLogger))
 						{
