Reduce number of files in small chunk test.

diff --git a/Engine/Source/Programs/Horde/StorageTests.xml b/Engine/Source/Programs/Horde/StorageTests.xml
--- a/Engine/Source/Programs/Horde/StorageTests.xml
+++ b/Engine/Source/Programs/Horde/StorageTests.xml
@@ -6,7 +6,7 @@
 	<!-- Small Chunks -->
 	<Agent Name="Small Agent A" Type="CompileWin64;Win64">
 		<Node Name="Create Small Chunks" Produces="#SmallChunks">
-			<RandomData Size="1024" Count="20000000" Different="true" Tag="#SmallChunks"/>
+			<RandomData Size="1024" Count="2000000" Different="true" Tag="#SmallChunks"/>
 		</Node>
 	</Agent>
 	<Agent Name="Small Agent B" Type="CompileWin64;Win64">
@@ -40,7 +40,7 @@
 	<!-- Small Identical Chunks -->
 	<Agent Name="Small Identical Agent A" Type="CompileWin64;Win64">
 		<Node Name="Create Small Identical Chunks" Produces="#SmallIdenticalChunks">
-			<RandomData Size="1024" Count="20000000" Different="false" Tag="#SmallIdenticalChunks"/>
+			<RandomData Size="1024" Count="2000000" Different="false" Tag="#SmallIdenticalChunks"/>
 		</Node>
 	</Agent>
 	<Agent Name="Small Identical Agent B" Type="CompileWin64;Win64">
