Horde: Fix DI for JobHandler objects.

diff --git a/Engine/Source/Programs/Horde/Horde.Agent/Program.cs b/Engine/Source/Programs/Horde/Horde.Agent/Program.cs
--- a/Engine/Source/Programs/Horde/Horde.Agent/Program.cs
+++ b/Engine/Source/Programs/Horde/Horde.Agent/Program.cs
@@ -185,10 +185,12 @@ namespace Horde.Agent
 			services.AddSingleton<IJobExecutorFactory, LocalExecutorFactory>();
 			services.AddSingleton<IJobExecutorFactory, TestExecutorFactory>();
 
+			services.AddSingleton<JobHandler>();
+
 			services.AddSingleton<LeaseHandler, ComputeHandler>();
 			services.AddSingleton<LeaseHandler, ComputeHandlerV2>();
 			services.AddSingleton<LeaseHandler, ConformHandler>();
-			services.AddSingleton<LeaseHandler, JobHandler>();
+			services.AddSingleton<LeaseHandler, JobHandler>(x => x.GetRequiredService<JobHandler>());
 			services.AddSingleton<LeaseHandler, RestartHandler>();
 			services.AddSingleton<LeaseHandler, ShutdownHandler>();
 			services.AddSingleton<LeaseHandler, UpgradeHandler>();
