Horde: Use consistent casing for agent commands.

diff --git a/Engine/Source/Programs/Horde/Horde.Agent/Commands/Compute/ComputeCommand.cs b/Engine/Source/Programs/Horde/Horde.Agent/Commands/Compute/ComputeCommand.cs
--- a/Engine/Source/Programs/Horde/Horde.Agent/Commands/Compute/ComputeCommand.cs
+++ b/Engine/Source/Programs/Horde/Horde.Agent/Commands/Compute/ComputeCommand.cs
@@ -25,7 +25,7 @@ namespace Horde.Agent.Commands
 	/// <summary>
 	/// Installs the agent as a service
 	/// </summary>
-	[Command("Compute", "Executes a command through the Horde Compute API")]
+	[Command("compute", "Executes a command through the Horde Compute API")]
 	class ComputeCommand : Command
 	{
 		class JsonRequirements
diff --git a/Engine/Source/Programs/Horde/Horde.Agent/Commands/Compute/ComputeCommandV2.cs b/Engine/Source/Programs/Horde/Horde.Agent/Commands/Compute/ComputeCommandV2.cs
--- a/Engine/Source/Programs/Horde/Horde.Agent/Commands/Compute/ComputeCommandV2.cs
+++ b/Engine/Source/Programs/Horde/Horde.Agent/Commands/Compute/ComputeCommandV2.cs
@@ -20,7 +20,7 @@ namespace Horde.Agent.Commands
 	/// <summary>
 	/// Installs the agent as a service
 	/// </summary>
-	[Command("ComputeV2", "Executes a command through the Horde Compute API")]
+	[Command("computeV2", "Executes a command through the Horde Compute API")]
 	class ComputeCommandV2 : Command
 	{
 		[CommandLine("-Cluster=", Description = "Cluster to execute on")]
diff --git a/Engine/Source/Programs/Horde/Horde.Agent/Commands/Parse/ParseCommand.cs b/Engine/Source/Programs/Horde/Horde.Agent/Commands/Parse/ParseCommand.cs
--- a/Engine/Source/Programs/Horde/Horde.Agent/Commands/Parse/ParseCommand.cs
+++ b/Engine/Source/Programs/Horde/Horde.Agent/Commands/Parse/ParseCommand.cs
@@ -16,7 +16,7 @@ namespace Horde.Agent.Commands.Parse
 	/// <summary>
 	/// Installs the agent as a service
 	/// </summary>
-	[Command("Parse", "Parses a file into structured logging output")]
+	[Command("parse", "Parses a file into structured logging output")]
 	class ParseCommand : Command
 	{
 		[CommandLine("-File=", Required = true)]
diff --git a/Engine/Source/Programs/Horde/Horde.Agent/Commands/Service/ServiceInstall.cs b/Engine/Source/Programs/Horde/Horde.Agent/Commands/Service/ServiceInstall.cs
--- a/Engine/Source/Programs/Horde/Horde.Agent/Commands/Service/ServiceInstall.cs
+++ b/Engine/Source/Programs/Horde/Horde.Agent/Commands/Service/ServiceInstall.cs
@@ -13,7 +13,7 @@ namespace Horde.Agent.Commands.Service
 	/// <summary>
 	/// Installs the agent as a service
 	/// </summary>
-	[Command("Service", "Install", "Installs the agent as a service")]
+	[Command("service", "install", "Installs the agent as a service")]
 	class InstallCommand : Command
 	{
 		/// <summary>
diff --git a/Engine/Source/Programs/Horde/Horde.Agent/Commands/Service/ServiceRun.cs b/Engine/Source/Programs/Horde/Horde.Agent/Commands/Service/ServiceRun.cs
--- a/Engine/Source/Programs/Horde/Horde.Agent/Commands/Service/ServiceRun.cs
+++ b/Engine/Source/Programs/Horde/Horde.Agent/Commands/Service/ServiceRun.cs
@@ -15,7 +15,7 @@ namespace Horde.Agent.Modes.Service
 	/// <summary>
 	/// 
 	/// </summary>
-	[Command("Service", "Run", "Runs the service in listen mode")]
+	[Command("service", "run", "Runs the service in listen mode")]
 	class RunCommand : Command
 	{
 		/// <summary>
diff --git a/Engine/Source/Programs/Horde/Horde.Agent/Commands/Service/ServiceUninstall.cs b/Engine/Source/Programs/Horde/Horde.Agent/Commands/Service/ServiceUninstall.cs
--- a/Engine/Source/Programs/Horde/Horde.Agent/Commands/Service/ServiceUninstall.cs
+++ b/Engine/Source/Programs/Horde/Horde.Agent/Commands/Service/ServiceUninstall.cs
@@ -11,7 +11,7 @@ namespace Horde.Agent.Commands.Service
 	/// <summary>
 	/// Uninstalls the service
 	/// </summary>
-	[Command("Service", "Uninstall", "Uninstalls the service")]
+	[Command("service", "uninstall", "Uninstalls the service")]
 	class UninstallCommand : Command
 	{
 		/// <summary>
diff --git a/Engine/Source/Programs/Horde/Horde.Agent/Commands/Service/ServiceUpgrade.cs b/Engine/Source/Programs/Horde/Horde.Agent/Commands/Service/ServiceUpgrade.cs
--- a/Engine/Source/Programs/Horde/Horde.Agent/Commands/Service/ServiceUpgrade.cs
+++ b/Engine/Source/Programs/Horde/Horde.Agent/Commands/Service/ServiceUpgrade.cs
@@ -20,7 +20,7 @@ namespace Horde.Agent.Commands
 	/// <summary>
 	/// Upgrades a running service to the current application
 	/// </summary>
-	[Command("Service", "Upgrade", "Replaces a running service with the application in the current directory")]
+	[Command("service", "upgrade", "Replaces a running service with the application in the current directory")]
 	class UpgradeCommand : Command
 	{
 		/// <summary>
diff --git a/Engine/Source/Programs/Horde/Horde.Agent/Commands/Utilities/BuildGraphCommand.cs b/Engine/Source/Programs/Horde/Horde.Agent/Commands/Utilities/BuildGraphCommand.cs
--- a/Engine/Source/Programs/Horde/Horde.Agent/Commands/Utilities/BuildGraphCommand.cs
+++ b/Engine/Source/Programs/Horde/Horde.Agent/Commands/Utilities/BuildGraphCommand.cs
@@ -21,7 +21,7 @@ namespace Horde.Agent.Commands.Utilities
 	/// <summary>
 	/// Runs a BuildGraph script and captures the processed log output
 	/// </summary>
-	[Command("BuildGraph", "Executes a BuildGraph script with the given arguments using a build of UAT within the current branch, and runs the output through the log processor")]
+	[Command("buildgraph", "Executes a BuildGraph script with the given arguments using a build of UAT within the current branch, and runs the output through the log processor")]
 	class BuildGraphCommand : Command
 	{
 		sealed class LogSink : IJsonRpcLogSink, IAsyncDisposable, IDisposable
diff --git a/Engine/Source/Programs/Horde/Horde.Agent/Commands/Utilities/CapsCommand.cs b/Engine/Source/Programs/Horde/Horde.Agent/Commands/Utilities/CapsCommand.cs
--- a/Engine/Source/Programs/Horde/Horde.Agent/Commands/Utilities/CapsCommand.cs
+++ b/Engine/Source/Programs/Horde/Horde.Agent/Commands/Utilities/CapsCommand.cs
@@ -11,7 +11,7 @@ namespace Horde.Agent.Commands.Utilities
 	/// <summary>
 	/// Shows capabilities of this agent
 	/// </summary>
-	[Command("Caps", "Lists detected capabilities of this agent")]
+	[Command("caps", "Lists detected capabilities of this agent")]
 	class CapsCommand : Command
 	{
 		readonly CapabilitiesService _capabilitiesService;
diff --git a/Engine/Source/Programs/Horde/Horde.Agent/Commands/Utilities/CreateCertCommand.cs b/Engine/Source/Programs/Horde/Horde.Agent/Commands/Utilities/CreateCertCommand.cs
--- a/Engine/Source/Programs/Horde/Horde.Agent/Commands/Utilities/CreateCertCommand.cs
+++ b/Engine/Source/Programs/Horde/Horde.Agent/Commands/Utilities/CreateCertCommand.cs
@@ -14,7 +14,7 @@ namespace Horde.Agent.Commands.Certs
 	/// <summary>
 	/// Creates a certificate that can be used for server/agent SSL connections
 	/// </summary>
-	[Command("CreateCert", "Creates a self-signed certificate that can be used for server/agent gRPC connections")]
+	[Command("createcert", "Creates a self-signed certificate that can be used for server/agent gRPC connections")]
 	class CreateCertCommand : Command
 	{
 		[CommandLine("-Server=")]
diff --git a/Engine/Source/Programs/Horde/Horde.Agent/Commands/Utilities/PublishCommand.cs b/Engine/Source/Programs/Horde/Horde.Agent/Commands/Utilities/PublishCommand.cs
--- a/Engine/Source/Programs/Horde/Horde.Agent/Commands/Utilities/PublishCommand.cs
+++ b/Engine/Source/Programs/Horde/Horde.Agent/Commands/Utilities/PublishCommand.cs
@@ -18,7 +18,7 @@ namespace Horde.Agent.Commands.Utilities
 	/// <summary>
 	/// Publishes the contents of this application directory to the server
 	/// </summary>
-	[Command("Publish", "Publishes this application to the server")]
+	[Command("publish", "Publishes this application to the server")]
 	class PublishCommand : Command
 	{
 		/// <summary>
diff --git a/Engine/Source/Programs/Horde/Horde.Agent/Commands/Vcs/VcsLog.cs b/Engine/Source/Programs/Horde/Horde.Agent/Commands/Vcs/VcsLog.cs
--- a/Engine/Source/Programs/Horde/Horde.Agent/Commands/Vcs/VcsLog.cs
+++ b/Engine/Source/Programs/Horde/Horde.Agent/Commands/Vcs/VcsLog.cs
@@ -10,7 +10,7 @@ using Microsoft.Extensions.Logging;
 
 namespace Horde.Agent.Commands.Vcs
 {
-	[Command("Vcs", "Log", "Print a history of commits")]
+	[Command("vcs", "log", "Print a history of commits")]
 	class VcsLog : VcsBase
 	{
 		[CommandLine("-Count")]
diff --git a/Engine/Source/Programs/Horde/Horde.Agent/Commands/Vcs/VcsStatus.cs b/Engine/Source/Programs/Horde/Horde.Agent/Commands/Vcs/VcsStatus.cs
--- a/Engine/Source/Programs/Horde/Horde.Agent/Commands/Vcs/VcsStatus.cs
+++ b/Engine/Source/Programs/Horde/Horde.Agent/Commands/Vcs/VcsStatus.cs
@@ -7,7 +7,7 @@ using Microsoft.Extensions.Logging;
 
 namespace Horde.Agent.Commands.Vcs
 {
-	[Command("Vcs", "Status", "Find status of local files")]
+	[Command("vcs", "status", "Find status of local files")]
 	class VcsStatus : VcsBase
 	{
 		public VcsStatus(IStorageClientFactory storageClientFactory)
diff --git a/Engine/Source/Programs/Horde/Horde.Agent/Commands/Workspace/WorkspaceClean.cs b/Engine/Source/Programs/Horde/Horde.Agent/Commands/Workspace/WorkspaceClean.cs
--- a/Engine/Source/Programs/Horde/Horde.Agent/Commands/Workspace/WorkspaceClean.cs
+++ b/Engine/Source/Programs/Horde/Horde.Agent/Commands/Workspace/WorkspaceClean.cs
@@ -10,7 +10,7 @@ using Microsoft.Extensions.Logging;
 
 namespace Horde.Agent.Commands.Workspace
 {
-	[Command("Workspace", "Clean", "Cleans all modified files from the workspace.")]
+	[Command("workspace", "clean", "Cleans all modified files from the workspace.")]
 	class WorkspaceClean : WorkspaceBase
 	{
 		[CommandLine("-Incremental")]
diff --git a/Engine/Source/Programs/Horde/Horde.Agent/Commands/Workspace/WorkspaceClear.cs b/Engine/Source/Programs/Horde/Horde.Agent/Commands/Workspace/WorkspaceClear.cs
--- a/Engine/Source/Programs/Horde/Horde.Agent/Commands/Workspace/WorkspaceClear.cs
+++ b/Engine/Source/Programs/Horde/Horde.Agent/Commands/Workspace/WorkspaceClear.cs
@@ -9,7 +9,7 @@ using Microsoft.Extensions.Logging;
 
 namespace Horde.Agent.Commands.Workspace
 {
-	[Command("Workspace", "Clear", "Empties the staging directory of any files, returning them to the cache")]
+	[Command("workspace", "clear", "Empties the staging directory of any files, returning them to the cache")]
 	class WorkspaceClear : WorkspaceBase
 	{
 		protected override Task ExecuteAsync(IPerforceConnection perforce, ManagedWorkspace repo, ILogger logger)
diff --git a/Engine/Source/Programs/Horde/Horde.Agent/Commands/Workspace/WorkspaceDump.cs b/Engine/Source/Programs/Horde/Horde.Agent/Commands/Workspace/WorkspaceDump.cs
--- a/Engine/Source/Programs/Horde/Horde.Agent/Commands/Workspace/WorkspaceDump.cs
+++ b/Engine/Source/Programs/Horde/Horde.Agent/Commands/Workspace/WorkspaceDump.cs
@@ -8,7 +8,7 @@ using Microsoft.Extensions.Logging;
 
 namespace Horde.Agent.Commands.Workspace
 {
-	[Command("Workspace", "Dump", "Dumps the contents of the repository to the log for analysis")]
+	[Command("workspace", "dump", "Dumps the contents of the repository to the log for analysis")]
 	class WorkspaceDump : WorkspaceBase
 	{
 		protected override Task ExecuteAsync(IPerforceConnection perforce, ManagedWorkspace repo, ILogger logger)
diff --git a/Engine/Source/Programs/Horde/Horde.Agent/Commands/Workspace/WorkspacePopulateCache.cs b/Engine/Source/Programs/Horde/Horde.Agent/Commands/Workspace/WorkspacePopulateCache.cs
--- a/Engine/Source/Programs/Horde/Horde.Agent/Commands/Workspace/WorkspacePopulateCache.cs
+++ b/Engine/Source/Programs/Horde/Horde.Agent/Commands/Workspace/WorkspacePopulateCache.cs
@@ -12,7 +12,7 @@ using Microsoft.Extensions.Logging;
 
 namespace Horde.Agent.Commands.Workspace
 {
-	[Command("Workspace", "PopulateCache", "Populates the cache with the head revision of the given streams")]
+	[Command("workspace", "populatecache", "Populates the cache with the head revision of the given streams")]
 	class WorkspacePopulateCache : WorkspaceBase
 	{
 		[CommandLine("-ClientAndStream=")]
diff --git a/Engine/Source/Programs/Horde/Horde.Agent/Commands/Workspace/WorkspacePurgeCache.cs b/Engine/Source/Programs/Horde/Horde.Agent/Commands/Workspace/WorkspacePurgeCache.cs
--- a/Engine/Source/Programs/Horde/Horde.Agent/Commands/Workspace/WorkspacePurgeCache.cs
+++ b/Engine/Source/Programs/Horde/Horde.Agent/Commands/Workspace/WorkspacePurgeCache.cs
@@ -9,7 +9,7 @@ using Microsoft.Extensions.Logging;
 
 namespace Horde.Agent.Commands.Workspace
 {
-	[Command("Workspace", "PurgeCache", "Shrink the size of the cache to the given size")]
+	[Command("workspace", "purgecache", "Shrink the size of the cache to the given size")]
 	class WorkspacePurgeCache : WorkspaceBase
 	{
 		[CommandLine("-Size=")]
diff --git a/Engine/Source/Programs/Horde/Horde.Agent/Commands/Workspace/WorkspaceRepair.cs b/Engine/Source/Programs/Horde/Horde.Agent/Commands/Workspace/WorkspaceRepair.cs
--- a/Engine/Source/Programs/Horde/Horde.Agent/Commands/Workspace/WorkspaceRepair.cs
+++ b/Engine/Source/Programs/Horde/Horde.Agent/Commands/Workspace/WorkspaceRepair.cs
@@ -9,7 +9,7 @@ using Microsoft.Extensions.Logging;
 
 namespace Horde.Agent.Commands.Workspace
 {
-	[Command("Workspace", "RepairCache", "Checks the integrity of the cache, and removes any invalid files")]
+	[Command("workspace", "repaircache", "Checks the integrity of the cache, and removes any invalid files")]
 	class WorkspaceRepair : WorkspaceBase
 	{
 		protected override Task ExecuteAsync(IPerforceConnection perforce, ManagedWorkspace repo, ILogger logger)
diff --git a/Engine/Source/Programs/Horde/Horde.Agent/Commands/Workspace/WorkspaceRevert.cs b/Engine/Source/Programs/Horde/Horde.Agent/Commands/Workspace/WorkspaceRevert.cs
--- a/Engine/Source/Programs/Horde/Horde.Agent/Commands/Workspace/WorkspaceRevert.cs
+++ b/Engine/Source/Programs/Horde/Horde.Agent/Commands/Workspace/WorkspaceRevert.cs
@@ -10,7 +10,7 @@ using Microsoft.Extensions.Logging;
 
 namespace Horde.Agent.Commands.Workspace
 {
-	[Command("Workspace", "Revert", "Revert all files that are open in the current workspace. Does not replace them with valid revisions.")]
+	[Command("workspace", "revert", "Revert all files that are open in the current workspace. Does not replace them with valid revisions.")]
 	class WorkspaceRevert : WorkspaceBase
 	{
 		[CommandLine("-Client=", Required = true)]
diff --git a/Engine/Source/Programs/Horde/Horde.Agent/Commands/Workspace/WorkspaceSetup.cs b/Engine/Source/Programs/Horde/Horde.Agent/Commands/Workspace/WorkspaceSetup.cs
--- a/Engine/Source/Programs/Horde/Horde.Agent/Commands/Workspace/WorkspaceSetup.cs
+++ b/Engine/Source/Programs/Horde/Horde.Agent/Commands/Workspace/WorkspaceSetup.cs
@@ -10,7 +10,7 @@ using Microsoft.Extensions.Logging;
 
 namespace Horde.Agent.Commands.Workspace
 {
-	[Command("Workspace", "Setup", "Creates or updates a client to use a given stream")]
+	[Command("workspace", "setup", "Creates or updates a client to use a given stream")]
 	class WorkspaceSetup : WorkspaceBase
 	{
 		[CommandLine("-Client=", Required = true)]
diff --git a/Engine/Source/Programs/Horde/Horde.Agent/Commands/Workspace/WorkspaceSizes.cs b/Engine/Source/Programs/Horde/Horde.Agent/Commands/Workspace/WorkspaceSizes.cs
--- a/Engine/Source/Programs/Horde/Horde.Agent/Commands/Workspace/WorkspaceSizes.cs
+++ b/Engine/Source/Programs/Horde/Horde.Agent/Commands/Workspace/WorkspaceSizes.cs
@@ -11,7 +11,7 @@ using Microsoft.Extensions.Logging;
 
 namespace Horde.Agent.Commands.Workspace
 {
-	[Command("Workspace", "Sizes", "Gathers stats for which streams take the most amount of space in the cache for the given configuration")]
+	[Command("workspace", "sizes", "Gathers stats for which streams take the most amount of space in the cache for the given configuration")]
 	class WorkspaceSizes : WorkspaceBase
 	{
 		[CommandLine("-TempClient=", Required = true)]
diff --git a/Engine/Source/Programs/Horde/Horde.Agent/Commands/Workspace/WorkspaceStatus.cs b/Engine/Source/Programs/Horde/Horde.Agent/Commands/Workspace/WorkspaceStatus.cs
--- a/Engine/Source/Programs/Horde/Horde.Agent/Commands/Workspace/WorkspaceStatus.cs
+++ b/Engine/Source/Programs/Horde/Horde.Agent/Commands/Workspace/WorkspaceStatus.cs
@@ -8,7 +8,7 @@ using Microsoft.Extensions.Logging;
 
 namespace Horde.Agent.Commands.Workspace
 {
-	[Command("Workspace", "Status", "Prints information about the state of the cache and workspace")]
+	[Command("workspace", "status", "Prints information about the state of the cache and workspace")]
 	class WorkspaceStatus : WorkspaceBase
 	{
 		protected override Task ExecuteAsync(IPerforceConnection perforce, ManagedWorkspace repo, ILogger logger)
diff --git a/Engine/Source/Programs/Horde/Horde.Agent/Commands/Workspace/WorkspaceSync.cs b/Engine/Source/Programs/Horde/Horde.Agent/Commands/Workspace/WorkspaceSync.cs
--- a/Engine/Source/Programs/Horde/Horde.Agent/Commands/Workspace/WorkspaceSync.cs
+++ b/Engine/Source/Programs/Horde/Horde.Agent/Commands/Workspace/WorkspaceSync.cs
@@ -11,7 +11,7 @@ using Microsoft.Extensions.Logging;
 
 namespace Horde.Agent.Commands.Workspace
 {
-	[Command("Workspace", "Sync", "Syncs the files for a particular stream and changelist")]
+	[Command("workspace", "sync", "Syncs the files for a particular stream and changelist")]
 	class WorkspaceSync : WorkspaceBase
 	{
 		[CommandLine("-Client=", Required = true)]
