Horde: Specify the active server profile when spawning a job task.

diff --git a/Engine/Source/Programs/Horde/Horde.Agent/Leases/Handlers/JobHandler.cs b/Engine/Source/Programs/Horde/Horde.Agent/Leases/Handlers/JobHandler.cs
--- a/Engine/Source/Programs/Horde/Horde.Agent/Leases/Handlers/JobHandler.cs
+++ b/Engine/Source/Programs/Horde/Horde.Agent/Leases/Handlers/JobHandler.cs
@@ -70,6 +70,7 @@ namespace Horde.Agent.Leases.Handlers
 					arguments.Add(Assembly.GetExecutingAssembly().Location);
 					arguments.Add("execute");
 					arguments.Add("job");
+					arguments.Add($"-Server={_settings.GetCurrentServerProfile().Name}");
 					arguments.Add($"-AgentId={session.AgentId}");
 					arguments.Add($"-SessionId={session.SessionId}");
 					arguments.Add($"-LeaseId={leaseId}");
