Horde: Missing file.

diff --git a/Engine/Source/Programs/Horde/Horde.Agent/Commands/Execution/ExecuteJobCommand.cs b/Engine/Source/Programs/Horde/Horde.Agent/Commands/Execution/ExecuteJobCommand.cs
new file mode 100644
--- /dev/null
+++ b/Engine/Source/Programs/Horde/Horde.Agent/Commands/Execution/ExecuteJobCommand.cs
@@ -0,0 +1,57 @@
+// Copyright Epic Games, Inc. All Rights Reserved.
+
+using System;
+using System.Threading;
+using System.Threading.Tasks;
+using EpicGames.Core;
+using Horde.Agent.Leases.Handlers;
+using Horde.Agent.Services;
+using Horde.Agent.Utility;
+using HordeCommon.Rpc.Tasks;
+using Microsoft.Extensions.Logging;
+using Microsoft.Extensions.Options;
+
+namespace Horde.Agent.Commands.Execution
+{
+	[Command("Execute", "Job", "Executes a job")]
+	class ExecuteJobCommand : Command
+	{
+		[CommandLine("-AgentId=", Required = true)]
+		public string AgentId { get; set; } = null!;
+
+		[CommandLine("-SessionId=", Required = true)]
+		public string SessionId { get; set; } = null!;
+
+		[CommandLine("-LeaseId", Required = true)]
+		public string LeaseId { get; set; } = null!;
+
+		[CommandLine("-Task=", Required = true)]
+		public string Task { get; set; } = null!;
+
+		[CommandLine("-WorkingDir=", Required = true)]
+		public DirectoryReference WorkingDir = null!;
+
+		readonly GrpcService _grpcService;
+		readonly JobHandler _jobHandler;
+		readonly AgentSettings _settings;
+
+		public ExecuteJobCommand(GrpcService grpcService, JobHandler jobHandler, IOptions<AgentSettings> settings)
+		{
+			_grpcService = grpcService;
+			_jobHandler = jobHandler;
+			_settings = settings.Value;
+		}
+
+		public override async Task<int> ExecuteAsync(ILogger logger)
+		{
+			ServerProfile serverProfile = _settings.GetCurrentServerProfile();
+			ExecuteJobTask executeTask = ExecuteJobTask.Parser.ParseFrom(Convert.FromBase64String(Task));
+
+			await using RpcConnection rpcConnection = new RpcConnection(() => _grpcService.CreateGrpcChannel(executeTask.Token), logger);
+			await using Session session = new Session(serverProfile.Url, AgentId, SessionId, executeTask.Token, rpcConnection, WorkingDir, _settings.ProcessNamesToTerminate, logger);
+
+			await _jobHandler.ExecuteInternalAsync(session, LeaseId, executeTask, CancellationToken.None);
+			return 0;
+		}
+	}
+}
