Horde: Remove uploading of XGE monitor files as artifacts

diff --git a/Engine/Source/Programs/Horde/Horde.Agent/Execution/JobExecutor.cs b/Engine/Source/Programs/Horde/Horde.Agent/Execution/JobExecutor.cs
--- a/Engine/Source/Programs/Horde/Horde.Agent/Execution/JobExecutor.cs
+++ b/Engine/Source/Programs/Horde/Horde.Agent/Execution/JobExecutor.cs
@@ -738,7 +738,6 @@ namespace Horde.Agent.Execution
 			if (_jobOptions.UseNewTempStorage ?? false)
 			{
 				bool result = await ExecuteWithTempStorageAsync(step, workspaceDir, arguments.ToString(), logger, cancellationToken);
-				await UploadXgeMonitorFilesAsync(step, logger, cancellationToken);
 				return result;
 			}
 			else
@@ -749,7 +748,6 @@ namespace Horde.Agent.Execution
 				}
 				
 				bool result = await ExecuteAutomationToolAsync(step, workspaceDir, arguments.ToString(), logger, cancellationToken) == 0;
-				await UploadXgeMonitorFilesAsync(step, logger, cancellationToken);
 				return result;
 			}
 		}
