Horde: Make agent view filters persistent
#rnx
#skipci

diff --git a/Engine/Source/Programs/Horde/HordeDashboard/src/components/AgentView.tsx b/Engine/Source/Programs/Horde/HordeDashboard/src/components/AgentView.tsx
--- a/Engine/Source/Programs/Horde/HordeDashboard/src/components/AgentView.tsx
+++ b/Engine/Source/Programs/Horde/HordeDashboard/src/components/AgentView.tsx
@@ -1,10 +1,10 @@
 // Copyright Epic Games, Inc. All Rights Reserved.  
-import { Checkbox, CommandButton, ConstrainMode, ContextualMenu, TagPicker, DefaultButton, DetailsHeader, DetailsList, DetailsListLayoutMode, Dialog, DialogType, DirectionalHint, Dropdown, FontSizes, FontWeights, getTheme, IBasePickerProps, IColumn, Icon, IconButton, IContextualMenuItem, IContextualMenuProps, IDetailsHeaderProps, IDetailsHeaderStyles, IDetailsListProps, ITag, ITagItemStyles, ITooltipHostStyles, Link as ReactLink, mergeStyles, mergeStyleSets, PrimaryButton, ProgressIndicator, ScrollablePane, ScrollbarVisibility, SearchBox, Selection, SelectionMode, Slider, Spinner, SpinnerSize, Stack, Sticky, StickyPositionType, TagItem, Text, TextField } from '@fluentui/react';
+import { Checkbox, CommandButton, ConstrainMode, ContextualMenu, DefaultButton, DetailsHeader, DetailsList, DetailsListLayoutMode, Dialog, DialogType, DirectionalHint, Dropdown, FontSizes, FontWeights, getTheme, IBasePickerProps, IColumn, Icon, IconButton, IContextualMenuItem, IContextualMenuProps, IDetailsHeaderProps, IDetailsHeaderStyles, IDetailsListProps, ITag, ITagItemStyles, ITooltipHostStyles, Link as ReactLink, mergeStyles, mergeStyleSets, PrimaryButton, ProgressIndicator, ScrollablePane, ScrollbarVisibility, SearchBox, Selection, SelectionMode, Slider, Spinner, SpinnerSize, Stack, Sticky, StickyPositionType, TagItem, TagPicker, Text, TextField } from '@fluentui/react';
 import { action, makeObservable, observable } from 'mobx';
 import { observer } from 'mobx-react-lite';
 import moment from 'moment-timezone';
 import React, { createRef, useEffect, useState } from 'react';
-import { Link, useNavigate, useParams } from 'react-router-dom';
+import { Link, useNavigate, useParams, useSearchParams } from 'react-router-dom';
 import Marquee from 'react-text-marquee';
 import backend from '../backend';
 import { agentStore } from '../backend/AgentStore';
@@ -15,7 +15,6 @@ import { hexToRGB, hordeClasses, linearInterpolate } from '../styles/Styles';
 import { Breadcrumbs } from './Breadcrumbs';
 import { ConfirmationDialog } from './ConfirmationDialog';
 import { HistoryModal } from './HistoryModal';
-import { useQuery } from './JobDetailCommon';
 import { TopNav } from './TopNav';
 
 
@@ -160,9 +159,9 @@ function getAgentCapability(agent: AgentData, inProp: string) {
 }
 
 function getAgentOs(agent: AgentData) {
-	const osDist = getAgentCapability(agent, "OSDistribution");
-	const osFamily = getAgentCapability(agent, "OSFamily");
-	return osDist ?? osFamily;
+   const osDist = getAgentCapability(agent, "OSDistribution");
+   const osFamily = getAgentCapability(agent, "OSFamily");
+   return osDist ?? osFamily;
 }
 
 function getTaskTime(agent: GetAgentResponse): number {
@@ -188,7 +187,12 @@ function getTaskTime(agent: GetAgentResponse): number {
 
 }
 
-
+type SearchState = {
+   agentSearch?: string;
+   exactSearch?: boolean;
+   filter?: string[];
+   columnMode?: string;
+}
 
 class LocalState {
 
@@ -233,6 +237,7 @@ class LocalState {
       this.selection = new Selection({ onSelectionChanged: () => { this.setAgentsSelectedCount(this.selection.getSelectedCount()); } });
       this.currentSelection = [];
       this.deleteAgentDialogIsOpen = false;
+      this.searchState = {};
    }
 
    @action
@@ -243,16 +248,23 @@ class LocalState {
 
    @action
    setAgentFilter(filter: string) {
+      this.searchState.agentSearch = filter?.trim() ? filter.trim() : undefined;
+      this.updateSearch();
       this.agentFilter = filter;
    }
 
    @action
    setAgentStaus(filter: Set<string>) {
+      const filters = Array.from(filter);
+      this.searchState.filter = filters.length ? filters : undefined;
+      this.updateSearch();
       this.agentStatusFilter = filter;
    }
 
    @action
    setExactMatch(match: boolean) {
+      this.searchState.exactSearch = match ? true : undefined;
+      this.updateSearch();
       this.filterExactMatch = match;
    }
 
@@ -397,11 +409,18 @@ class LocalState {
       this.columnsState.filter(colState => { return colState.isCheckable; }).forEach(colState => {
          if (colState.key === key) {
             colState.isChecked = true;
+            this.searchState.columnMode = key ? key : undefined;
+            this.updateSearch();
+      
          }
          else {
             colState.isChecked = false;
          }
       });
+
+      //this.searchState.columnMode = key;
+      //this.updateSearch();
+      
       this.columnMenuProps = this._updateColumnProps();
    }
 
@@ -495,6 +514,89 @@ class LocalState {
       this.headerContextMenuOpen = false;
    }
 
+
+   @observable
+   searchUpdated: number = 0;
+
+   @action
+   setSearchUpdated() {
+      this.searchUpdated++;
+   }
+
+
+   searchState: SearchState = {};
+   search: URLSearchParams = new URLSearchParams();
+
+   updateSearch(): boolean {
+
+      const state = { ...this.searchState } as SearchState;
+
+      state.filter = state.filter?.sort((a, b) => a.localeCompare(b));
+
+      const search = new URLSearchParams();
+      const csearch = this.search.toString();
+
+      if (state.agentSearch) {
+         search.append("agent", state.agentSearch);
+      }
+
+      if (state.exactSearch) {
+         search.append("exact", "true");
+      }
+
+      state.filter?.forEach(f => {
+         if (f) {
+            search.append("filter", f);
+         }
+      });
+
+      if (state.columnMode) {
+         search.append("mode", state.columnMode);
+      }
+
+      if (search.toString() !== csearch) {
+         this.search = search;
+         this.setSearchUpdated();
+         return true;
+      }
+
+      return false;
+   }
+
+   stateFromSearch(search: URLSearchParams) {
+
+      const state: SearchState = {};
+
+      const filters = search.getAll("filter") ?? undefined;
+      const agentSearch = search.get("agent") ?? undefined;
+      const exact = search.get("exact") ?? undefined;
+      const mode = search.get("mode") ?? undefined;
+
+      state.filter = filters?.sort((a, b) => a.localeCompare(b));
+      state.columnMode = mode?.length ? mode : undefined;
+      state.agentSearch = agentSearch?.length ? agentSearch : undefined;
+      state.exactSearch = exact?.trim() === "true" ? true : undefined;
+
+      this.searchState = state;
+ 
+      if (state.agentSearch) {
+         this.agentFilter = state.agentSearch;
+      }
+      if (state.exactSearch) {
+         this.filterExactMatch = true;
+      }
+
+      if (state.filter?.length) {
+         this.agentStatusFilter = new Set<string>(state.filter);
+      }
+
+      if (state.columnMode) {
+         this.setColumnItemChecked(state.columnMode);
+      }
+
+      return state;
+   }
+
    constructor() {
       makeObservable(this);
       this.columnsState = [
@@ -1176,16 +1278,6 @@ const agentStatus = ["Active", "Ready", "Disabled", "Pending Conform", "Pending
 
 export const AgentMenuBar: React.FC = observer(() => {
 
-   const [initial, setInitial] = useState(true);
-   const query = useQuery();
-   const search = query.get("search") ? query.get("search")! : "";
-
-   if (search && initial) {
-      localState.setAgentFilter(search);
-      localState.setExactMatch(true);
-      setInitial(false);
-   }
-
    let selectedButton: any = <div></div>;
    if (localState.agentsSelectedCount !== 0) {
       selectedButton = <PrimaryButton
@@ -1598,10 +1690,29 @@ export const PoolSelectionModal: React.FC = observer(() => {
    );
 });
 
+export const SearchUpdate: React.FC = observer(() => {
+
+   const [searchParams, setSearchParams] = useSearchParams();
+   const [state, setState] = useState<{ search?: string }>({ search: searchParams.toString() });
+
+   // subscribe
+   if (localState.searchUpdated) { }
+
+   const csearch = localState.search.toString();   
+
+   if (state.search !== csearch) {
+      setSearchParams(csearch);
+      setState({ search: csearch });      
+   }
+
+   return null;
+});
+
 export const AgentView: React.FC = observer(() => {
    const { agentId } = useParams<{ agentId: string }>();
    const navigate = useNavigate();
    const [initAgentUpdater, setInitAgentUpdater] = useState(false);
+   const [searchParams] = useSearchParams();
 
    // adjust automatically to viewport changes
    useWindowSize();
@@ -1616,6 +1727,7 @@ export const AgentView: React.FC = observer(() => {
    if (!initAgentUpdater) {
       agentStore.update().then(() => {
          localState.resetState();
+         localState.stateFromSearch(searchParams);
          setInitAgentUpdater(true);
       });
       return <Stack className={hordeClasses.horde}>
@@ -1976,8 +2088,9 @@ export const AgentView: React.FC = observer(() => {
             <Stack grow styles={{ root: { backgroundColor: 'rgb(250, 249, 249)' } }} />
             <Stack tokens={{ maxWidth: 1800, childrenGap: 4 }} styles={{ root: { width: 1800, height: '100vh', backgroundColor: 'rgb(250, 249, 249)', paddingTop: 18, paddingLeft: 12 } }}>
                <Stack className={hordeClasses.raised} styles={{ root: { paddingRight: '40px' } }}>
+                  <SearchUpdate />
                   <Stack.Item>
-                     <AgentMenuBar></AgentMenuBar>
+                     <AgentMenuBar />
                   </Stack.Item>
                   <Stack style={{ position: "relative", height: "calc(100vh - 240px)" }}>
                      <ScrollablePane scrollbarVisibility={ScrollbarVisibility.always}>
@@ -2103,23 +2216,23 @@ export const AgentView: React.FC = observer(() => {
       return <Icon iconName="FullCircle" className={className} />;
    }
 
-	function getPropFromDevice(agent: AgentData, propKey: string, propValue: string | null = null) {
-		const unknownElement = <Stack styles={{ root: { height: '100%' } }} horizontal horizontalAlign={"center"}><Stack.Item align={"center"}>Unknown</Stack.Item></Stack>;
-		let prop = propValue ?? getAgentCapability(agent, propKey);
-		if (prop === null)
-			return unknownElement;
-
-		if (propKey.indexOf("RAM") !== -1) {
-			prop += " GB";
-		}
-		return (
-			<Stack styles={{ root: { height: '100%' } }} horizontal horizontalAlign={"center"}>
-				<Stack.Item align={"center"} className={agentStyles.ellipsesStackItem}>
-					<Text title={prop} key={`sysInfo_${agent.id}_${propKey}`}>{`${prop}`}</Text>
-				</Stack.Item>
-			</Stack>
-		);
-	}
+   function getPropFromDevice(agent: AgentData, propKey: string, propValue: string | null = null) {
+      const unknownElement = <Stack styles={{ root: { height: '100%' } }} horizontal horizontalAlign={"center"}><Stack.Item align={"center"}>Unknown</Stack.Item></Stack>;
+      let prop = propValue ?? getAgentCapability(agent, propKey);
+      if (prop === null)
+         return unknownElement;
+
+      if (propKey.indexOf("RAM") !== -1) {
+         prop += " GB";
+      }
+      return (
+         <Stack styles={{ root: { height: '100%' } }} horizontal horizontalAlign={"center"}>
+            <Stack.Item align={"center"} className={agentStyles.ellipsesStackItem}>
+               <Text title={prop} key={`sysInfo_${agent.id}_${propKey}`}>{`${prop}`}</Text>
+            </Stack.Item>
+         </Stack>
+      );
+   }
 
    // when an actual item is drawn
    function onRenderAgentListItem(agent: AgentData, index?: number, column?: IColumn) {
@@ -2197,15 +2310,15 @@ export const AgentView: React.FC = observer(() => {
                            text={poolObjs[idx].name}
                            primary
                            menuProps={menuProps}
-                           menuIconProps={{iconName: ""}}
-                           className={agentStyles.buttonFont}                           
-                           onClick={(ev) => { ev.preventDefault();  ev.stopPropagation()}}
+                           menuIconProps={{ iconName: "" }}
+                           className={agentStyles.buttonFont}
+                           onClick={(ev) => { ev.preventDefault(); ev.stopPropagation() }}
                            styles={{
                               root: { border: '0px', backgroundColor: color, color: textColor },
                               rootHovered: { border: '0px', backgroundColor: color, color: textColor, },
                               rootPressed: { border: '0px', backgroundColor: color, color: textColor, }
-                           }} /> 
-                        
+                           }} />
+
                      </Stack.Item>
                   );
                   poolSearchNames.push(poolObjs[idx].name);
diff --git a/Engine/Source/Programs/Horde/HordeDashboard/src/components/PoolView.tsx b/Engine/Source/Programs/Horde/HordeDashboard/src/components/PoolView.tsx
--- a/Engine/Source/Programs/Horde/HordeDashboard/src/components/PoolView.tsx
+++ b/Engine/Source/Programs/Horde/HordeDashboard/src/components/PoolView.tsx
@@ -943,7 +943,7 @@ const PoolPanel: React.FC = () => {
             <Stack horizontal tokens={{ childrenGap: 118 }}>
                <Stack style={{ minWidth: 224 }}>
                   <Stack style={{ paddingBottom: 8 }}>
-                     <PrimaryButton text={pool.name} href={`/agents?search=${encodeURI(pool.id)}`} target="_blank" style={{ color: "#FFFFFF", backgroundColor: linearInterpolate(color), border: "unset", flexShrink: 1 }} />
+                     <PrimaryButton text={pool.name} href={`/agents?agent=${encodeURI(pool.id)}&exact=true`} target="_blank" style={{ color: "#FFFFFF", backgroundColor: linearInterpolate(color), border: "unset", flexShrink: 1 }} />
                   </Stack>
                   <Stack>
                      <DetailsList
diff --git a/Engine/Source/Programs/Horde/HordeDashboard/src/components/jobDetailsV2/JobDetailViewSteps.tsx b/Engine/Source/Programs/Horde/HordeDashboard/src/components/jobDetailsV2/JobDetailViewSteps.tsx
--- a/Engine/Source/Programs/Horde/HordeDashboard/src/components/jobDetailsV2/JobDetailViewSteps.tsx
+++ b/Engine/Source/Programs/Horde/HordeDashboard/src/components/jobDetailsV2/JobDetailViewSteps.tsx
@@ -244,7 +244,7 @@ export const StepsPanelInner: React.FC<{ jobDetails: JobDetailsV2, depStepId?: s
                if (item.agentPool && (item.batch?.state === JobStepBatchState.Ready || item.batch?.error === JobStepBatchError.NoAgentsOnline || item.batch?.error === JobStepBatchError.NoAgentsInPool)) {
                   return <Stack horizontal disableShrink={true} horizontalAlign="start">
                      {(error || warning) && <Icon styles={{ root: { paddingTop: 3, paddingRight: 8 } }} className={style!.className} iconName="Square" />}
-                     <Link to={`/agents?search=${encodeURI(item.agentPool)}`} ><Text nowrap={true}>{batchText}</Text></Link>
+                     <Link to={`/agents?agent=${encodeURI(item.agentPool)}&exact=true`} ><Text nowrap={true}>{batchText}</Text></Link>
                   </Stack>
                      ;
 
