Horde: Write exceptions while running conforms to the log.

diff --git a/Engine/Source/Programs/Horde/Horde.Agent/Leases/Handlers/ConformHandler.cs b/Engine/Source/Programs/Horde/Horde.Agent/Leases/Handlers/ConformHandler.cs
--- a/Engine/Source/Programs/Horde/Horde.Agent/Leases/Handlers/ConformHandler.cs
+++ b/Engine/Source/Programs/Horde/Horde.Agent/Leases/Handlers/ConformHandler.cs
@@ -35,7 +35,20 @@ namespace Horde.Agent.Leases.Handlers
 		public override async Task<LeaseResult> ExecuteAsync(ISession session, string leaseId, ConformTask conformTask, CancellationToken cancellationToken)
 		{
 			await using IServerLogger conformLogger = _serverLoggerFactory.CreateLogger(session, conformTask.LogId, null, null);
+			try
+			{
+				LeaseResult result = await ExecuteInternalAsync(session, leaseId, conformTask, conformLogger, cancellationToken);
+				return result;
+			}
+			catch (Exception ex)
+			{
+				conformLogger.LogError(ex, "Unhandled exception while running conform: {Message}", ex.Message);
+				throw;
+			}
+		}
 
+		async Task<LeaseResult> ExecuteInternalAsync(ISession session, string leaseId, ConformTask conformTask, IServerLogger conformLogger, CancellationToken cancellationToken)
+		{
 			conformLogger.LogInformation("Conforming, lease {LeaseId}", leaseId);
 			await session.TerminateProcessesAsync(conformLogger, cancellationToken);
 
