Horde: Allow agents to authenticate for job requests with a special lease claim.

diff --git a/Engine/Source/Programs/Horde/Horde.Agent.Tests/Execution/WorkspaceExecutorTest.cs b/Engine/Source/Programs/Horde/Horde.Agent.Tests/Execution/WorkspaceExecutorTest.cs
--- a/Engine/Source/Programs/Horde/Horde.Agent.Tests/Execution/WorkspaceExecutorTest.cs
+++ b/Engine/Source/Programs/Horde/Horde.Agent.Tests/Execution/WorkspaceExecutorTest.cs
@@ -3,12 +3,14 @@
 using System;
 using System.Collections.Generic;
 using System.IO;
+using System.Net.Http.Headers;
 using System.Threading;
 using System.Threading.Tasks;
 using EpicGames.Core;
 using Horde.Agent.Execution;
 using Horde.Agent.Services;
 using Horde.Common;
+using HordeCommon.Rpc;
 using HordeCommon.Rpc.Tasks;
 using Microsoft.Extensions.Logging;
 using Microsoft.Extensions.Logging.Abstractions;
@@ -49,7 +51,8 @@ public sealed class WorkspaceExecutorTest : IDisposable
 		_workspace.SetFile(1, "main.cpp", "main");
 		_workspace.SetFile(1, "foo/bar/baz.h", "baz");
 
-		JobExecutorOptions executorOptions = new JobExecutorOptions(_session, null!, JobId, "batch1", AgentType, default, "", null!, new JobOptions());
+		BeginBatchResponse batch = new BeginBatchResponse { Change = 1 };
+		JobExecutorOptions executorOptions = new JobExecutorOptions(_session, null!, JobId, "batch1", batch, default, "", null!, new JobOptions());
 		_executor = new (executorOptions, null, _workspace, NullLogger.Instance);
 	}
 
@@ -71,7 +74,8 @@ public sealed class WorkspaceExecutorTest : IDisposable
 	[TestMethod]
 	public async Task RegularAndAutoSdkWorkspace()
 	{
-		JobExecutorOptions executorOptions = new JobExecutorOptions(_session, null!, JobId, "batch1", AgentType, default, "", null!, new JobOptions());
+		BeginBatchResponse batch = new BeginBatchResponse { Change = 1 };
+		JobExecutorOptions executorOptions = new JobExecutorOptions(_session, null!, JobId, "batch1", batch, default, "", null!, new JobOptions());
 		WorkspaceExecutor executor = new (executorOptions, _autoSdkWorkspace, _workspace, NullLogger.Instance);
 
 		await executor.InitializeAsync(_logger, CancellationToken.None);
@@ -84,7 +88,8 @@ public sealed class WorkspaceExecutorTest : IDisposable
 	[TestMethod]
 	public async Task EnvVars()
 	{
-		JobExecutorOptions executorOptions = new JobExecutorOptions(_session, null!, JobId, "batch1", AgentType, default, "", null!, new JobOptions());
+		BeginBatchResponse batch = new BeginBatchResponse { Change = 1 };
+		JobExecutorOptions executorOptions = new JobExecutorOptions(_session, null!, JobId, "batch1", batch, default, "", null!, new JobOptions());
 		WorkspaceExecutor executor = new (executorOptions, _autoSdkWorkspace, _workspace, NullLogger.Instance);
 		await executor.InitializeAsync(_logger, CancellationToken.None);
 
@@ -109,7 +114,8 @@ public sealed class WorkspaceExecutorTest : IDisposable
 		_server.AddJob("jobPreflight", StreamId, 1, 1000);		
 		_workspace.SetFile(1000, "New/Feature/Foo.cs", "foo");
 
-		JobExecutorOptions executorOptions = new JobExecutorOptions(_session, null!, "jobPreflight", "batch1", AgentType, default, "", null!, new JobOptions());
+		BeginBatchResponse batch = new BeginBatchResponse { Change = 1, PreflightChange = 1000 };
+		JobExecutorOptions executorOptions = new JobExecutorOptions(_session, null!, "jobPreflight", "batch1", batch, default, "", null!, new JobOptions());
 		WorkspaceExecutor executor = new (executorOptions, null, _workspace, NullLogger.Instance);
 		
 		await executor.InitializeAsync(_logger, CancellationToken.None);
@@ -122,8 +128,9 @@ public sealed class WorkspaceExecutorTest : IDisposable
 	[TestMethod]
 	public async Task JobWithNoChange()
 	{
-		_server.AddJob("jobNoChange", StreamId, 0, 0);		
-		JobExecutorOptions executorOptions = new JobExecutorOptions(_session, null!, "jobNoChange", "batch1", AgentType, default, "", null!, new JobOptions());
+		_server.AddJob("jobNoChange", StreamId, 0, 0);
+		BeginBatchResponse batch = new BeginBatchResponse { };
+		JobExecutorOptions executorOptions = new JobExecutorOptions(_session, null!, "jobNoChange", "batch1", batch, default, "", null!, new JobOptions());
 		WorkspaceExecutor executor = new (executorOptions, null, _workspace, NullLogger.Instance);
 		await Assert.ThrowsExceptionAsync<WorkspaceMaterializationException>(() => executor.InitializeAsync(_logger, CancellationToken.None));
 	}
diff --git a/Engine/Source/Programs/Horde/Horde.Agent.Tests/Helper.cs b/Engine/Source/Programs/Horde/Horde.Agent.Tests/Helper.cs
--- a/Engine/Source/Programs/Horde/Horde.Agent.Tests/Helper.cs
+++ b/Engine/Source/Programs/Horde/Horde.Agent.Tests/Helper.cs
@@ -117,6 +117,7 @@ namespace Horde.Agent.Tests
 
 			res.AgentType = "agentType1";
 			res.LogId = "logId1";
+			res.Change = 1;
 
 			return Wrap(res);
 		}
diff --git a/Engine/Source/Programs/Horde/Horde.Agent/Execution/JobExecutor.cs b/Engine/Source/Programs/Horde/Horde.Agent/Execution/JobExecutor.cs
--- a/Engine/Source/Programs/Horde/Horde.Agent/Execution/JobExecutor.cs
+++ b/Engine/Source/Programs/Horde/Horde.Agent/Execution/JobExecutor.cs
@@ -49,19 +49,19 @@ namespace Horde.Agent.Execution
 		public IServerStorageFactory StorageFactory { get; }
 		public string JobId { get; }
 		public string BatchId { get; }
-		public string AgentType { get; }
+		public BeginBatchResponse Batch { get; }
 		public NamespaceId NamespaceId { get; }
 		public string StoragePrefix { get; }
 		public string Token { get; }
 		public JobOptions JobOptions { get; }
 
-		public JobExecutorOptions(ISession session, IServerStorageFactory storageFactory, string jobId, string batchId, string agentType, NamespaceId namespaceId, string storagePrefix, string token, JobOptions jobOptions)
+		public JobExecutorOptions(ISession session, IServerStorageFactory storageFactory, string jobId, string batchId, BeginBatchResponse batch, NamespaceId namespaceId, string storagePrefix, string token, JobOptions jobOptions)
 		{
 			Session = session;
 			StorageFactory = storageFactory;
 			JobId = jobId;
 			BatchId = batchId;
-			AgentType = agentType;
+			Batch = batch;
 			NamespaceId = namespaceId;
 			StoragePrefix = storagePrefix;
 			Token = token;
@@ -180,18 +180,14 @@ namespace Horde.Agent.Execution
 		protected string? _scriptFileName;
 		protected bool _preprocessScript;
 
-		protected string _jobId;
-		protected string _batchId;
-		protected string _agentTypeName;
-
 		/// <summary>
 		/// Logger for the local agent process (as opposed to job logger)
 		/// </summary>
 		protected readonly ILogger _logger;
 
-		protected GetJobResponse _job;
-		protected GetStreamResponse _stream;
-		protected GetAgentTypeResponse _agentType;
+		protected readonly string _jobId;
+		protected readonly string _batchId;
+		protected readonly BeginBatchResponse _batch;
 
 		protected List<string> _additionalArguments = new List<string>();
 
@@ -218,7 +214,7 @@ namespace Horde.Agent.Execution
 
 			_jobId = options.JobId;
 			_batchId = options.BatchId;
-			_agentTypeName = options.AgentType;
+			_batch = options.Batch;
 
 			_namespaceId = options.NamespaceId;
 			_storagePrefix = options.StoragePrefix;
@@ -226,34 +222,16 @@ namespace Horde.Agent.Execution
 			_jobOptions = options.JobOptions;
 
 			_logger = logger;
-
-			_job = null!;
-			_stream = null!;
-			_agentType = null!;
 		}
 
-		public virtual async Task InitializeAsync(ILogger logger, CancellationToken cancellationToken)
+		public virtual Task InitializeAsync(ILogger logger, CancellationToken cancellationToken)
 		{
-			// Get the job settings
-			_job = await RpcConnection.InvokeAsync((JobRpc.JobRpcClient x) => x.GetJobAsync(new GetJobRequest(_jobId), null, null, cancellationToken), cancellationToken);
-
-			// Get the stream settings
-			_stream = await RpcConnection.InvokeAsync((JobRpc.JobRpcClient x) => x.GetStreamAsync(new GetStreamRequest(_job.StreamId), null, null, cancellationToken), cancellationToken);
-
-			// Get the agent type to determine how to configure this machine
-			_agentType = _stream.AgentTypes.FirstOrDefault(x => x.Key == _agentTypeName).Value;
-			if (_agentType == null)
-			{
-				_agentType = new GetAgentTypeResponse();
-			}
-
-			foreach (KeyValuePair<string, string> envVar in _agentType.Environment)
+			// Setup the agent type
+			foreach (KeyValuePair<string, string> envVar in _batch.Environment)
 			{
 				_envVars[envVar.Key] = envVar.Value;
 			}
 
-			logger.LogInformation("Configured as agent type {AgentType}", _agentTypeName);
-
 			// Figure out if we're running as an admin
 			if (RuntimeInformation.IsOSPlatform(OSPlatform.Windows))
 			{
@@ -268,7 +246,7 @@ namespace Horde.Agent.Execution
 			}
 
 			// Get the BuildGraph arguments
-			foreach (string argument in _job.Arguments)
+			foreach (string argument in _batch.Arguments)
 			{
 				const string RemapAgentTypesPrefix = "-RemapAgentTypes=";
 				if (argument.StartsWith(RemapAgentTypesPrefix, StringComparison.OrdinalIgnoreCase))
@@ -299,9 +277,9 @@ namespace Horde.Agent.Execution
 					_additionalArguments.Add(argument);
 				}
 			}
-			if (_job.PreflightChange != 0)
+			if (_batch.PreflightChange != 0)
 			{
-				_additionalArguments.Add($"-set:PreflightChange={_job.PreflightChange}");
+				_additionalArguments.Add($"-set:PreflightChange={_batch.PreflightChange}");
 			}
 
 			if (RuntimeInformation.IsOSPlatform(OSPlatform.Windows))
@@ -316,6 +294,8 @@ namespace Horde.Agent.Execution
 					_logger.LogInformation("Unable to locate XGE directory. Not installed?");
 				}
 			}
+
+			return Task.CompletedTask;
 		}
 
 		public static bool IsUserAdministrator()
@@ -518,7 +498,7 @@ namespace Horde.Agent.Execution
 						thisAgentTypeName = validAgentTypeName;
 					}
 
-					if (_stream!.AgentTypes.ContainsKey(thisAgentTypeName))
+					if (_batch.ValidAgentTypes.Contains(thisAgentTypeName))
 					{
 						agentTypeName = thisAgentTypeName;
 						break;
@@ -628,17 +608,17 @@ namespace Horde.Agent.Execution
 					}
 				}
 
-				if (exportedBadge.Change == _job.Change || exportedBadge.Change == 0)
+				if (exportedBadge.Change == _batch.Change || exportedBadge.Change == 0)
 				{
 					createLabel.Change = LabelChange.Current;
 				}
-				else if (exportedBadge.Change == _job.CodeChange)
+				else if (exportedBadge.Change == _batch.CodeChange)
 				{
 					createLabel.Change = LabelChange.Code;
 				}
 				else
 				{
-					logger.LogWarning("Badge is set to display for changelist {Change}. This is neither the current changelist ({CurrentChange}) or the current code changelist ({CurrentCodeChange}).", exportedBadge.Change, _job.Change, _job.CodeChange);
+					logger.LogWarning("Badge is set to display for changelist {Change}. This is neither the current changelist ({CurrentChange}) or the current code changelist ({CurrentCodeChange}).", exportedBadge.Change, _batch.Change, _batch.CodeChange);
 				}
 
 				if (exportedBadge.Dependencies != null)
diff --git a/Engine/Source/Programs/Horde/Horde.Agent/Execution/PerforceExecutor.cs b/Engine/Source/Programs/Horde/Horde.Agent/Execution/PerforceExecutor.cs
--- a/Engine/Source/Programs/Horde/Horde.Agent/Execution/PerforceExecutor.cs
+++ b/Engine/Source/Programs/Horde/Horde.Agent/Execution/PerforceExecutor.cs
@@ -97,20 +97,20 @@ namespace Horde.Agent.Execution
 				_workspace = await Utility.WorkspaceInfo.SetupWorkspaceAsync(_workspaceInfo, _rootDir, useHaveTable, logger, cancellationToken);
 
 				// Figure out the change to build
-				if (_job.Change == 0)
+				if (_batch.Change == 0)
 				{
-					List<ChangesRecord> changes = await _workspace.PerforceClient.GetChangesAsync(ChangesOptions.None, 1, ChangeStatus.Submitted, new[] { _stream.Name + "/..." }, cancellationToken);
-					_job.Change = changes[0].Number;
+					List<ChangesRecord> changes = await _workspace.PerforceClient.GetChangesAsync(ChangesOptions.None, 1, ChangeStatus.Submitted, new[] { _batch.StreamName + "/..." }, cancellationToken);
+					_batch.Change = changes[0].Number;
 
 					UpdateJobRequest updateJobRequest = new UpdateJobRequest();
 					updateJobRequest.JobId = _jobId;
-					updateJobRequest.Change = _job.Change;
+					updateJobRequest.Change = _batch.Change;
 					await RpcConnection.InvokeAsync((JobRpc.JobRpcClient x) => x.UpdateJobAsync(updateJobRequest, null, null, cancellationToken), cancellationToken);
 				}
 
 				// Sync the workspace
-				int syncPreflightChange = (_job.ClonedPreflightChange != 0) ? _job.ClonedPreflightChange : _job.PreflightChange;
-				await _workspace.SyncAsync(_job.Change, syncPreflightChange, null, cancellationToken);
+				int syncPreflightChange = (_batch.ClonedPreflightChange != 0) ? _batch.ClonedPreflightChange : _batch.PreflightChange;
+				await _workspace.SyncAsync(_batch.Change, syncPreflightChange, null, cancellationToken);
 
 				// Remove any cached BuildGraph manifests
 				DirectoryReference manifestDir = DirectoryReference.Combine(_workspace.WorkspaceDir, "Engine", "Saved", "BuildGraph");
@@ -131,10 +131,10 @@ namespace Horde.Agent.Execution
 			DeleteEngineUserSettings(logger);
 
 			// Get the temp storage directory
-			if (!String.IsNullOrEmpty(_agentType!.TempStorageDir))
+			if (!String.IsNullOrEmpty(_batch.TempStorageDir))
 			{
-				string escapedStreamName = Regex.Replace(_stream!.Name, "[^a-zA-Z0-9_-]", "+");
-				_sharedStorageDir = DirectoryReference.Combine(new DirectoryReference(_agentType!.TempStorageDir), escapedStreamName, $"CL {_job!.Change} - Job {_jobId}");
+				string escapedStreamName = Regex.Replace(_batch.StreamName, "[^a-zA-Z0-9_-]", "+");
+				_sharedStorageDir = DirectoryReference.Combine(new DirectoryReference(_batch.TempStorageDir), escapedStreamName, $"CL {_batch.Change} - Job {_jobId}");
 				CopyAutomationTool(_sharedStorageDir, _workspace.WorkspaceDir, logger);
 			}
 
@@ -144,11 +144,11 @@ namespace Horde.Agent.Execution
 			_envVars["uebp_PORT"] = _workspace.ServerAndPort;
 			_envVars["uebp_USER"] = _workspace.UserName;
 			_envVars["uebp_CLIENT"] = _workspace.ClientName;
-			_envVars["uebp_BuildRoot_P4"] = _stream!.Name;
-			_envVars["uebp_BuildRoot_Escaped"] = _stream!.Name.Replace('/', '+');
+			_envVars["uebp_BuildRoot_P4"] = _batch.StreamName;
+			_envVars["uebp_BuildRoot_Escaped"] = _batch.StreamName.Replace('/', '+');
 			_envVars["uebp_CLIENT_ROOT"] = $"//{_workspace.ClientName}";
-			_envVars["uebp_CL"] = _job!.Change.ToString();
-			_envVars["uebp_CodeCL"] = _job!.CodeChange.ToString();
+			_envVars["uebp_CL"] = _batch.Change.ToString();
+			_envVars["uebp_CodeCL"] = _batch.CodeChange.ToString();
 			_envVars["P4USER"] = _workspace.UserName;
 			_envVars["P4CLIENT"] = _workspace.ClientName;
 
@@ -203,10 +203,10 @@ namespace Horde.Agent.Execution
 		PerforceLogger CreatePerforceLogger(ILogger logger)
 		{
 			PerforceLogger perforceLogger = new PerforceLogger(logger);
-			perforceLogger.AddClientView(_workspace.WorkspaceDir, _workspace.StreamView, _job.Change);
+			perforceLogger.AddClientView(_workspace.WorkspaceDir, _workspace.StreamView, _batch.Change);
 			if (_autoSdkWorkspace != null)
 			{
-				perforceLogger.AddClientView(_autoSdkWorkspace.WorkspaceDir, _autoSdkWorkspace.StreamView, _job.Change);
+				perforceLogger.AddClientView(_autoSdkWorkspace.WorkspaceDir, _autoSdkWorkspace.StreamView, _batch.Change);
 			}
 			return perforceLogger;
 		}
diff --git a/Engine/Source/Programs/Horde/Horde.Agent/Execution/WorkspaceExecutor.cs b/Engine/Source/Programs/Horde/Horde.Agent/Execution/WorkspaceExecutor.cs
--- a/Engine/Source/Programs/Horde/Horde.Agent/Execution/WorkspaceExecutor.cs
+++ b/Engine/Source/Programs/Horde/Horde.Agent/Execution/WorkspaceExecutor.cs
@@ -32,7 +32,7 @@ namespace Horde.Agent.Execution
 		{
 			await base.InitializeAsync(logger, cancellationToken);
 			
-			if (_job.Change == 0)
+			if (_batch.Change == 0)
 			{
 				throw new WorkspaceMaterializationException("Jobs with an empty change number are not supported");
 			}
@@ -47,7 +47,7 @@ namespace Horde.Agent.Execution
 				autoSdkWorkspaceSettings = await _autoSdkWorkspace.InitializeAsync(cancellationToken);
 
 				// Match change for AutoSDK and actual job change
-				int autoSdkChangeNumber = _job.Change;
+				int autoSdkChangeNumber = _batch.Change;
 
 				SyncOptions syncOptions = new();
 				await _autoSdkWorkspace.SyncAsync(autoSdkChangeNumber, syncOptions, cancellationToken);
@@ -60,8 +60,8 @@ namespace Horde.Agent.Execution
 				workspaceSettings = await _workspace.InitializeAsync(cancellationToken);
 				scope.Span.SetTag(Datadog.Trace.OpenTracing.DatadogTags.ResourceName, workspaceSettings.Identifier);
 				
-				int preflightChange = (_job.ClonedPreflightChange != 0) ? _job.ClonedPreflightChange : _job.PreflightChange;
-				await _workspace.SyncAsync(_job.Change, new SyncOptions(), cancellationToken);
+				int preflightChange = (_batch.ClonedPreflightChange != 0) ? _batch.ClonedPreflightChange : _batch.PreflightChange;
+				await _workspace.SyncAsync(_batch.Change, new SyncOptions(), cancellationToken);
 				
 				// TODO: Purging of cache for ManagedWorkspace did happen here in WorkspaceInfo
 				
@@ -78,10 +78,10 @@ namespace Horde.Agent.Execution
 			PerforceExecutor.DeleteEngineUserSettings(logger);
 
 			// Get the temp storage directory
-			if (!String.IsNullOrEmpty(_agentType!.TempStorageDir))
+			if (!String.IsNullOrEmpty(_batch.TempStorageDir))
 			{
-				string escapedStreamName = Regex.Replace(_stream!.Name, "[^a-zA-Z0-9_-]", "+");
-				_sharedStorageDir = DirectoryReference.Combine(new DirectoryReference(_agentType!.TempStorageDir), escapedStreamName, $"CL {_job!.Change} - Job {_jobId}");
+				string escapedStreamName = Regex.Replace(_batch.StreamName, "[^a-zA-Z0-9_-]", "+");
+				_sharedStorageDir = DirectoryReference.Combine(new DirectoryReference(_batch.TempStorageDir), escapedStreamName, $"CL {_batch.Change} - Job {_jobId}");
 				CopyAutomationTool(_sharedStorageDir, workspaceSettings.DirectoryPath, logger);
 			}
 
@@ -90,8 +90,8 @@ namespace Horde.Agent.Execution
 			_envVars["uebp_LOCAL_ROOT"] = workspaceSettings.DirectoryPath.FullName;
 			_envVars["uebp_BuildRoot_P4"] = workspaceSettings.StreamRoot;
 			_envVars["uebp_BuildRoot_Escaped"] = workspaceSettings.StreamRoot.Replace('/', '+');
-			_envVars["uebp_CL"] = _job!.Change.ToString();
-			_envVars["uebp_CodeCL"] = _job!.CodeChange.ToString();
+			_envVars["uebp_CL"] = _batch.Change.ToString();
+			_envVars["uebp_CodeCL"] = _batch.CodeChange.ToString();
 
 			if (autoSdkWorkspaceSettings != null)
 			{
diff --git a/Engine/Source/Programs/Horde/Horde.Agent/Leases/Handlers/JobHandler.cs b/Engine/Source/Programs/Horde/Horde.Agent/Leases/Handlers/JobHandler.cs
--- a/Engine/Source/Programs/Horde/Horde.Agent/Leases/Handlers/JobHandler.cs
+++ b/Engine/Source/Programs/Horde/Horde.Agent/Leases/Handlers/JobHandler.cs
@@ -123,11 +123,9 @@ namespace Horde.Agent.Leases.Handlers
 
 			// Start executing the current batch
 			BeginBatchResponse batch = await session.RpcConnection.InvokeAsync<JobRpc.JobRpcClient, BeginBatchResponse>(x => x.BeginBatchAsync(new BeginBatchRequest(executeTask.JobId, executeTask.BatchId, leaseId), null, null, cancellationToken), cancellationToken);
-
-			// Execute the batch
 			try
 			{
-				JobExecutorOptions options = new JobExecutorOptions(session, _serverStorageFactory, executeTask.JobId, executeTask.BatchId, batch.AgentType, new NamespaceId(executeTask.NamespaceId), executeTask.StoragePrefix, executeTask.Token, jobOptions);
+				JobExecutorOptions options = new JobExecutorOptions(session, _serverStorageFactory, executeTask.JobId, executeTask.BatchId, batch, new NamespaceId(executeTask.NamespaceId), executeTask.StoragePrefix, executeTask.Token, jobOptions);
 				await ExecuteBatchAsync(session, leaseId, executeTask.Workspace, executeTask.AutoSdkWorkspace, options, logger, cancellationToken);
 			}
 			catch (Exception ex)
diff --git a/Engine/Source/Programs/Horde/Horde.Build/Acls/AclService.cs b/Engine/Source/Programs/Horde/Horde.Build/Acls/AclService.cs
--- a/Engine/Source/Programs/Horde/Horde.Build/Acls/AclService.cs
+++ b/Engine/Source/Programs/Horde/Horde.Build/Acls/AclService.cs
@@ -7,6 +7,7 @@ using System.Linq;
 using System.Security.Claims;
 using System.Threading.Tasks;
 using Horde.Build.Agents;
+using Horde.Build.Agents.Leases;
 using Horde.Build.Agents.Sessions;
 using Horde.Build.Server;
 using Horde.Build.Utilities;
@@ -97,6 +98,11 @@ namespace Horde.Build.Acls
 
     internal static class ClaimExtensions
 	{
+		public static bool HasLeaseClaim(this ClaimsPrincipal user, LeaseId leaseId)
+		{
+			return user.HasClaim(HordeClaimTypes.Lease, leaseId.ToString());
+		}
+
 		public static bool HasSessionClaim(this ClaimsPrincipal user, SessionId sessionId)
 		{
 			return user.HasClaim(HordeClaimTypes.AgentSessionId, sessionId.ToString());
diff --git a/Engine/Source/Programs/Horde/Horde.Build/Agents/Leases/LeaseId.cs b/Engine/Source/Programs/Horde/Horde.Build/Agents/Leases/LeaseId.cs
--- a/Engine/Source/Programs/Horde/Horde.Build/Agents/Leases/LeaseId.cs
+++ b/Engine/Source/Programs/Horde/Horde.Build/Agents/Leases/LeaseId.cs
@@ -21,6 +21,22 @@ namespace Horde.Build.Agents.Leases
 		/// </summary>
 		public static LeaseId GenerateNewId() => new LeaseId(ObjectId.GenerateNewId());
 
+		/// <inheritdoc cref="ObjectId.TryParse(System.String, out ObjectId)"/>
+		public static bool TryParse(string text, out LeaseId leaseId)
+		{
+			ObjectId objectId;
+			if (ObjectId.TryParse(text, out objectId))
+			{
+				leaseId = new LeaseId(objectId);
+				return true;
+			}
+			else
+			{
+				leaseId = default;
+				return false;
+			}
+		}
+
 		/// <inheritdoc cref="ObjectId.Parse(System.String)"/>
 		public static LeaseId Parse(string text) => new LeaseId(ObjectId.Parse(text));
 
diff --git a/Engine/Source/Programs/Horde/Horde.Build/Jobs/JobRpcService.cs b/Engine/Source/Programs/Horde/Horde.Build/Jobs/JobRpcService.cs
--- a/Engine/Source/Programs/Horde/Horde.Build/Jobs/JobRpcService.cs
+++ b/Engine/Source/Programs/Horde/Horde.Build/Jobs/JobRpcService.cs
@@ -127,9 +127,13 @@ namespace Horde.Build.Jobs
 			{
 				throw new StructuredRpcException(StatusCode.PermissionDenied, "Batch {JobId}:{BatchId} has no session id", job.Id, batch.Id);
 			}
+			if (batch.LeaseId == null)
+			{
+				throw new StructuredRpcException(StatusCode.PermissionDenied, "Batch {JobId}:{BatchId} has no lease id", job.Id, batch.Id);
+			}
 
 			ClaimsPrincipal principal = context.GetHttpContext().User;
-			if (!principal.HasSessionClaim(batch.SessionId.Value))
+			if (!principal.HasSessionClaim(batch.SessionId.Value) && !principal.HasLeaseClaim(batch.LeaseId.Value))
 			{
 				throw new StructuredRpcException(StatusCode.PermissionDenied, "Session id {SessionId} not valid for step {JobId}:{BatchId}:{StepId}. Expected {ExpectedSessionId}.", principal.GetSessionClaim() ?? SessionId.Empty, job.Id, batch.Id, step.Id, batch.SessionId.Value);
 			}
@@ -315,10 +319,20 @@ namespace Horde.Build.Jobs
 			}
 
 			IGraph graph = await _jobService.GetGraphAsync(job);
-
+			AgentConfig agentConfig = streamConfig.AgentTypes[graph.Groups[batch.GroupIdx].AgentType];
+			
 			BeginBatchResponse response = new BeginBatchResponse();
 			response.LogId = batch.LogId.ToString();
 			response.AgentType = graph.Groups[batch.GroupIdx].AgentType;
+			response.StreamName = streamConfig.Name;
+			response.Change = job.Change;
+			response.CodeChange = job.CodeChange;
+			response.PreflightChange = job.PreflightChange;
+			response.ClonedPreflightChange = job.ClonedPreflightChange;
+			response.Arguments.AddRange(job.Arguments);
+			response.Environment.Add(agentConfig.Environment);
+			response.ValidAgentTypes.Add(streamConfig.AgentTypes.Keys);
+
 			return response;
 		}
 
@@ -586,9 +600,13 @@ namespace Horde.Build.Jobs
 			{
 				throw new StructuredRpcException(StatusCode.PermissionDenied, "Batch {JobId}:{BatchId} has no session id", job.Id, batchId);
 			}
+			if (batch.LeaseId == null)
+			{
+				throw new StructuredRpcException(StatusCode.PermissionDenied, "Batch {JobId}:{BatchId} has no lease id", job.Id, batchId);
+			}
 
 			ClaimsPrincipal principal = context.GetHttpContext().User;
-			if (!principal.HasSessionClaim(batch.SessionId.Value))
+			if (!principal.HasSessionClaim(batch.SessionId.Value) && !principal.HasLeaseClaim(batch.LeaseId.Value))
 			{
 				throw new StructuredRpcException(StatusCode.PermissionDenied, "Session id {SessionId} not valid for batch {JobId}:{BatchId}. Expected {ExpectedSessionId}.", principal.GetSessionClaim() ?? SessionId.Empty, job.Id, batchId, batch.SessionId.Value);
 			}
diff --git a/Engine/Source/Programs/Horde/Horde.Build/Jobs/JobService.cs b/Engine/Source/Programs/Horde/Horde.Build/Jobs/JobService.cs
--- a/Engine/Source/Programs/Horde/Horde.Build/Jobs/JobService.cs
+++ b/Engine/Source/Programs/Horde/Horde.Build/Jobs/JobService.cs
@@ -1328,11 +1328,11 @@ namespace Horde.Build.Jobs
 		/// <returns>True if the action is authorized</returns>
 		public static bool AuthorizeSession(IJobStepBatch batch, ClaimsPrincipal user)
 		{
-			if(batch.SessionId != null)
+			if (batch.SessionId != null)
 			{
 				foreach (Claim claim in user.Claims)
 				{
-					if(claim.Type == HordeClaimTypes.AgentSessionId)
+					if (claim.Type == HordeClaimTypes.AgentSessionId)
 					{
 						SessionId sessionIdValue;
 						if (SessionId.TryParse(claim.Value, out sessionIdValue) && sessionIdValue == batch.SessionId.Value)
@@ -1342,6 +1342,17 @@ namespace Horde.Build.Jobs
 					}
 				}
 			}
+			if (batch.LeaseId != null)
+			{
+				foreach (Claim claim in user.Claims)
+				{
+					LeaseId leaseIdValue;
+					if (LeaseId.TryParse(claim.Value, out leaseIdValue) && leaseIdValue == batch.LeaseId.Value)
+					{
+						return true;
+					}
+				}
+			}
 			return false;
 		}
 		
