Horde: Clarify job complete message in Slack

diff --git a/Engine/Source/Programs/Horde/Horde.Build/Notifications/Sinks/SlackNotificationSink.cs b/Engine/Source/Programs/Horde/Horde.Build/Notifications/Sinks/SlackNotificationSink.cs
--- a/Engine/Source/Programs/Horde/Horde.Build/Notifications/Sinks/SlackNotificationSink.cs
+++ b/Engine/Source/Programs/Horde/Horde.Build/Notifications/Sinks/SlackNotificationSink.cs
@@ -564,7 +564,7 @@ namespace Horde.Build.Notifications.Sinks
 						autoSubmitMessage = $"\n\n```{job.AutoSubmitMessage}```";
 					}
 
-					attachment.AddSection($"Files in CL *{job.PreflightChange}* were *not submitted*. Please resolve the following issues and submit manually.{autoSubmitMessage}");
+					attachment.AddSection($"Files in CL *{job.PreflightChange}* were *not submitted*. Please resolve the aforementioned issues and submit manually.{autoSubmitMessage}");
 				}
 			}
 
