Horde: Add workspace materializer for ManagedWorkspace

diff --git a/Engine/Source/Programs/Horde/Horde.Agent.Tests/PerforceExecutorTest.cs b/Engine/Source/Programs/Horde/Horde.Agent.Tests/PerforceExecutorTest.cs
deleted file mode 100644
diff --git a/Engine/Source/Programs/Horde/Horde.Agent.Tests/WorkspaceInfoTest.cs b/Engine/Source/Programs/Horde/Horde.Agent.Tests/WorkspaceInfoTest.cs
new file mode 100644
--- /dev/null
+++ b/Engine/Source/Programs/Horde/Horde.Agent.Tests/WorkspaceInfoTest.cs
@@ -0,0 +1,24 @@
+// Copyright Epic Games, Inc. All Rights Reserved.
+
+using Horde.Agent.Utility;
+using Microsoft.VisualStudio.TestTools.UnitTesting;
+
+namespace Horde.Agent.Tests;
+
+[TestClass]
+public class WorkspaceInfoTest
+{
+	[TestMethod]
+	public void ShouldUseHaveTable()
+	{
+		Assert.IsTrue(WorkspaceInfo.ShouldUseHaveTable(null));
+		Assert.IsTrue(WorkspaceInfo.ShouldUseHaveTable(""));
+		Assert.IsTrue(WorkspaceInfo.ShouldUseHaveTable("name=managedWorkspace"));
+		Assert.IsTrue(WorkspaceInfo.ShouldUseHaveTable("name=managedWorkspace$#@!@#"));
+		Assert.IsTrue(WorkspaceInfo.ShouldUseHaveTable("name=managedWorkspace$#@!@#"));
+		Assert.IsTrue(WorkspaceInfo.ShouldUseHaveTable("name=managedWorkspace&useHaveTable=true"));
+		
+		Assert.IsFalse(WorkspaceInfo.ShouldUseHaveTable("name=managedWorkspace&useHaveTable=false"));
+		Assert.IsFalse(WorkspaceInfo.ShouldUseHaveTable("name=ManagedWorkspace&useHaveTable=FalsE"));
+	}
+}
diff --git a/Engine/Source/Programs/Horde/Horde.Agent/Execution/ManagedWorkspaceMaterializer.cs b/Engine/Source/Programs/Horde/Horde.Agent/Execution/ManagedWorkspaceMaterializer.cs
new file mode 100644
--- /dev/null
+++ b/Engine/Source/Programs/Horde/Horde.Agent/Execution/ManagedWorkspaceMaterializer.cs
@@ -0,0 +1,160 @@
+// Copyright Epic Games, Inc. All Rights Reserved.
+
+using System.Threading;
+using System.Threading.Tasks;
+using EpicGames.Core;
+using Horde.Agent.Utility;
+using HordeCommon.Rpc.Messages;
+using Microsoft.Extensions.Logging;
+using OpenTracing;
+using OpenTracing.Util;
+
+namespace Horde.Agent.Execution;
+
+/// <summary>
+/// Workspace materializer wrapping ManagedWorkspace and WorkspaceInfo
+/// </summary>
+public class ManagedWorkspaceMaterializer : IWorkspaceMaterializer
+{
+	private readonly AgentWorkspace _agentWorkspace;
+	private readonly DirectoryReference _workingDir;
+	private readonly bool _useSyncMarker;
+	private readonly bool _useCacheFile;
+	private readonly ILogger<ManagedWorkspaceMaterializer> _logger;
+	private WorkspaceInfo? _workspace;
+	
+	/// <summary>
+	/// Constructor
+	/// </summary>
+	/// <param name="agentWorkspace">Workspace configuration</param>
+	/// <param name="workingDir">Where to put synced Perforce files and any cached data/metadata</param>
+	/// <param name="useSyncMarker">Whether to use a sync marker for identifying last synced change number</param>
+	/// <param name="useCacheFile">Whether to use a cache file during syncs</param>
+	/// <param name="logger">Logger</param>
+	public ManagedWorkspaceMaterializer(
+		AgentWorkspace agentWorkspace,
+		DirectoryReference workingDir,
+		bool useSyncMarker,
+		bool useCacheFile,
+		ILogger<ManagedWorkspaceMaterializer> logger)
+	{
+		_agentWorkspace = agentWorkspace;
+		_workingDir = workingDir;
+		_useSyncMarker = useSyncMarker;
+		_useCacheFile = useCacheFile;
+		_logger = logger;
+	}
+
+	/// <inheritdoc/>
+	public async Task<WorkspaceMaterializerSettings> InitializeAsync(CancellationToken cancellationToken)
+	{
+		using IScope scope = CreateTraceSpan("ManagedWorkspaceMaterializer.InitializeAsync");
+		
+		bool useHaveTable = WorkspaceInfo.ShouldUseHaveTable(_agentWorkspace.Method);
+		_workspace = await WorkspaceInfo.SetupWorkspaceAsync(_agentWorkspace, _workingDir, useHaveTable, _logger, cancellationToken);
+		return new WorkspaceMaterializerSettings(_workingDir, _agentWorkspace.Identifier, _agentWorkspace.Stream);
+	}
+
+	/// <inheritdoc/>
+	public async Task FinalizeAsync(CancellationToken cancellationToken)
+	{
+		using IScope scope = CreateTraceSpan("ManagedWorkspaceMaterializer.FinalizeAsync");
+
+		if (_workspace != null)
+		{
+			await _workspace.CleanAsync(cancellationToken);	
+		}
+	}
+
+	/// <inheritdoc/>
+	public Task<WorkspaceMaterializerSettings> GetSettingsAsync(CancellationToken cancellationToken)
+	{
+		return Task.FromResult(new WorkspaceMaterializerSettings(_workingDir, _agentWorkspace.Identifier, _agentWorkspace.Stream));
+	}
+
+	/// <inheritdoc/>
+	public async Task SyncAsync(int changeNum, SyncOptions options, CancellationToken cancellationToken)
+	{
+		using IScope scope = CreateTraceSpan("ManagedWorkspaceMaterializer.SyncAsync");
+		scope.Span.SetTag("ChangeNum", changeNum);
+		scope.Span.SetTag("RemoveUntracked", options.RemoveUntracked);
+
+		if (_workspace == null)
+		{
+			throw new WorkspaceMaterializationException("Workspace not initialized");
+		}
+
+		if (!IsAlreadySynced(changeNum))
+		{
+			FileReference? cacheFile = _useCacheFile ? FileReference.Combine(_workspace.MetadataDir, "Contents.dat") : null;
+			await _workspace.SyncAsync(changeNum, 0, cacheFile, cancellationToken);
+			MarkChangeNumSynced(changeNum);
+		}
+	}
+
+	private (FileReference file, string fileContent) GetSyncMarker(int changeNum)
+	{
+		if (_workspace == null)
+		{
+			throw new WorkspaceMaterializationException("Workspace not initialized");
+		}
+
+		return (FileReference.Combine(_workspace.MetadataDir, "Synced.txt"), $"Synced to CL {changeNum}");
+	}
+	
+	/// <summary>
+	/// Check if workspace is already synced to given change number
+	/// </summary>
+	/// <param name="changeNum">Change number to check</param>
+	/// <returns>True if already synced</returns>
+	private bool IsAlreadySynced(int changeNum)
+	{
+		if (!_useSyncMarker) return false;
+
+		(FileReference syncFile, string syncText) = GetSyncMarker(changeNum);
+		if (!FileReference.Exists(syncFile) || FileReference.ReadAllText(syncFile) != syncText)
+		{
+			FileReference.Delete(syncFile);
+			return false;
+		}
+
+		return true;
+	}
+	
+	/// <summary>
+	/// Mark a change number as synced with a file on disk
+	/// </summary>
+	/// <param name="changeNum">Change number to mark</param>
+	/// <returns>True if already synced</returns>
+	private void MarkChangeNumSynced(int changeNum)
+	{
+		if (!_useSyncMarker) return;
+		
+		(FileReference syncFile, string syncText) = GetSyncMarker(changeNum);
+		FileReference.WriteAllText(syncFile, syncText);
+	}
+
+	/// <inheritdoc/>
+	public async Task UnshelveAsync(int changeNum, CancellationToken cancellationToken)
+	{
+		using IScope scope = CreateTraceSpan("ManagedWorkspaceMaterializer.UnshelveAsync");
+		scope.Span.SetTag("ChangeNum", changeNum);
+		
+		if (_workspace == null)
+		{
+			throw new WorkspaceMaterializationException("Workspace not initialized");
+		}
+		
+		await _workspace.UnshelveAsync(changeNum, cancellationToken);
+	}
+
+	private IScope CreateTraceSpan(string operationName)
+	{
+		IScope scope = GlobalTracer.Instance.BuildSpan(operationName).WithResourceName(_agentWorkspace.Identifier).StartActive();
+		scope.Span.SetTag("UseHaveTable", WorkspaceInfo.ShouldUseHaveTable(_agentWorkspace.Method));
+		scope.Span.SetTag("Cluster", _agentWorkspace.Cluster);
+		scope.Span.SetTag("Incremental", _agentWorkspace.Incremental);
+		scope.Span.SetTag("Method", _agentWorkspace.Method);
+		return scope;
+	}
+}
\ No newline at end of file
diff --git a/Engine/Source/Programs/Horde/Horde.Agent/Execution/PerforceExecutor.cs b/Engine/Source/Programs/Horde/Horde.Agent/Execution/PerforceExecutor.cs
--- a/Engine/Source/Programs/Horde/Horde.Agent/Execution/PerforceExecutor.cs
+++ b/Engine/Source/Programs/Horde/Horde.Agent/Execution/PerforceExecutor.cs
@@ -57,8 +57,8 @@ namespace Horde.Agent.Execution
 			{
 				using IScope _ = GlobalTracer.Instance.BuildSpan("Workspace").WithResourceName("AutoSDK").StartActive();
 				
-				bool useHaveTable = ShouldUseHaveTable(_autoSdkWorkspaceInfo.Method);
-				_autoSdkWorkspace = await Utility.WorkspaceInfo.SetupWorkspaceAsync(_autoSdkWorkspaceInfo, _rootDir, useHaveTable, logger, cancellationToken);
+				bool useHaveTable = WorkspaceInfo.ShouldUseHaveTable(_autoSdkWorkspaceInfo.Method);
+				_autoSdkWorkspace = await WorkspaceInfo.SetupWorkspaceAsync(_autoSdkWorkspaceInfo, _rootDir, useHaveTable, logger, cancellationToken);
 
 				DirectoryReference legacyDir = DirectoryReference.Combine(_autoSdkWorkspace.MetadataDir, "HostWin64");
 				if (DirectoryReference.Exists(legacyDir))
@@ -93,8 +93,8 @@ namespace Horde.Agent.Execution
 			using (IScope scope = GlobalTracer.Instance.BuildSpan("Workspace").WithResourceName(_workspaceInfo.Identifier).StartActive())
 			{
 				// Sync the regular workspace
-				bool useHaveTable = ShouldUseHaveTable(_workspaceInfo.Method);
-				_workspace = await Utility.WorkspaceInfo.SetupWorkspaceAsync(_workspaceInfo, _rootDir, useHaveTable, logger, cancellationToken);
+				bool useHaveTable = WorkspaceInfo.ShouldUseHaveTable(_workspaceInfo.Method);
+				_workspace = await WorkspaceInfo.SetupWorkspaceAsync(_workspaceInfo, _rootDir, useHaveTable, logger, cancellationToken);
 
 				// Figure out the change to build
 				if (_batch.Change == 0)
@@ -245,8 +245,8 @@ namespace Horde.Agent.Execution
 			List<WorkspaceInfo> workspaces = new List<WorkspaceInfo>();
 			foreach (AgentWorkspace pendingWorkspace in pendingWorkspaces)
 			{
-				bool useHaveTable = ShouldUseHaveTable(pendingWorkspace.Method);
-				WorkspaceInfo workspace = await Utility.WorkspaceInfo.SetupWorkspaceAsync(pendingWorkspace, rootDir, useHaveTable, logger, cancellationToken);
+				bool useHaveTable = WorkspaceInfo.ShouldUseHaveTable(pendingWorkspace.Method);
+				WorkspaceInfo workspace = await WorkspaceInfo.SetupWorkspaceAsync(pendingWorkspace, rootDir, useHaveTable, logger, cancellationToken);
 				workspaces.Add(workspace);
 			}
 
@@ -402,37 +402,6 @@ namespace Horde.Agent.Execution
 				}
 			}
 		}
-
-		/// <summary>
-		/// Parse a text string as a query string to determine use of have table
-		/// </summary>
-		/// <param name="method">Method text string from workspace config</param>
-		/// <returns>True if have table should be enabled</returns>
-		public static bool ShouldUseHaveTable(string? method)
-		{
-			const string NameKey = "name";
-			const string ManagedWorkspaceValue = "managedWorkspace";
-			const string UseHaveTableKey = "useHaveTable";
-			
-			if (String.IsNullOrEmpty(method))
-			{
-				return true;
-			}
-
-			NameValueCollection nameValues = HttpUtility.ParseQueryString(method);
-			string? name = nameValues[NameKey];
-			string? useHaveTable = nameValues[UseHaveTableKey];
-			
-			if (name != null && name.Equals(ManagedWorkspaceValue, StringComparison.OrdinalIgnoreCase))
-			{
-				if (useHaveTable != null && useHaveTable.Equals("false", StringComparison.OrdinalIgnoreCase))
-				{
-					return false;
-				}
-			}
-
-			return true;
-		}
 	}
 
 	class PerforceExecutorFactory : IJobExecutorFactory
diff --git a/Engine/Source/Programs/Horde/Horde.Agent/Execution/WorkspaceExecutor.cs b/Engine/Source/Programs/Horde/Horde.Agent/Execution/WorkspaceExecutor.cs
--- a/Engine/Source/Programs/Horde/Horde.Agent/Execution/WorkspaceExecutor.cs
+++ b/Engine/Source/Programs/Horde/Horde.Agent/Execution/WorkspaceExecutor.cs
@@ -129,15 +129,31 @@ namespace Horde.Agent.Execution
 
 	class WorkspaceExecutorFactory : IJobExecutorFactory
 	{
+		private readonly ILoggerFactory _loggerFactory;
+		
 		public string Name => WorkspaceExecutor.Name;
 
-		public WorkspaceExecutorFactory()
+		public WorkspaceExecutorFactory(ILoggerFactory loggerFactory)
 		{
+			_loggerFactory = loggerFactory;
 		}
 
 		public IJobExecutor CreateExecutor(AgentWorkspace workspaceInfo, AgentWorkspace? autoSdkWorkspaceInfo, JobExecutorOptions options)
 		{
-			throw new NotImplementedException("WorkspaceExecutor cannot be instantiated as there are no IWorkspaceMaterializer implementations yet");
+			IWorkspaceMaterializer workspaceMaterializer;
+			IWorkspaceMaterializer? autoSdkMaterializer = null;
+
+			// Default to ManagedWorkspaceMaterializer for now. Different types of materializer can later be set via AgentWorkspace or JobOptions.
+			workspaceMaterializer = new ManagedWorkspaceMaterializer(
+				workspaceInfo, options.Session.WorkingDir, false, false, _loggerFactory.CreateLogger<ManagedWorkspaceMaterializer>());
+			
+			if (autoSdkWorkspaceInfo != null)
+			{
+				autoSdkMaterializer = new ManagedWorkspaceMaterializer(
+					autoSdkWorkspaceInfo, options.Session.WorkingDir, true, true, _loggerFactory.CreateLogger<ManagedWorkspaceMaterializer>());
+			}
+			
+			return new WorkspaceExecutor(options, autoSdkMaterializer, workspaceMaterializer, _loggerFactory.CreateLogger<WorkspaceExecutor>());
 		}
 	}
 }
diff --git a/Engine/Source/Programs/Horde/Horde.Agent/Utility/WorkspaceInfo.cs b/Engine/Source/Programs/Horde/Horde.Agent/Utility/WorkspaceInfo.cs
--- a/Engine/Source/Programs/Horde/Horde.Agent/Utility/WorkspaceInfo.cs
+++ b/Engine/Source/Programs/Horde/Horde.Agent/Utility/WorkspaceInfo.cs
@@ -2,10 +2,12 @@
 
 using System;
 using System.Collections.Generic;
+using System.Collections.Specialized;
 using System.Linq;
 using System.Text.RegularExpressions;
 using System.Threading;
 using System.Threading.Tasks;
+using System.Web;
 using EpicGames.Core;
 using EpicGames.Perforce;
 using EpicGames.Perforce.Managed;
@@ -232,6 +234,37 @@ namespace Horde.Agent.Utility
 			return new WorkspaceInfo(perforceClient, hostName, streamName, streamView, metadataDir, workspaceDir, view, removeUntrackedFiles, newRepository);
 		}
 
+		/// <summary>
+		/// Parse a text string as a query string to determine use of have table
+		/// </summary>
+		/// <param name="method">Method text string from workspace config</param>
+		/// <returns>True if have table should be enabled</returns>
+		public static bool ShouldUseHaveTable(string? method)
+		{
+			const string NameKey = "name";
+			const string ManagedWorkspaceValue = "managedWorkspace";
+			const string UseHaveTableKey = "useHaveTable";
+			
+			if (String.IsNullOrEmpty(method))
+			{
+				return true;
+			}
+
+			NameValueCollection nameValues = HttpUtility.ParseQueryString(method);
+			string? name = nameValues[NameKey];
+			string? useHaveTable = nameValues[UseHaveTableKey];
+			
+			if (name != null && name.Equals(ManagedWorkspaceValue, StringComparison.OrdinalIgnoreCase))
+			{
+				if (useHaveTable != null && useHaveTable.Equals("false", StringComparison.OrdinalIgnoreCase))
+				{
+					return false;
+				}
+			}
+
+			return true;
+		}
+		
 		/// <summary>
 		/// Gets the latest change in the stream
 		/// </summary>
@@ -314,10 +347,20 @@ namespace Horde.Agent.Utility
 				await Repository.PurgeAsync(0, cancellationToken);
 			}
 
-			// if we're running a preflight
-			if (preflightChange > 0)
+			await UnshelveAsync(preflightChange, cancellationToken);
+		}
+
+		/// <summary>
+		/// Unshelves a changelist
+		/// Assumes base changelist already has been synced.
+		/// </summary>
+		/// <param name="change">Change number to unshelve</param>
+		/// <param name="cancellationToken">Cancellation token</param>
+		public async Task UnshelveAsync(int change, CancellationToken cancellationToken)
+		{
+			if (change > 0)
 			{
-				await Repository.UnshelveAsync(PerforceClient, preflightChange, cancellationToken);
+				await Repository.UnshelveAsync(PerforceClient, change, cancellationToken);
 			}
 		}
 
