Horde: Add missing file

diff --git a/Engine/Source/Programs/Horde/Horde.Agent/Program.cs b/Engine/Source/Programs/Horde/Horde.Agent/Program.cs
--- a/Engine/Source/Programs/Horde/Horde.Agent/Program.cs
+++ b/Engine/Source/Programs/Horde/Horde.Agent/Program.cs
@@ -182,6 +182,7 @@ namespace Horde.Agent
 			services.AddSingleton<GrpcService>();
 
 			services.AddSingleton<IJobExecutorFactory, PerforceExecutorFactory>();
+			services.AddSingleton<IJobExecutorFactory, WorkspaceExecutorFactory>();
 			services.AddSingleton<IJobExecutorFactory, LocalExecutorFactory>();
 			services.AddSingleton<IJobExecutorFactory, TestExecutorFactory>();
 
