Horde: Fix agent model on agents view
#rnx
#skipci

diff --git a/Engine/Source/Programs/Horde/HordeDashboard/src/App.tsx b/Engine/Source/Programs/Horde/HordeDashboard/src/App.tsx
--- a/Engine/Source/Programs/Horde/HordeDashboard/src/App.tsx
+++ b/Engine/Source/Programs/Horde/HordeDashboard/src/App.tsx
@@ -133,7 +133,7 @@ const Main: React.FC = () => {
                { path: "log/:logId", element: <LogView /> },
                { path: "testreport/:testdataId", element: <TestReportView /> },
                { path: "stream/:streamId", element: <StreamView /> },
-               { path: "agents/:agentId?", element: <AgentView /> },
+               { path: "agents", element: <AgentView/> },
                { path: "admin/token", element: <AdminToken /> },
                { path: "reports/utilization", element: <UtilizationReportView /> },
                { path: "preflight", element: <PreflightRedirector /> },
diff --git a/Engine/Source/Programs/Horde/HordeDashboard/src/components/AgentView.tsx b/Engine/Source/Programs/Horde/HordeDashboard/src/components/AgentView.tsx
--- a/Engine/Source/Programs/Horde/HordeDashboard/src/components/AgentView.tsx
+++ b/Engine/Source/Programs/Horde/HordeDashboard/src/components/AgentView.tsx
@@ -4,7 +4,7 @@ import { action, makeObservable, observable } from 'mobx';
 import { observer } from 'mobx-react-lite';
 import moment from 'moment-timezone';
 import React, { createRef, useEffect, useState } from 'react';
-import { Link, useNavigate, useParams, useSearchParams } from 'react-router-dom';
+import { Link, useSearchParams } from 'react-router-dom';
 import Marquee from 'react-text-marquee';
 import backend from '../backend';
 import { agentStore } from '../backend/AgentStore';
@@ -188,6 +188,7 @@ function getTaskTime(agent: GetAgentResponse): number {
 }
 
 type SearchState = {
+   agentId?: string;
    agentSearch?: string;
    exactSearch?: boolean;
    filter?: string[];
@@ -268,6 +269,13 @@ class LocalState {
       this.filterExactMatch = match;
    }
 
+   @action
+   setAgentId(id?: string) {
+      this.searchState.agentId = id;
+      this.updateSearch();    
+   }
+
+
    @action
    private _onColumnClick(ev: React.MouseEvent<HTMLElement>, column: IColumn) {
       if (column.key === "pools") {
@@ -411,7 +419,7 @@ class LocalState {
             colState.isChecked = true;
             this.searchState.columnMode = key ? key : undefined;
             this.updateSearch();
-      
+
          }
          else {
             colState.isChecked = false;
@@ -420,7 +428,7 @@ class LocalState {
 
       //this.searchState.columnMode = key;
       //this.updateSearch();
-      
+
       this.columnMenuProps = this._updateColumnProps();
    }
 
@@ -544,6 +552,10 @@ class LocalState {
          search.append("exact", "true");
       }
 
+      if (state.agentId) {
+         search.append("agentId", state.agentId);
+      }
+
       state.filter?.forEach(f => {
          if (f) {
             search.append("filter", f);
@@ -569,6 +581,7 @@ class LocalState {
 
       const filters = search.getAll("filter") ?? undefined;
       const agentSearch = search.get("agent") ?? undefined;
+      const agentId = search.get("agentId") ?? undefined;
       const exact = search.get("exact") ?? undefined;
       const mode = search.get("mode") ?? undefined;
 
@@ -576,9 +589,10 @@ class LocalState {
       state.columnMode = mode?.length ? mode : undefined;
       state.agentSearch = agentSearch?.length ? agentSearch : undefined;
       state.exactSearch = exact?.trim() === "true" ? true : undefined;
+      state.agentId = agentId?.trim() ? agentId : undefined;
 
       this.searchState = state;
- 
+
       if (state.agentSearch) {
          this.agentFilter = state.agentSearch;
       }
@@ -1698,22 +1712,22 @@ export const SearchUpdate: React.FC = observer(() => {
    // subscribe
    if (localState.searchUpdated) { }
 
-   const csearch = localState.search.toString();   
+   const csearch = localState.search.toString();
 
    if (state.search !== csearch) {
       setSearchParams(csearch);
-      setState({ search: csearch });      
+      setState({ search: csearch });
    }
 
    return null;
 });
 
-export const AgentView: React.FC = observer(() => {
-   const { agentId } = useParams<{ agentId: string }>();
-   const navigate = useNavigate();
+export const AgentView: React.FC = observer(() => {   
    const [initAgentUpdater, setInitAgentUpdater] = useState(false);
    const [searchParams] = useSearchParams();
 
+   const agentId = searchParams.get("agentId") ? searchParams.get("agentId") : undefined;   
+
    // adjust automatically to viewport changes
    useWindowSize();
 
@@ -2008,7 +2022,7 @@ export const AgentView: React.FC = observer(() => {
    }
 
    function onHistoryModalDismiss() {
-      navigate('/agents');
+      localState.setAgentId(undefined);
    }
 
    let agentItems = agentStore.agents.filter(filterAgents, localState.agentFilter).sort(sortAgents);
@@ -2244,7 +2258,7 @@ export const AgentView: React.FC = observer(() => {
                      {getAgentStatusIcon(agent)}
                   </Stack.Item>
                   <Stack.Item align={"center"}>
-                     <ReactLink styles={{ root: { paddingTop: '1px' } }} title="Lease and Session History" onClick={() => { navigate(`/agents/${agent.id}`); }}>{agent.name}</ReactLink>
+                     <ReactLink styles={{ root: { paddingTop: '1px' } }} title="Lease and Session History" onClick={() => localState.setAgentId(agent.id) }>{agent.name}</ReactLink>
                   </Stack.Item>
                </Stack>
             );
diff --git a/Engine/Source/Programs/Horde/HordeDashboard/src/components/AuditLog.tsx b/Engine/Source/Programs/Horde/HordeDashboard/src/components/AuditLog.tsx
--- a/Engine/Source/Programs/Horde/HordeDashboard/src/components/AuditLog.tsx
+++ b/Engine/Source/Programs/Horde/HordeDashboard/src/components/AuditLog.tsx
@@ -399,7 +399,7 @@ export const AuditLogPanel: React.FC<{ agentId?: string, issueId?: string }> = o
          },
          {
             text: `Agent Audit - ${agentId}`,
-            link: `/agents/${encodeURIComponent(agentId)}`
+            link: `/agents?agentId=${encodeURIComponent(agentId)}`
          }
       ];
 
diff --git a/Engine/Source/Programs/Horde/HordeDashboard/src/components/LogSource.tsx b/Engine/Source/Programs/Horde/HordeDashboard/src/components/LogSource.tsx
--- a/Engine/Source/Programs/Horde/HordeDashboard/src/components/LogSource.tsx
+++ b/Engine/Source/Programs/Horde/HordeDashboard/src/components/LogSource.tsx
@@ -618,11 +618,11 @@ class LeaseLogSource extends LogSource {
          },
          {
             text: `${this.agent?.name}`,
-            link: `/agents/${this.agent?.id}`
+            link: `/agents?agentId=${this.agent?.id}`
          },
          {
             text: `${this.lease?.name ?? this.leaseId}`,
-            link: `/agents/${this.agent?.id}`
+            link: `/agents?agentId=${this.agent?.id}`
          }
 
       ];
diff --git a/Engine/Source/Programs/Horde/HordeDashboard/src/components/jobDetailsV2/JobDetailViewSteps.tsx b/Engine/Source/Programs/Horde/HordeDashboard/src/components/jobDetailsV2/JobDetailViewSteps.tsx
--- a/Engine/Source/Programs/Horde/HordeDashboard/src/components/jobDetailsV2/JobDetailViewSteps.tsx
+++ b/Engine/Source/Programs/Horde/HordeDashboard/src/components/jobDetailsV2/JobDetailViewSteps.tsx
@@ -245,9 +245,7 @@ export const StepsPanelInner: React.FC<{ jobDetails: JobDetailsV2, depStepId?: s
                   return <Stack horizontal disableShrink={true} horizontalAlign="start">
                      {(error || warning) && <Icon styles={{ root: { paddingTop: 3, paddingRight: 8 } }} className={style!.className} iconName="Square" />}
                      <Link to={`/agents?agent=${encodeURI(item.agentPool)}&exact=true`} ><Text nowrap={true}>{batchText}</Text></Link>
-                  </Stack>
-                     ;
-
+                  </Stack>                     
                }
                return <Stack horizontal disableShrink={true} horizontalAlign="start">
                   {(error || warning) && <Icon styles={{ root: { paddingTop: 3, paddingRight: 8 } }} className={style!.className} iconName="Square" />}
