Horde: Optimization to trendsview
#rnx
#skipci

diff --git a/Engine/Source/Programs/Horde/HordeDashboard/src/components/jobDetailsV2/JobDetailStepTrends.tsx b/Engine/Source/Programs/Horde/HordeDashboard/src/components/jobDetailsV2/JobDetailStepTrends.tsx
--- a/Engine/Source/Programs/Horde/HordeDashboard/src/components/jobDetailsV2/JobDetailStepTrends.tsx
+++ b/Engine/Source/Programs/Horde/HordeDashboard/src/components/jobDetailsV2/JobDetailStepTrends.tsx
@@ -455,7 +455,8 @@ const StepTrendsPanelInner: React.FC<{ jobDetails: JobDetailsV2; dataView: Trend
    if (!dataView.cacheData) {
       const query: JobTimingsQuery = {
          streamId: jobDetails!.stream!.id,
-         template: [selectedTemplateRef.id]
+         template: [selectedTemplateRef.id],
+         count: selectedTemplateRef.id?.indexOf("incremental") === -1 ? 50 : 100
       };
       dataView.populateChartDataCache(query);
    }
