Horde: Promote pool view to top level
#skipci
#rnx

diff --git a/Engine/Source/Programs/Horde/HordeDashboard/src/App.tsx b/Engine/Source/Programs/Horde/HordeDashboard/src/App.tsx
--- a/Engine/Source/Programs/Horde/HordeDashboard/src/App.tsx
+++ b/Engine/Source/Programs/Horde/HordeDashboard/src/App.tsx
@@ -128,7 +128,7 @@ const Main: React.FC = () => {
             path: "/", element: <Root />, errorElement: <RouteError />, children: [
                { path: "index", element: <UserHomeView /> },
                { path: "project/:projectId", element: <ProjectHome /> },
-               { path: "pool/:poolId", element: <PoolView /> },
+               { path: "pools", element: <PoolView /> },
                { path: "job/:jobId", element: <JobDetailViewV2 /> },
                { path: "log/:logId", element: <LogView /> },
                { path: "testreport/:testdataId", element: <TestReportView /> },
diff --git a/Engine/Source/Programs/Horde/HordeDashboard/src/components/AgentView.tsx b/Engine/Source/Programs/Horde/HordeDashboard/src/components/AgentView.tsx
--- a/Engine/Source/Programs/Horde/HordeDashboard/src/components/AgentView.tsx
+++ b/Engine/Source/Programs/Horde/HordeDashboard/src/components/AgentView.tsx
@@ -1583,7 +1583,7 @@ export const PoolEditorModal: React.FC = observer(() => {
                <Stack styles={{ root: { paddingTop: 10 } }}>
                   <Stack horizontal>
                      <Stack style={{ paddingLeft: 12 }}>
-                        <DefaultButton href={`/pool/${editPoolsModalState.lastSelectedPool?.pool?.id}`} target="_blank">Details</DefaultButton>
+                        <DefaultButton href={`/pools?pool=${editPoolsModalState.lastSelectedPool?.pool?.id}`} target="_blank">Details</DefaultButton>
                      </Stack>
                      <Stack grow />
                      <PrimaryButton disabled={!editPoolsModalState.isPoolValueValid} onClick={() => { editPoolsModalState.isDirectEdit ? editPoolsModalState.saveChanges() : editPoolsModalState.setEditorOpen(false, false); }} styles={{ root: { marginRight: "10px" } }}>
@@ -2296,7 +2296,7 @@ export const AgentView: React.FC = observer(() => {
                            key: 'agent_pool_view',
                            text: 'View Pool',
                            onClick: (ev) => {
-                              window.open(`/pool/${poolObjs[idx]?.id}`, '_blank')
+                              window.open(`/pools?pool=${poolObjs[idx]?.id}`, '_blank')
                            }
                         },
                         {
diff --git a/Engine/Source/Programs/Horde/HordeDashboard/src/components/HistoryModal.tsx b/Engine/Source/Programs/Horde/HordeDashboard/src/components/HistoryModal.tsx
--- a/Engine/Source/Programs/Horde/HordeDashboard/src/components/HistoryModal.tsx
+++ b/Engine/Source/Programs/Horde/HordeDashboard/src/components/HistoryModal.tsx
@@ -475,7 +475,7 @@ export const HistoryModal: React.FC<{ agentId: string | undefined, onDismiss: (.
          poolItems.push({
             key: `pools_${poolId}`,
             text: pool.name,
-            href: `/pool/${pool.id}`,
+            href: `/pools?pool=${pool.id}`,
             target: "_blank"
          })
       }
diff --git a/Engine/Source/Programs/Horde/HordeDashboard/src/components/PoolView.tsx b/Engine/Source/Programs/Horde/HordeDashboard/src/components/PoolView.tsx
--- a/Engine/Source/Programs/Horde/HordeDashboard/src/components/PoolView.tsx
+++ b/Engine/Source/Programs/Horde/HordeDashboard/src/components/PoolView.tsx
@@ -1,10 +1,11 @@
 // Copyright Epic Games, Inc. All Rights Reserved.
 
-import { Checkbox, ComboBox, DefaultButton, DetailsList, DetailsListLayoutMode, Dialog, DialogFooter, DialogType, IColumn, IComboBoxOption, MessageBar, MessageBarType, Position, PrimaryButton, ScrollablePane, ScrollbarVisibility, SelectionMode, Slider, SpinButton, Spinner, SpinnerSize, Stack, Text, TextField } from "@fluentui/react";
+import { Checkbox, ComboBox, DefaultButton, DetailsList, DetailsListLayoutMode, Dialog, DialogFooter, DialogType, IColumn, IComboBoxOption, ITag, MessageBar, MessageBarType, Position, PrimaryButton, ScrollablePane, ScrollbarVisibility, SelectionMode, Slider, SpinButton, Spinner, SpinnerSize, Stack, TagPicker, Text, TextField } from "@fluentui/react";
 import { observer } from "mobx-react-lite";
-import React, { useEffect, useState } from "react";
-import { Link, useParams } from "react-router-dom";
+import React, { useEffect, useId, useState } from "react";
+import { Link, useParams, useSearchParams } from "react-router-dom";
 import backend from "../backend";
+import { agentStore } from "../backend/AgentStore";
 import { GetAgentResponse, GetBatchResponse, GetJobTimingResponse, GetPoolResponse, GetStepResponse, JobData, JobQuery, JobState, JobStepBatchState, JobStepState, PoolSizeStrategy, StepData, UpdatePoolRequest } from "../backend/Api";
 import dashboard from "../backend/Dashboard";
 import { PollBase } from "../backend/PollBase";
@@ -29,9 +30,7 @@ const UNSET_VALUE: string = "-";
 class PoolHandler extends PollBase {
 
    constructor(pollTime = 15000) {
-
       super(pollTime);
-
    }
 
    clear() {
@@ -46,6 +45,20 @@ class PoolHandler extends PollBase {
       super.stop();
    }
 
+   loadPools() {
+
+      if (PoolHandler.pools) {
+         return;
+      }
+
+      PoolHandler.pools = [];
+
+      backend.getPools().then(p => {
+         PoolHandler.pools = p.sort((a, b) => a.name.localeCompare(b.name));
+         this.setUpdated();
+      });
+   }
+
    set(poolId?: string) {
 
       if (poolId === this.poolId) {
@@ -59,7 +72,6 @@ class PoolHandler extends PollBase {
       } else {
          this.poolId = undefined;
       }
-
    }
 
    async poll(): Promise<void> {
@@ -303,6 +315,8 @@ class PoolHandler extends PollBase {
    agents?: GetAgentResponse[];
    pool?: GetPoolResponse;
 
+   static pools?: GetPoolResponse[];
+
 }
 
 const handler = new PoolHandler();
@@ -684,7 +698,7 @@ const AutoScalerPanel: React.FC = () => {
          color = pool.properties["Color"];
       }
 
-	  setState({ name: pool.name, color: color, autoscale: pool.enableAutoscaling ?? false, minAgents: pool.minAgents, reserveAgents: pool.numReserveAgents, strategy: pool.sizeStrategy ?? PoolSizeStrategy.LeaseUtilization, conformInterval: pool.conformInterval ?? 24 });
+      setState({ name: pool.name, color: color, autoscale: pool.enableAutoscaling ?? false, minAgents: pool.minAgents, reserveAgents: pool.numReserveAgents, strategy: pool.sizeStrategy ?? PoolSizeStrategy.LeaseUtilization, conformInterval: pool.conformInterval ?? 24 });
       return null;
    }
 
@@ -724,7 +738,7 @@ const AutoScalerPanel: React.FC = () => {
                   properties: { Color: state.color! }
                }
 
-			   await backend.updatePool(pool.id, update);
+               await backend.updatePool(pool.id, update);
 
                handler.pool = await backend.getPool(pool.id);
 
@@ -984,7 +998,7 @@ const PoolPanel: React.FC = () => {
 export const PoolView: React.FC = observer(() => {
 
    const windowSize = useWindowSize();
-   const { poolId } = useParams<{ poolId: string }>();
+   const [searchParams, setSearchParams] = useSearchParams();
 
    useEffect(() => {
 
@@ -999,12 +1013,56 @@ export const PoolView: React.FC = observer(() => {
    // subscribe
    if (handler.updated) { };
 
-   handler.set(poolId);
+   if (!PoolHandler.pools) {
+      handler.loadPools();
+      return null;
+   }
+
+   if (!PoolHandler.pools.length) {
+      return null;
+   }
+
+   const poolId = searchParams.get("pool") ?? "";
+
+   if (poolId) {
+      handler.set(poolId);
+   }
 
    const pool = handler.pool;
 
    const vw = Math.max(document.documentElement.clientWidth, window.innerWidth || 0);
 
+   const poolTags: ITag[] = PoolHandler.pools.map(p => {
+      return { key: p.id, name: p.name }
+   });
+
+   let defaultItems: ITag[] | undefined;
+   if (poolId) {
+      const d = poolTags.find(t => t.key === poolId);
+      if (d) {
+         defaultItems = [d];
+      }
+   }
+
+
+
+   const listContainsTagList = (tag: ITag, tagList?: ITag[]) => {
+      if (!tagList || !tagList.length || tagList.length === 0) {
+         return false;
+      }
+      return tagList.some(compareTag => compareTag.key === tag.key);
+   };
+
+   const filterSuggestedTags = (filterText: string, tagList?: ITag[]): ITag[] => {
+      return filterText
+         ? poolTags.filter(
+            tag => tag.name.toLowerCase().indexOf(filterText.toLowerCase()) === 0 && !listContainsTagList(tag, tagList),
+         )
+         : poolTags;
+   };
+
+   const getTextFromItem = (item: ITag) => item.name;
+
    return <Stack className={hordeClasses.horde}>
       <TopNav />
       <Breadcrumbs items={[{ text: "Agents", link: "/agents" }, { text: `Pool ${handler.pool?.name ?? "Loading"}` }]} />
@@ -1012,19 +1070,45 @@ export const PoolView: React.FC = observer(() => {
          <div key={`windowsize_poolview_${windowSize.width}_${windowSize.height}`} style={{ width: vw / 2 - 896, flexShrink: 0, backgroundColor: modeColors.background }} />
          <Stack tokens={{ childrenGap: 0 }} styles={{ root: { backgroundColor: modeColors.background, width: "100%" } }}>
             <Stack style={{ maxWidth: 1778, paddingTop: 6, marginLeft: 4, height: 'calc(100vh - 8px)' }}>
-               <Stack className={hordeClasses.raised}>
+               <Stack horizontal style={{ paddingTop: 8, paddingBottom: 16 }}>
+                  <Stack grow />
+                  <Stack style={{width: 320}}>
+                     <TagPicker inputProps={{ placeholder: "Select a pool" }}
+                        defaultSelectedItems={ defaultItems}
+                        onResolveSuggestions={filterSuggestedTags}
+                        getTextFromItem={getTextFromItem}
+                        onEmptyResolveSuggestions={(selected) => {
+                           return poolTags;
+                        }}
+
+                        onItemSelected={(item) => {
+
+                           if (!item?.key) {
+                              return null;
+                           }
+
+                           setSearchParams(`?pool=${item.key}`);
+
+                           return item;
+
+                        }}
+
+                        itemLimit={1} />
+                  </Stack>
+               </Stack>
+               {!!pool && <Stack className={hordeClasses.raised}>
                   <Stack style={{ position: "relative", width: "100%", height: 'calc(100vh - 228px)' }}>
-                     {!!pool && <ScrollablePane scrollbarVisibility={ScrollbarVisibility.always}>
+                     <ScrollablePane scrollbarVisibility={ScrollbarVisibility.always}>
                         <Stack tokens={{ childrenGap: 18 }}>
                            <PoolPanel />
                            <BatchPanel />
                            <StepPanel stepState={StepState.Active} />
                            <StepPanel stepState={StepState.Pending} />
                         </Stack>
-                     </ScrollablePane>}
-                     {!pool && <Spinner size={SpinnerSize.large} />}
+                     </ScrollablePane>
                   </Stack>
-               </Stack>
+               </Stack>}
+               {!pool && <Spinner size={SpinnerSize.large} />}
             </Stack>
          </Stack>
       </Stack>
diff --git a/Engine/Source/Programs/Horde/HordeDashboard/src/components/TopNav.tsx b/Engine/Source/Programs/Horde/HordeDashboard/src/components/TopNav.tsx
--- a/Engine/Source/Programs/Horde/HordeDashboard/src/components/TopNav.tsx
+++ b/Engine/Source/Programs/Horde/HordeDashboard/src/components/TopNav.tsx
@@ -362,6 +362,12 @@ export const TopNav: React.FC<{ suppressServer?: boolean }> = observer(({ suppre
          link: `/agents`
       });
 
+      resourceItems.push({
+         key: "admin_pools",
+         text: "Pools",
+         link: `/pools`
+      });
+
       resourceItems.push({
          key: "admin_perforce_servers",
          text: "Perforce Servers",
