Horde: Fix bad pool spinner
#rnx
#skipci

diff --git a/Engine/Source/Programs/Horde/HordeDashboard/src/components/PoolView.tsx b/Engine/Source/Programs/Horde/HordeDashboard/src/components/PoolView.tsx
--- a/Engine/Source/Programs/Horde/HordeDashboard/src/components/PoolView.tsx
+++ b/Engine/Source/Programs/Horde/HordeDashboard/src/components/PoolView.tsx
@@ -1108,7 +1108,7 @@ export const PoolView: React.FC = observer(() => {
                      </ScrollablePane>
                   </Stack>
                </Stack>}
-               {!pool && <Spinner size={SpinnerSize.large} />}
+               {!pool && !!poolId && <Spinner size={SpinnerSize.large} />}
             </Stack>
          </Stack>
       </Stack>
