Horde: Allow authenticating to write log files with a lease claim.

diff --git a/Engine/Source/Programs/Horde/Horde.Build.Tests/IssueServiceTests.cs b/Engine/Source/Programs/Horde/Horde.Build.Tests/IssueServiceTests.cs
--- a/Engine/Source/Programs/Horde/Horde.Build.Tests/IssueServiceTests.cs
+++ b/Engine/Source/Programs/Horde/Horde.Build.Tests/IssueServiceTests.cs
@@ -214,7 +214,7 @@ namespace Horde.Build.Tests
 				{
 					SubResourceId stepId = new SubResourceId((ushort)((groupIdx * 100) + nodeIdx));
 
-					ILogFile logFile = LogFileService.CreateLogFileAsync(jobId, null, LogType.Json, false).Result;
+					ILogFile logFile = LogFileService.CreateLogFileAsync(jobId, null, null, LogType.Json, false).Result;
 
 					Mock<IJobStep> step = new Mock<IJobStep>(MockBehavior.Strict);
 					step.SetupGet(x => x.Id).Returns(stepId);
diff --git a/Engine/Source/Programs/Horde/Horde.Build.Tests/LogFileServiceTests.cs b/Engine/Source/Programs/Horde/Horde.Build.Tests/LogFileServiceTests.cs
--- a/Engine/Source/Programs/Horde/Horde.Build.Tests/LogFileServiceTests.cs
+++ b/Engine/Source/Programs/Horde/Horde.Build.Tests/LogFileServiceTests.cs
@@ -76,7 +76,7 @@ namespace Horde.Build.Tests
         public async Task WriteLogLifecycleOldTest()
         {
 			JobId jobId = JobId.GenerateNewId();
-            ILogFile logFile = await _logFileService.CreateLogFileAsync(jobId, null, LogType.Text, useNewStorageBackend: false, logId: null, cancellationToken: CancellationToken.None);
+            ILogFile logFile = await _logFileService.CreateLogFileAsync(jobId, null, null, LogType.Text, useNewStorageBackend: false, logId: null, cancellationToken: CancellationToken.None);
 
             logFile = (await ((ILogFileService)_logFileService).WriteLogDataAsync(logFile, 0, 0, Encoding.ASCII.GetBytes("hello\n"), true))!;
 			logFile = (await ((ILogFileService)_logFileService).WriteLogDataAsync(logFile, 6, 1, Encoding.ASCII.GetBytes("foo\nbar\n"), true))!;
@@ -147,7 +147,7 @@ namespace Horde.Build.Tests
         public async Task WriteLogLifecycle(ILogFileService lfs, int maxChunkLength)
         {
 			JobId jobId = JobId.GenerateNewId();
-            ILogFile logFile = await lfs.CreateLogFileAsync(jobId, null, LogType.Text, useNewStorageBackend: false);
+            ILogFile logFile = await lfs.CreateLogFileAsync(jobId, null, null, LogType.Text, useNewStorageBackend: false);
 
             string str1 = "hello\n";
             string str2 = "foo\nbar\n";
@@ -254,7 +254,7 @@ namespace Horde.Build.Tests
 			// Will implicitly test GetLogFileAsync(), AddCachedLogFile()
 			JobId jobId = JobId.GenerateNewId();
             SessionId sessionId = SessionId.GenerateNewId();
-            ILogFile a = await _logFileService.CreateLogFileAsync(jobId, sessionId, LogType.Text, useNewStorageBackend: false, logId: null, CancellationToken.None);
+            ILogFile a = await _logFileService.CreateLogFileAsync(jobId, null, sessionId, LogType.Text, useNewStorageBackend: false, logId: null, CancellationToken.None);
             ILogFile b = (await _logFileService.GetCachedLogFileAsync(a.Id, CancellationToken.None))!;
             Assert.AreEqual(a.JobId, b.JobId);
             Assert.AreEqual(a.SessionId, b.SessionId);
@@ -263,7 +263,7 @@ namespace Horde.Build.Tests
             ILogFile? notFound = await _logFileService.GetCachedLogFileAsync(LogId.GenerateNewId(), CancellationToken.None);
             Assert.IsNull(notFound);
 
-            await _logFileService.CreateLogFileAsync(JobId.GenerateNewId(), SessionId.GenerateNewId(), LogType.Text, useNewStorageBackend: false, logId: null, cancellationToken: CancellationToken.None);
+            await _logFileService.CreateLogFileAsync(JobId.GenerateNewId(), null, SessionId.GenerateNewId(), LogType.Text, useNewStorageBackend: false, logId: null, cancellationToken: CancellationToken.None);
             Assert.AreEqual(2, (await _logFileService.GetLogFilesAsync()).Count);
         }
 
@@ -272,8 +272,8 @@ namespace Horde.Build.Tests
         {
 			JobId jobId = JobId.GenerateNewId();
             SessionId sessionId = SessionId.GenerateNewId();
-            ILogFile logFile = await _logFileService.CreateLogFileAsync(jobId, sessionId, LogType.Text, useNewStorageBackend: false, logId: null, cancellationToken: CancellationToken.None);
-            ILogFile logFileNoSession = await _logFileService.CreateLogFileAsync(jobId, null, LogType.Text, useNewStorageBackend: false, logId: null, cancellationToken: CancellationToken.None);
+            ILogFile logFile = await _logFileService.CreateLogFileAsync(jobId, null, sessionId, LogType.Text, useNewStorageBackend: false, logId: null, cancellationToken: CancellationToken.None);
+            ILogFile logFileNoSession = await _logFileService.CreateLogFileAsync(jobId, null, null, LogType.Text, useNewStorageBackend: false, logId: null, cancellationToken: CancellationToken.None);
 
 			ClaimsPrincipal hasClaim = new ClaimsPrincipal(new ClaimsIdentity(new List<Claim>
             {
@@ -293,7 +293,7 @@ namespace Horde.Build.Tests
 		public async Task ChunkSplitting()
 		{
 			JobId jobId = JobId.GenerateNewId();
-			ILogFile logFile = await _logFileService.CreateLogFileAsync(jobId, null, LogType.Text, useNewStorageBackend: false, logId: null, cancellationToken: CancellationToken.None);
+			ILogFile logFile = await _logFileService.CreateLogFileAsync(jobId, null, null, LogType.Text, useNewStorageBackend: false, logId: null, cancellationToken: CancellationToken.None);
 
 			long offset = 0;
 
diff --git a/Engine/Source/Programs/Horde/Horde.Build.Tests/LogIndexingTests.cs b/Engine/Source/Programs/Horde/Horde.Build.Tests/LogIndexingTests.cs
--- a/Engine/Source/Programs/Horde/Horde.Build.Tests/LogIndexingTests.cs
+++ b/Engine/Source/Programs/Horde/Horde.Build.Tests/LogIndexingTests.cs
@@ -64,7 +64,7 @@ namespace Horde.Build.Tests
 		public async Task IndexTests()
 		{
 			JobId jobId = JobId.GenerateNewId();
-			ILogFile logFile = await _logFileService.CreateLogFileAsync(jobId, null, LogType.Text, useNewStorageBackend: false);
+			ILogFile logFile = await _logFileService.CreateLogFileAsync(jobId, null, null, LogType.Text, useNewStorageBackend: false);
 
 			// Write the test data to the log file in blocks
 			int offset = 0;
@@ -135,7 +135,7 @@ namespace Horde.Build.Tests
 		public async Task PartialTokenTests()
 		{
 			JobId jobId = JobId.GenerateNewId();
-			ILogFile logFile = await _logFileService.CreateLogFileAsync(jobId, null, LogType.Text, useNewStorageBackend: false);
+			ILogFile logFile = await _logFileService.CreateLogFileAsync(jobId, null, null, LogType.Text, useNewStorageBackend: false);
 
 			string[] lines =
 			{
@@ -177,7 +177,7 @@ namespace Horde.Build.Tests
 		public async Task AppendIndexTests()
 		{
 			JobId jobId = JobId.GenerateNewId();
-			ILogFile logFile = await _logFileService.CreateLogFileAsync(jobId, null, LogType.Text, useNewStorageBackend: false);
+			ILogFile logFile = await _logFileService.CreateLogFileAsync(jobId, null, null, LogType.Text, useNewStorageBackend: false);
 
 			logFile = await WriteLogDataAsync(logFile, 0, 0, Encoding.UTF8.GetBytes("abc\n"), true);
 			logFile = await WriteLogDataAsync(logFile, 4, 1, Encoding.UTF8.GetBytes("def\n"), true);
diff --git a/Engine/Source/Programs/Horde/Horde.Build.Tests/TestDataTests.cs b/Engine/Source/Programs/Horde/Horde.Build.Tests/TestDataTests.cs
--- a/Engine/Source/Programs/Horde/Horde.Build.Tests/TestDataTests.cs
+++ b/Engine/Source/Programs/Horde/Horde.Build.Tests/TestDataTests.cs
@@ -118,7 +118,7 @@ namespace Horde.Build.Tests
 				{
 					SubResourceId stepId = new SubResourceId((ushort)((groupIdx * 100) + nodeIdx));
 
-					ILogFile logFile = LogFileService.CreateLogFileAsync(jobId, null, LogType.Json, useNewStorageBackend: false).Result;
+					ILogFile logFile = LogFileService.CreateLogFileAsync(jobId, null, null, LogType.Json, useNewStorageBackend: false).Result;
 
 					Mock<IJobStep> step = new Mock<IJobStep>(MockBehavior.Strict);
 					step.SetupGet(x => x.Id).Returns(stepId);
diff --git a/Engine/Source/Programs/Horde/Horde.Build/Jobs/JobRpcService.cs b/Engine/Source/Programs/Horde/Horde.Build/Jobs/JobRpcService.cs
--- a/Engine/Source/Programs/Horde/Horde.Build/Jobs/JobRpcService.cs
+++ b/Engine/Source/Programs/Horde/Horde.Build/Jobs/JobRpcService.cs
@@ -427,7 +427,7 @@ namespace Horde.Build.Jobs
 			}
 
 			// Create a log file if necessary
-			log.Value ??= await _logFileService.CreateLogFileAsync(job.Id, batch.SessionId, LogType.Json, job.JobOptions?.UseNewLogStorage ?? false);
+			log.Value ??= await _logFileService.CreateLogFileAsync(job.Id, batch.LeaseId, batch.SessionId, LogType.Json, job.JobOptions?.UseNewLogStorage ?? false);
 
 			// Get the node for this step
 			IGraph graph = await _jobService.GetGraphAsync(job);
diff --git a/Engine/Source/Programs/Horde/Horde.Build/Jobs/JobTaskSource.cs b/Engine/Source/Programs/Horde/Horde.Build/Jobs/JobTaskSource.cs
--- a/Engine/Source/Programs/Horde/Horde.Build/Jobs/JobTaskSource.cs
+++ b/Engine/Source/Programs/Horde/Horde.Build/Jobs/JobTaskSource.cs
@@ -773,7 +773,7 @@ namespace Horde.Build.Jobs
 					if (waiter.LeaseSource.TrySetResult(lease))
 					{
 						_logger.LogDebug("Assigned lease {LeaseId} to agent {AgentId}", leaseId, agent.Id);
-						await _logFileService.CreateLogFileAsync(job.Id, agent.SessionId, LogType.Json, job.JobOptions?.UseNewLogStorage ?? false, logId);
+						await _logFileService.CreateLogFileAsync(job.Id, leaseId, agent.SessionId, LogType.Json, job.JobOptions?.UseNewLogStorage ?? false, logId);
 						return lease;
 					}
 				}
diff --git a/Engine/Source/Programs/Horde/Horde.Build/Logs/ILogFile.cs b/Engine/Source/Programs/Horde/Horde.Build/Logs/ILogFile.cs
--- a/Engine/Source/Programs/Horde/Horde.Build/Logs/ILogFile.cs
+++ b/Engine/Source/Programs/Horde/Horde.Build/Logs/ILogFile.cs
@@ -4,6 +4,7 @@ using System;
 using System.Collections.Generic;
 using EpicGames.Core;
 using EpicGames.Horde.Storage;
+using Horde.Build.Agents.Leases;
 using Horde.Build.Agents.Sessions;
 using Horde.Build.Jobs;
 using Horde.Build.Storage;
@@ -52,6 +53,11 @@ namespace Horde.Build.Logs
 		/// </summary>
 		public JobId JobId { get; }
 
+		/// <summary>
+		/// The lease allowed to write to this log
+		/// </summary>
+		public LeaseId? LeaseId { get; }
+
 		/// <summary>
 		/// The session allowed to write to this log
 		/// </summary>
diff --git a/Engine/Source/Programs/Horde/Horde.Build/Logs/ILogFileCollection.cs b/Engine/Source/Programs/Horde/Horde.Build/Logs/ILogFileCollection.cs
--- a/Engine/Source/Programs/Horde/Horde.Build/Logs/ILogFileCollection.cs
+++ b/Engine/Source/Programs/Horde/Horde.Build/Logs/ILogFileCollection.cs
@@ -3,6 +3,7 @@
 using System.Collections.Generic;
 using System.Threading;
 using System.Threading.Tasks;
+using Horde.Build.Agents.Leases;
 using Horde.Build.Agents.Sessions;
 using Horde.Build.Jobs;
 using Horde.Build.Utilities;
@@ -52,13 +53,14 @@ namespace Horde.Build.Logs
 		/// Creates a new log
 		/// </summary>
 		/// <param name="jobId">Unique id of the job that owns this log file</param>
+		/// <param name="leaseId">Agent lease allowed to update the log</param>
 		/// <param name="sessionId">Agent session allowed to update the log</param>
 		/// <param name="type">Type of events to be stored in the log</param>
 		/// <param name="newStorageBackend">Whether to use the new storage backend</param>
 		/// <param name="logId">ID of the log file (optional)</param>
 		/// <param name="cancellationToken">Cancellation token for the call</param>
 		/// <returns>The new log file document</returns>
-		Task<ILogFile> CreateLogFileAsync(JobId jobId, SessionId? sessionId, LogType type, bool newStorageBackend, LogId? logId = null, CancellationToken cancellationToken = default);
+		Task<ILogFile> CreateLogFileAsync(JobId jobId, LeaseId? leaseId, SessionId? sessionId, LogType type, bool newStorageBackend, LogId? logId = null, CancellationToken cancellationToken = default);
 
 		/// <summary>
 		/// Updates the line count for a log file (v2 backend only)
diff --git a/Engine/Source/Programs/Horde/Horde.Build/Logs/LogFileCollection.cs b/Engine/Source/Programs/Horde/Horde.Build/Logs/LogFileCollection.cs
--- a/Engine/Source/Programs/Horde/Horde.Build/Logs/LogFileCollection.cs
+++ b/Engine/Source/Programs/Horde/Horde.Build/Logs/LogFileCollection.cs
@@ -6,6 +6,7 @@ using System.Net;
 using System.Threading;
 using System.Threading.Tasks;
 using EpicGames.Horde.Storage;
+using Horde.Build.Agents.Leases;
 using Horde.Build.Agents.Sessions;
 using Horde.Build.Jobs;
 using Horde.Build.Server;
@@ -57,7 +58,12 @@ namespace Horde.Build.Logs
 			[BsonRequired]
 			public JobId JobId { get; set; }
 
+			[BsonIgnoreIfNull]
+			public LeaseId? LeaseId { get; set; }
+
+			[BsonIgnoreIfNull]
 			public SessionId? SessionId { get; set; }
+
 			public LogType Type { get; set; }
 			public bool UseNewStorageBackend { get; set; }
 
@@ -85,10 +91,11 @@ namespace Horde.Build.Logs
 			{
 			}
 
-			public LogFileDocument(JobId jobId, SessionId? sessionId, LogType type, bool newStorageBackend, LogId? logId)
+			public LogFileDocument(JobId jobId, LeaseId? leaseId, SessionId? sessionId, LogType type, bool newStorageBackend, LogId? logId)
 			{
 				Id = logId ?? LogId.GenerateNewId();
 				JobId = jobId;
+				LeaseId = leaseId;
 				SessionId = sessionId;
 				Type = type;
 				UseNewStorageBackend = newStorageBackend;
@@ -125,9 +132,9 @@ namespace Horde.Build.Logs
 		}
 
 		/// <inheritdoc/>
-		public async Task<ILogFile> CreateLogFileAsync(JobId jobId, SessionId? sessionId, LogType type, bool newStorageBackend, LogId? logId, CancellationToken cancellationToken)
+		public async Task<ILogFile> CreateLogFileAsync(JobId jobId, LeaseId? leaseId, SessionId? sessionId, LogType type, bool newStorageBackend, LogId? logId, CancellationToken cancellationToken)
 		{
-			LogFileDocument newLogFile = new (jobId, sessionId, type, newStorageBackend, logId);
+			LogFileDocument newLogFile = new (jobId, leaseId, sessionId, type, newStorageBackend, logId);
 			await _logFiles.InsertOneAsync(newLogFile, null, cancellationToken);
 			return newLogFile;
 		}
diff --git a/Engine/Source/Programs/Horde/Horde.Build/Logs/LogFileService.cs b/Engine/Source/Programs/Horde/Horde.Build/Logs/LogFileService.cs
--- a/Engine/Source/Programs/Horde/Horde.Build/Logs/LogFileService.cs
+++ b/Engine/Source/Programs/Horde/Horde.Build/Logs/LogFileService.cs
@@ -16,6 +16,7 @@ using System.Threading.Tasks;
 using EpicGames.Core;
 using EpicGames.Horde.Logs;
 using EpicGames.Horde.Storage;
+using Horde.Build.Agents.Leases;
 using Horde.Build.Agents.Sessions;
 using Horde.Build.Jobs;
 using Horde.Build.Logs.Data;
@@ -58,13 +59,14 @@ namespace Horde.Build.Logs
 		/// Creates a new log
 		/// </summary>
 		/// <param name="jobId">Unique id of the job that owns this log file</param>
+		/// <param name="leaseId">Lease allowed to update the log</param>
 		/// <param name="sessionId">Agent session allowed to update the log</param>
 		/// <param name="type">Type of events to be stored in the log</param>
 		/// <param name="useNewStorageBackend">Whether to use the new storage backend for log data</param>
 		/// <param name="cancellationToken">Cancellation token for the call</param>
 		/// <param name="logId">ID of the log file (optional)</param>
 		/// <returns>The new log file document</returns>
-		Task<ILogFile> CreateLogFileAsync(JobId jobId, SessionId? sessionId, LogType type, bool useNewStorageBackend, LogId? logId = null, CancellationToken cancellationToken = default);
+		Task<ILogFile> CreateLogFileAsync(JobId jobId, LeaseId? leaseId, SessionId? sessionId, LogType type, bool useNewStorageBackend, LogId? logId = null, CancellationToken cancellationToken = default);
 
 		/// <summary>
 		/// Gets a logfile by ID
@@ -597,9 +599,9 @@ namespace Horde.Build.Logs
 		}
 
 		/// <inheritdoc/>
-		public Task<ILogFile> CreateLogFileAsync(JobId jobId, SessionId? sessionId, LogType type, bool useNewStorageBackend, LogId? logId, CancellationToken cancellationToken)
+		public Task<ILogFile> CreateLogFileAsync(JobId jobId, LeaseId? leaseId, SessionId? sessionId, LogType type, bool useNewStorageBackend, LogId? logId, CancellationToken cancellationToken)
 		{
-			return _logFiles.CreateLogFileAsync(jobId, sessionId, type, useNewStorageBackend, logId, cancellationToken);
+			return _logFiles.CreateLogFileAsync(jobId, leaseId, sessionId, type, useNewStorageBackend, logId, cancellationToken);
 		}
 
 		/// <inheritdoc/>
diff --git a/Engine/Source/Programs/Horde/Horde.Build/Logs/LogsController.cs b/Engine/Source/Programs/Horde/Horde.Build/Logs/LogsController.cs
--- a/Engine/Source/Programs/Horde/Horde.Build/Logs/LogsController.cs
+++ b/Engine/Source/Programs/Horde/Horde.Build/Logs/LogsController.cs
@@ -368,6 +368,10 @@ namespace Horde.Build.Logs
 		async Task<bool> AuthorizeAsync(ILogFile logFile, AclAction action, ClaimsPrincipal user)
 		{
 			GlobalConfig globalConfig = _globalConfig.Value;
+			if (logFile.LeaseId != null && user.HasLeaseClaim(logFile.LeaseId.Value))
+			{
+				return true;
+			}
 			if (logFile.SessionId != null && LogFileService.AuthorizeForSession(logFile, user))
 			{
 				return true;
diff --git a/Engine/Source/Programs/Horde/Horde.Build/Tasks/ConformTaskSource.cs b/Engine/Source/Programs/Horde/Horde.Build/Tasks/ConformTaskSource.cs
--- a/Engine/Source/Programs/Horde/Horde.Build/Tasks/ConformTaskSource.cs
+++ b/Engine/Source/Programs/Horde/Horde.Build/Tasks/ConformTaskSource.cs
@@ -176,7 +176,7 @@ namespace Horde.Build.Tasks
 					LeaseId leaseId = LeaseId.GenerateNewId();
 					if (await AllocateConformLeaseAsync(agent.Id, task.Workspaces, leaseId))
 					{
-						ILogFile log = await _logService.CreateLogFileAsync(JobId.Empty, agent.SessionId, LogType.Json, useNewStorageBackend: false, cancellationToken: cancellationToken);
+						ILogFile log = await _logService.CreateLogFileAsync(JobId.Empty, leaseId, agent.SessionId, LogType.Json, useNewStorageBackend: false, cancellationToken: cancellationToken);
 						task.LogId = log.Id.ToString();
 						task.RemoveUntrackedFiles = agent.RequestFullConform;
 
diff --git a/Engine/Source/Programs/Horde/Horde.Build/Tasks/RestartTaskSource.cs b/Engine/Source/Programs/Horde/Horde.Build/Tasks/RestartTaskSource.cs
--- a/Engine/Source/Programs/Horde/Horde.Build/Tasks/RestartTaskSource.cs
+++ b/Engine/Source/Programs/Horde/Horde.Build/Tasks/RestartTaskSource.cs
@@ -40,14 +40,15 @@ namespace Horde.Build.Tasks
 				return await DrainAsync(cancellationToken);
 			}
 
-			ILogFile log = await _logService.CreateLogFileAsync(JobId.Empty, agent.SessionId, LogType.Json, useNewStorageBackend: false, cancellationToken: cancellationToken);
+			LeaseId leaseId = LeaseId.GenerateNewId();
+			ILogFile log = await _logService.CreateLogFileAsync(JobId.Empty, leaseId, agent.SessionId, LogType.Json, useNewStorageBackend: false, cancellationToken: cancellationToken);
 
 			RestartTask task = new RestartTask();
 			task.LogId = log.Id.ToString();
 
 			byte[] payload = Any.Pack(task).ToByteArray();
 
-			return Lease(new AgentLease(LeaseId.GenerateNewId(), "Restart", null, null, log.Id, LeaseState.Pending, null, true, payload));
+			return Lease(new AgentLease(leaseId, "Restart", null, null, log.Id, LeaseState.Pending, null, true, payload));
 		}
 	}
 }
diff --git a/Engine/Source/Programs/Horde/Horde.Build/Tasks/ShutdownTaskSource.cs b/Engine/Source/Programs/Horde/Horde.Build/Tasks/ShutdownTaskSource.cs
--- a/Engine/Source/Programs/Horde/Horde.Build/Tasks/ShutdownTaskSource.cs
+++ b/Engine/Source/Programs/Horde/Horde.Build/Tasks/ShutdownTaskSource.cs
@@ -39,14 +39,15 @@ namespace Horde.Build.Tasks
 				return await DrainAsync(cancellationToken);
 			}
 
-			ILogFile log = await _logService.CreateLogFileAsync(JobId.Empty, agent.SessionId, LogType.Json, useNewStorageBackend: false, cancellationToken: cancellationToken);
+			LeaseId leaseId = LeaseId.GenerateNewId();
+			ILogFile log = await _logService.CreateLogFileAsync(JobId.Empty, leaseId, agent.SessionId, LogType.Json, useNewStorageBackend: false, cancellationToken: cancellationToken);
 
 			ShutdownTask task = new ShutdownTask();
 			task.LogId = log.Id.ToString();
 
 			byte[] payload = Any.Pack(task).ToByteArray();
 
-			return Lease(new AgentLease(LeaseId.GenerateNewId(), "Shutdown", null, null, log.Id, LeaseState.Pending, null, true, payload));
+			return Lease(new AgentLease(leaseId, "Shutdown", null, null, log.Id, LeaseState.Pending, null, true, payload));
 		}
 	}
 }
diff --git a/Engine/Source/Programs/Horde/Horde.Build/Tasks/UpgradeTaskSource.cs b/Engine/Source/Programs/Horde/Horde.Build/Tasks/UpgradeTaskSource.cs
--- a/Engine/Source/Programs/Horde/Horde.Build/Tasks/UpgradeTaskSource.cs
+++ b/Engine/Source/Programs/Horde/Horde.Build/Tasks/UpgradeTaskSource.cs
@@ -58,7 +58,8 @@ namespace Horde.Build.Tasks
 				return await DrainAsync(cancellationToken);
 			}
 
-			ILogFile logFile = await _logService.CreateLogFileAsync(JobId.Empty, agent.SessionId, LogType.Json, useNewStorageBackend: false, cancellationToken: cancellationToken);
+			LeaseId leaseId = LeaseId.GenerateNewId();
+			ILogFile logFile = await _logService.CreateLogFileAsync(JobId.Empty, leaseId, agent.SessionId, LogType.Json, useNewStorageBackend: false, cancellationToken: cancellationToken);
 
 			UpgradeTask task = new UpgradeTask();
 			task.SoftwareId = requiredVersion.ToString();
@@ -77,7 +78,7 @@ namespace Horde.Build.Tasks
 				payload = Any.Pack(task).ToByteArray();
 			}
 
-			return Lease(new AgentLease(LeaseId.GenerateNewId(), $"Upgrade to {requiredVersion}", null, null, logFile.Id, LeaseState.Pending, null, true, payload));
+			return Lease(new AgentLease(leaseId, $"Upgrade to {requiredVersion}", null, null, logFile.Id, LeaseState.Pending, null, true, payload));
 		}
 
 		/// <summary>
