Horde: Fix bad filter on pool view
#skipci
#rnx

diff --git a/Engine/Source/Programs/Horde/HordeDashboard/src/components/PoolView.tsx b/Engine/Source/Programs/Horde/HordeDashboard/src/components/PoolView.tsx
--- a/Engine/Source/Programs/Horde/HordeDashboard/src/components/PoolView.tsx
+++ b/Engine/Source/Programs/Horde/HordeDashboard/src/components/PoolView.tsx
@@ -1056,7 +1056,7 @@ export const PoolView: React.FC = observer(() => {
    const filterSuggestedTags = (filterText: string, tagList?: ITag[]): ITag[] => {
       return filterText
          ? poolTags.filter(
-            tag => tag.name.toLowerCase().indexOf(filterText.toLowerCase()) === 0 && !listContainsTagList(tag, tagList),
+            tag => tag.name.toLowerCase().indexOf(filterText.toLowerCase()) !== -1 && !listContainsTagList(tag, tagList),
          )
          : poolTags;
    };
