Horde: Allow clients to auth for legacy log writes using a lease claim.

diff --git a/Engine/Source/Programs/Horde/Horde.Build/Logs/LogFileService.cs b/Engine/Source/Programs/Horde/Horde.Build/Logs/LogFileService.cs
--- a/Engine/Source/Programs/Horde/Horde.Build/Logs/LogFileService.cs
+++ b/Engine/Source/Programs/Horde/Horde.Build/Logs/LogFileService.cs
@@ -16,6 +16,7 @@ using System.Threading.Tasks;
 using EpicGames.Core;
 using EpicGames.Horde.Logs;
 using EpicGames.Horde.Storage;
+using Horde.Build.Acls;
 using Horde.Build.Agents.Leases;
 using Horde.Build.Agents.Sessions;
 using Horde.Build.Jobs;
@@ -1846,14 +1847,15 @@ namespace Horde.Build.Logs
 		/// <returns>True if the action is authorized</returns>
 		public static bool AuthorizeForSession(ILogFile logFile, ClaimsPrincipal user)
 		{
-			if(logFile.SessionId != null)
+			if (logFile.SessionId != null && user.HasSessionClaim(logFile.SessionId.Value))
 			{
-				return user.HasClaim(HordeClaimTypes.AgentSessionId, logFile.SessionId.Value.ToString());
+				return true;
 			}
-			else
+			if (logFile.LeaseId != null && user.HasLeaseClaim(logFile.LeaseId.Value))
 			{
-				return false;
+				return true;
 			}
+			return false;
 		}
 
 		/// <inheritdoc/>
