Horde: Set correct path to sync dir returned by ManagedWorkspaceMaterializer

diff --git a/Engine/Source/Programs/Horde/Horde.Agent/Execution/ManagedWorkspaceMaterializer.cs b/Engine/Source/Programs/Horde/Horde.Agent/Execution/ManagedWorkspaceMaterializer.cs
--- a/Engine/Source/Programs/Horde/Horde.Agent/Execution/ManagedWorkspaceMaterializer.cs
+++ b/Engine/Source/Programs/Horde/Horde.Agent/Execution/ManagedWorkspaceMaterializer.cs
@@ -52,7 +52,7 @@ public class ManagedWorkspaceMaterializer : IWorkspaceMaterializer
 		
 		bool useHaveTable = WorkspaceInfo.ShouldUseHaveTable(_agentWorkspace.Method);
 		_workspace = await WorkspaceInfo.SetupWorkspaceAsync(_agentWorkspace, _workingDir, useHaveTable, _logger, cancellationToken);
-		return new WorkspaceMaterializerSettings(_workingDir, _agentWorkspace.Identifier, _agentWorkspace.Stream);
+		return await GetSettingsAsync(cancellationToken);
 	}
 
 	/// <inheritdoc/>
@@ -69,7 +69,9 @@ public class ManagedWorkspaceMaterializer : IWorkspaceMaterializer
 	/// <inheritdoc/>
 	public Task<WorkspaceMaterializerSettings> GetSettingsAsync(CancellationToken cancellationToken)
 	{
-		return Task.FromResult(new WorkspaceMaterializerSettings(_workingDir, _agentWorkspace.Identifier, _agentWorkspace.Stream));
+		// ManagedWorkspace store synced files in a sub-directory from the top working dir.
+		DirectoryReference syncDir = DirectoryReference.Combine(_workingDir, _agentWorkspace.Identifier, "Sync");
+		return Task.FromResult(new WorkspaceMaterializerSettings(syncDir, _agentWorkspace.Identifier, _agentWorkspace.Stream));
 	}
 
 	/// <inheritdoc/>
