Horde: Better pool breadcrumb when none selected
#rnx
preflight none
#skipci

diff --git a/Engine/Source/Programs/Horde/HordeDashboard/src/components/PoolView.tsx b/Engine/Source/Programs/Horde/HordeDashboard/src/components/PoolView.tsx
--- a/Engine/Source/Programs/Horde/HordeDashboard/src/components/PoolView.tsx
+++ b/Engine/Source/Programs/Horde/HordeDashboard/src/components/PoolView.tsx
@@ -2,10 +2,9 @@
 
 import { Checkbox, ComboBox, DefaultButton, DetailsList, DetailsListLayoutMode, Dialog, DialogFooter, DialogType, IColumn, IComboBoxOption, ITag, MessageBar, MessageBarType, Position, PrimaryButton, ScrollablePane, ScrollbarVisibility, SelectionMode, Slider, SpinButton, Spinner, SpinnerSize, Stack, TagPicker, Text, TextField } from "@fluentui/react";
 import { observer } from "mobx-react-lite";
-import React, { useEffect, useId, useState } from "react";
-import { Link, useParams, useSearchParams } from "react-router-dom";
+import React, { useEffect, useState } from "react";
+import { Link, useSearchParams } from "react-router-dom";
 import backend from "../backend";
-import { agentStore } from "../backend/AgentStore";
 import { GetAgentResponse, GetBatchResponse, GetJobTimingResponse, GetPoolResponse, GetStepResponse, JobData, JobQuery, JobState, JobStepBatchState, JobStepState, PoolSizeStrategy, StepData, UpdatePoolRequest } from "../backend/Api";
 import dashboard from "../backend/Dashboard";
 import { PollBase } from "../backend/PollBase";
@@ -1063,9 +1062,14 @@ export const PoolView: React.FC = observer(() => {
 
    const getTextFromItem = (item: ITag) => item.name;
 
+   let breadText = `Pool ${handler.pool?.name ?? "Loading"}`;
+   if (!poolId) {
+      breadText = "Select a Pool"
+   }
+
    return <Stack className={hordeClasses.horde}>
       <TopNav />
-      <Breadcrumbs items={[{ text: "Agents", link: "/agents" }, { text: `Pool ${handler.pool?.name ?? "Loading"}` }]} />
+      <Breadcrumbs items={[{ text: "Agents", link: "/agents" }, { text: breadText }]} />
       <Stack horizontal>
          <div key={`windowsize_poolview_${windowSize.width}_${windowSize.height}`} style={{ width: vw / 2 - 896, flexShrink: 0, backgroundColor: modeColors.background }} />
          <Stack tokens={{ childrenGap: 0 }} styles={{ root: { backgroundColor: modeColors.background, width: "100%" } }}>
