Horde: Improve step timeline end caps
#rnx
#skipci

diff --git a/Engine/Source/Programs/Horde/HordeDashboard/src/components/jobDetailsV2/JobDetailTimeline.tsx b/Engine/Source/Programs/Horde/HordeDashboard/src/components/jobDetailsV2/JobDetailTimeline.tsx
--- a/Engine/Source/Programs/Horde/HordeDashboard/src/components/jobDetailsV2/JobDetailTimeline.tsx
+++ b/Engine/Source/Programs/Horde/HordeDashboard/src/components/jobDetailsV2/JobDetailTimeline.tsx
@@ -313,7 +313,7 @@ class TimelineDataView extends JobDataView {
          this.generateSpans();
          this.dirty = false;
       }
-      
+
       if (!this.allSpans) {
          return [];
       }
@@ -365,7 +365,7 @@ class TimelineDataView extends JobDataView {
    filterCost: number = 0;
    filterStepId?: string;
 
-   dirty = false; 
+   dirty = false;
 
 }
 
@@ -446,14 +446,40 @@ class TimelineRenderer {
          .selectAll()
          .data(d3.group(I, i => Y[i]))
          .join("g")
-         .attr("transform", ([y]) => `translate(0,${(yScale(y) as any) + 16})`);
 
-      g.selectAll("line")
+
+      g.append("g").selectAll("line")
          .data(([, I]: any) => I)
          .join("line")
+         .attr("class", "separator")
+         .attr("x1", i => xScale(X[i as any].utcStart.getTime() / 1000))
+         .attr("x2", i => xScale(X[i as any].utcStart.getTime() / 1000))
+         .attr("y1", i => (yScale(X[i as any].lane) as number - 5 + 16))
+         .attr("y2", i => (yScale(X[i as any].lane) as number + 5 + 16))
+         .attr("stroke-linecap", 0)
+         .attr("stroke-width", i => X[i as any].filtered ? 0 : 3)
+         .attr("stroke", i => {
+            const span = spans[i as any];
+            if (span.type === 0) {
+               return dashboard.darktheme ? "#5B6367" : "#D3D2D1";
+            }
+            if (span.type === 1) {
+               return dashboard.darktheme ? "#506F7C" : "#BAD1DB";
+            }
+
+            return colors[span.step?.outcome ?? "Unspecified"];
+         })
+
+
+      g.append("g").selectAll("line")
+         .data(([, I]: any) => I)
+         .join("line")
+         .attr("class", "spanline")
          .attr("x1", i => xScale(X[i as any].utcStart.getTime() / 1000))
          .attr("x2", i => xScale(X[i as any].utcFinish.getTime() / 1000))
          .attr("stroke-width", i => X[i as any].filtered ? 0 : 5)
+         .attr("y1", i => (yScale(X[i as any].lane) as number + 16))
+         .attr("y2", i => (yScale(X[i as any].lane) as number + 16))
          .attr("stroke-linecap", 0)
          .attr("stroke", i => {
             const span = spans[i as any];
@@ -461,20 +487,12 @@ class TimelineRenderer {
                return dashboard.darktheme ? "#5B6367" : "#D3D2D1";
             }
             if (span.type === 1) {
-               return dashboard.darktheme ? "#506F7C" : "#BAD1DB";               
+               return dashboard.darktheme ? "#506F7C" : "#BAD1DB";
             }
 
             return colors[span.step?.outcome ?? "Unspecified"];
-
          })
 
-      const radius = 2.5
-      g.selectAll("circle")
-         .data(([, I]: any) => I)
-         .join("circle")
-         .attr("cx", i => xScale(X[i as any].utcStart.getTime() / 1000))
-         .attr("fill", i => "#2D3F5F")
-         .attr("r", i => X[i as any].filtered ? 0 : radius);
 
 
       const xAxis = (g: SelectionType) => {
@@ -499,9 +517,10 @@ class TimelineRenderer {
 
       function zoomed(event: any) {
          xScale.range([margin.left, width - margin.right].map(d => event.transform.applyX(d)));
-         svg!.selectAll("circle")
-            .attr("cx", i => xScale(X[i as any].utcStart.getTime() / 1000));
-         svg!.selectAll("line")
+         svg!.selectAll(".separator")
+            .attr("x1", i => { console.log("hi"); if (i as number >= X.length) return 0; return xScale(X[i as any].utcStart.getTime() / 1000) })
+            .attr("x2", i => { if (i as number >= X.length) return 0; return xScale(X[i as any].utcStart.getTime() / 1000) })
+         svg!.selectAll(".spanline")
             .attr("x1", i => { if (i as number >= X.length) return 0; return xScale(X[i as any].utcStart.getTime() / 1000) })
             .attr("x2", i => { if (i as number >= X.length) return 0; return xScale(X[i as any].utcFinish.getTime() / 1000) })
 
@@ -633,10 +652,10 @@ class TimelineRenderer {
                topY = 0;
             }
 
-            let tipX = mouseX ;
+            let tipX = mouseX;
             let translateX = "0%";
 
-            if (tipX > 1500) {               
+            if (tipX > 1500) {
                translateX = "-180%";
             }
 
@@ -711,11 +730,11 @@ const TimelineGraph: React.FC<{ dataView: TimelineDataView, filterTime: number,
    // todo: we can probably cache
    dataView.filterTime = filterTime;
    dataView.filterCost = filterCost;
-   if (dataView.filterStepId !==  stepId) {
+   if (dataView.filterStepId !== stepId) {
       dataView.filterStepId = stepId;
       dataView.dirty = true;
    }
-   
+
 
    if (!state.graph) {
       setState({ ...state, graph: new TimelineRenderer(dataView) })
