Horde: Log changes to AWS EC2 instance type to agent's audit log

diff --git a/Engine/Source/Programs/Horde/Horde.Build.Tests/AgentServiceTest.cs b/Engine/Source/Programs/Horde/Horde.Build.Tests/AgentServiceTest.cs
--- a/Engine/Source/Programs/Horde/Horde.Build.Tests/AgentServiceTest.cs
+++ b/Engine/Source/Programs/Horde/Horde.Build.Tests/AgentServiceTest.cs
@@ -7,11 +7,13 @@ using System.Security.Claims;
 using System.Threading.Tasks;
 using EpicGames.Core;
 using Horde.Build.Agents;
+using Horde.Build.Auditing;
 using Horde.Build.Jobs;
 using Horde.Build.Server;
 using Horde.Build.Utilities;
 using HordeCommon;
 using Microsoft.AspNetCore.Mvc;
+using Microsoft.Extensions.DependencyInjection;
 using Microsoft.VisualStudio.TestTools.UnitTesting;
 
 namespace Horde.Build.Tests;
@@ -55,6 +57,34 @@ public class AgentServiceTest : TestSetup
 		Assert.IsFalse(AgentService.AuthorizeSession(agent, GetUser(agent)));
 	}
 	
+	[TestMethod]
+	public async Task AuditLogAwsInstanceTypeChanges()
+	{
+		Fixture fixture = await CreateFixtureAsync();
+		IAuditLogChannel<AgentId> agentLogger = AgentCollection.GetLogger(fixture.Agent1.Id);
+
+		async Task<bool> AuditLogContains(string text)
+		{
+			await FlushAuditLogsAsync();
+			return await agentLogger.FindAsync().AnyAsync(x => x.Data.Contains(text, StringComparison.Ordinal));
+		}
+		
+		List<string> props = new () { $"{KnownPropertyNames.AwsInstanceType}=m5.large" };
+		IAgent agent = await AgentService.CreateSessionAsync(fixture.Agent1, AgentStatus.Ok, props, new Dictionary<string, int>(), "test");
+		Assert.IsFalse(await AuditLogContains("AWS EC2 instance type changed"));
+		
+		props = new () { $"{KnownPropertyNames.AwsInstanceType}=c6.xlarge" };
+		agent = await AgentService.CreateSessionAsync(agent, AgentStatus.Ok, props, new Dictionary<string, int>(), "test");
+		Assert.IsTrue(await AuditLogContains("AWS EC2 instance type changed"));
+	}
+
+	private Task FlushAuditLogsAsync()
+	{
+		IAuditLog<AgentId> auditLog = ServiceProvider.GetRequiredService<IAuditLog<AgentId>>();
+		AuditLog<AgentId> log = (AuditLog<AgentId>)auditLog;
+		return log.FlushMessagesInternalAsync();
+	}
+	
 	[TestMethod]
 	public async Task GetAgentRateTest()
 	{
diff --git a/Engine/Source/Programs/Horde/Horde.Build/Agents/AgentService.cs b/Engine/Source/Programs/Horde/Horde.Build/Agents/AgentService.cs
--- a/Engine/Source/Programs/Horde/Horde.Build/Agents/AgentService.cs
+++ b/Engine/Source/Programs/Horde/Horde.Build/Agents/AgentService.cs
@@ -342,6 +342,7 @@ namespace Horde.Build.Agents
 					newAgent = await Agents.TryStartSessionAsync(agent, newSession.Id, sessionExpiresAt, status, properties, resources, pools, dynamicPools, version);
 					if(newAgent != null)
 					{
+						LogPropertyChanges(agentLogger, agent.Properties, newAgent.Properties);
 						agent = newAgent;
 						agentLogger.LogInformation("Session {SessionId} started", newSession.Id);
 						break;
@@ -379,6 +380,23 @@ namespace Horde.Build.Agents
 			return true;
 		}
 
+		/// <summary>
+		/// Compare properties and write changes to audit log
+		/// </summary>
+		private static void LogPropertyChanges(IAuditLogChannel<AgentId> agentLogger, IReadOnlyList<string> before, IReadOnlyList<string> after)
+		{
+			const string AwsInstanceTypeKey = KnownPropertyNames.AwsInstanceType + "=";
+			string beforeProp = before.FirstOrDefault(x => x.StartsWith(AwsInstanceTypeKey, StringComparison.Ordinal), String.Empty);
+			string afterProp = after.FirstOrDefault(x => x.StartsWith(AwsInstanceTypeKey, StringComparison.Ordinal), String.Empty);
+
+			if (!String.IsNullOrEmpty(beforeProp) && !String.IsNullOrEmpty(afterProp) && beforeProp != afterProp)
+			{
+				string oldInstanceType = beforeProp.Replace(AwsInstanceTypeKey, "", StringComparison.Ordinal);
+				string newInstanceType = afterProp.Replace(AwsInstanceTypeKey, "", StringComparison.Ordinal);
+				agentLogger.LogInformation("AWS EC2 instance type changed from {OldInstanceType} to {NewInstanceType}", oldInstanceType, newInstanceType);
+			}
+		}
+
 		/// <summary>
 		/// Waits for a lease to be assigned to an agent
 		/// </summary>
diff --git a/Engine/Source/Programs/Horde/Horde.Build/Agents/IAgent.cs b/Engine/Source/Programs/Horde/Horde.Build/Agents/IAgent.cs
--- a/Engine/Source/Programs/Horde/Horde.Build/Agents/IAgent.cs
+++ b/Engine/Source/Programs/Horde/Horde.Build/Agents/IAgent.cs
@@ -352,6 +352,11 @@ namespace Horde.Build.Agents
 		/// Amount of RAM, in GB
 		/// </summary>
 		public const string Ram = "RAM";
+		
+		/// <summary>
+		/// AWS: Instance type
+		/// </summary>
+		public const string AwsInstanceType = "aws-instance-type";
 	}
 
 	/// <summary>
diff --git a/Engine/Source/Programs/Horde/Horde.Build/Auditing/AuditLog.cs b/Engine/Source/Programs/Horde/Horde.Build/Auditing/AuditLog.cs
--- a/Engine/Source/Programs/Horde/Horde.Build/Auditing/AuditLog.cs
+++ b/Engine/Source/Programs/Horde/Horde.Build/Auditing/AuditLog.cs
@@ -189,19 +189,28 @@ namespace Horde.Build.Auditing
 			DisposeAsync().AsTask().Wait();
 		}
 
+		/// <summary>
+		/// Flush any pending messages to database
+		/// Exposed as internal for use in tests
+		/// </summary>
+		internal async Task FlushMessagesInternalAsync()
+		{
+			List<AuditLogMessage> newMessages = new ();
+			while (_messageChannel.Reader.TryRead(out AuditLogMessage? newMessage))
+			{
+				newMessages.Add(newMessage);
+			}
+			if (newMessages.Count > 0)
+			{
+				await _messages.InsertManyAsync(newMessages);
+			}
+		}
+		
 		async Task WriteMessagesAsync()
 		{
 			while (await _messageChannel.Reader.WaitToReadAsync())
 			{
-				List<AuditLogMessage> newMessages = new List<AuditLogMessage>();
-				while (_messageChannel.Reader.TryRead(out AuditLogMessage? newMessage))
-				{
-					newMessages.Add(newMessage);
-				}
-				if (newMessages.Count > 0)
-				{
-					await _messages.InsertManyAsync(newMessages);
-				}
+				await FlushMessagesInternalAsync();
 			}
 		}
 
