Horde: Add batch ready time to response
#rnx

diff --git a/Engine/Source/Programs/Horde/Horde.Build/Jobs/JobMessages.cs b/Engine/Source/Programs/Horde/Horde.Build/Jobs/JobMessages.cs
--- a/Engine/Source/Programs/Horde/Horde.Build/Jobs/JobMessages.cs
+++ b/Engine/Source/Programs/Horde/Horde.Build/Jobs/JobMessages.cs
@@ -800,6 +800,11 @@ namespace Horde.Build.Jobs
 		/// </summary>
 		public DateTimeOffset? FinishTime { get; set; }
 
+		/// <summary>
+		/// Time at which the group became ready (UTC).
+		/// </summary>
+		public DateTimeOffset? ReadyTime { get; set; }
+
 		/// <summary>
 		/// Converts this batch into a public response object
 		/// </summary>
@@ -822,6 +827,7 @@ namespace Horde.Build.Jobs
 			WeightedPriority = batch.SchedulePriority;
 			StartTime = batch.StartTimeUtc;
 			FinishTime = batch.FinishTimeUtc;
+			ReadyTime = batch.ReadyTimeUtc;
 		}
 	}
 
diff --git a/Engine/Source/Programs/Horde/HordeDashboard/src/backend/Api.ts b/Engine/Source/Programs/Horde/HordeDashboard/src/backend/Api.ts
--- a/Engine/Source/Programs/Horde/HordeDashboard/src/backend/Api.ts
+++ b/Engine/Source/Programs/Horde/HordeDashboard/src/backend/Api.ts
@@ -1269,6 +1269,9 @@ export type GetBatchResponse = {
 	/**The priority of this batch */
 	weightedPriority: number;
 
+	/**Time at which the group became ready (UTC) */
+	readyTime?: Date | string;
+
 	/**Time at which the group started (UTC). */
 	startTime?: Date | string;
 
