Horde: Fix commit tags not working correctly for starting schedules.

diff --git a/Engine/Source/Programs/Horde/Horde.Build/Jobs/Schedules/ScheduleService.cs b/Engine/Source/Programs/Horde/Horde.Build/Jobs/Schedules/ScheduleService.cs
--- a/Engine/Source/Programs/Horde/Horde.Build/Jobs/Schedules/ScheduleService.cs
+++ b/Engine/Source/Programs/Horde/Horde.Build/Jobs/Schedules/ScheduleService.cs
@@ -351,7 +351,7 @@ namespace Horde.Build.Jobs.Schedules
 
 			// Cache the Perforce history as we're iterating through changes to improve query performance
 			ICommitCollection commits = _commitService.GetCollection(stream.Config);
-			IAsyncEnumerable<ICommit> commitEnumerable = commits.FindAsync(minChangeNumber, null, null, schedule.Config.Commits, cancellationToken);
+			IAsyncEnumerable<ICommit> commitEnumerable = commits.FindAsync(minChangeNumber, null, null, schedule.Config.GetAllCommitTags(), cancellationToken);
 			await using IAsyncEnumerator<ICommit> commitEnumerator = commitEnumerable.GetAsyncEnumerator(cancellationToken);
 
 			// Start as many jobs as possible
@@ -395,7 +395,7 @@ namespace Horde.Build.Jobs.Schedules
 				}
 
 				// Adjust the changelist for the desired filter
-				if (commit == null || await ShouldBuildChangeAsync(commit, schedule.Config.Commits, fileFilter, cancellationToken))
+				if (commit == null || await ShouldBuildChangeAsync(commit, schedule.Config.GetAllCommitTags(), fileFilter, cancellationToken))
 				{
 					int codeChange = change;
 
diff --git a/Engine/Source/Programs/Horde/Horde.Build/Streams/StreamConfig.cs b/Engine/Source/Programs/Horde/Horde.Build/Streams/StreamConfig.cs
--- a/Engine/Source/Programs/Horde/Horde.Build/Streams/StreamConfig.cs
+++ b/Engine/Source/Programs/Horde/Horde.Build/Streams/StreamConfig.cs
@@ -917,35 +917,31 @@ namespace Horde.Build.Streams
 		/// <summary>
 		/// The types of changes to run for
 		/// </summary>
-		[Obsolete("Use Tags instead")]
-		public List<ChangeContentFlags> Filter
+		[Obsolete("Use Commits instead")]
+		public List<ChangeContentFlags>? Filter { get; set; }
+
+		/// <summary>
+		/// Merges the list of commit tags with the legacy filter field
+		/// </summary>
+#pragma warning disable CS0618 // Type or member is obsolete
+		internal List<CommitTag> GetAllCommitTags()
 		{
-			get
-			{
-				List<ChangeContentFlags> flags = new List<ChangeContentFlags>();
-				if (Commits.Contains(CommitTag.Code))
-				{
-					flags.Add(ChangeContentFlags.ContainsCode);
-				}
-				if (Commits.Contains(CommitTag.Content))
-				{
-					flags.Add(ChangeContentFlags.ContainsContent);
-				}
-				return flags;
-			}
-			set
+			List<CommitTag> commits = Commits;
+			if (Filter != null && Filter.Count > 0)
 			{
-				Commits.Clear();
-				if (value.Contains(ChangeContentFlags.ContainsCode))
+				commits = new List<CommitTag>(commits);
+				if (Filter.Contains(ChangeContentFlags.ContainsCode))
 				{
-					Commits.Add(CommitTag.Code);
+					commits.Add(CommitTag.Code);
 				}
-				if (value.Contains(ChangeContentFlags.ContainsContent))
+				if (Filter.Contains(ChangeContentFlags.ContainsContent))
 				{
-					Commits.Add(CommitTag.Content);
+					commits.Add(CommitTag.Content);
 				}
 			}
+			return commits;
 		}
+#pragma warning restore CS0618 // Type or member is obsolete
 
 		/// <summary>
 		/// Files that should cause the job to trigger
