Horde: Skip use of Perforce (-NoP4) when invoking UAT through WorkspaceExecutor

It has no direct knowledge and should not set P4 specific env vars like PerforceExecutor does.

diff --git a/Engine/Source/Programs/Horde/Horde.Agent.Tests/Execution/WorkspaceExecutorTest.cs b/Engine/Source/Programs/Horde/Horde.Agent.Tests/Execution/WorkspaceExecutorTest.cs
--- a/Engine/Source/Programs/Horde/Horde.Agent.Tests/Execution/WorkspaceExecutorTest.cs
+++ b/Engine/Source/Programs/Horde/Horde.Agent.Tests/Execution/WorkspaceExecutorTest.cs
@@ -3,13 +3,11 @@
 using System;
 using System.Collections.Generic;
 using System.IO;
-using System.Net.Http.Headers;
 using System.Threading;
 using System.Threading.Tasks;
 using EpicGames.Core;
 using Horde.Agent.Execution;
 using Horde.Agent.Services;
-using Horde.Common;
 using HordeCommon.Rpc;
 using HordeCommon.Rpc.Tasks;
 using Microsoft.Extensions.Logging;
@@ -53,7 +51,7 @@ public sealed class WorkspaceExecutorTest : IDisposable
 
 		BeginBatchResponse batch = new BeginBatchResponse { Change = 1 };
 		JobExecutorOptions executorOptions = new JobExecutorOptions(_session, null!, JobId, "batch1", batch, default, "", null!, new JobOptions());
-		_executor = new (executorOptions, null, _workspace, NullLogger.Instance);
+		_executor = new (executorOptions, _workspace, null, NullLogger.Instance);
 	}
 
 	public void Dispose()
@@ -76,7 +74,7 @@ public sealed class WorkspaceExecutorTest : IDisposable
 	{
 		BeginBatchResponse batch = new BeginBatchResponse { Change = 1 };
 		JobExecutorOptions executorOptions = new JobExecutorOptions(_session, null!, JobId, "batch1", batch, default, "", null!, new JobOptions());
-		WorkspaceExecutor executor = new (executorOptions, _autoSdkWorkspace, _workspace, NullLogger.Instance);
+		WorkspaceExecutor executor = new (executorOptions, _workspace, _autoSdkWorkspace, NullLogger.Instance);
 
 		await executor.InitializeAsync(_logger, CancellationToken.None);
 		AssertWorkspaceFile(_autoSdkWorkspace, "HostWin64/Android/base.h", "base");
@@ -90,7 +88,7 @@ public sealed class WorkspaceExecutorTest : IDisposable
 	{
 		BeginBatchResponse batch = new BeginBatchResponse { Change = 1 };
 		JobExecutorOptions executorOptions = new JobExecutorOptions(_session, null!, JobId, "batch1", batch, default, "", null!, new JobOptions());
-		WorkspaceExecutor executor = new (executorOptions, _autoSdkWorkspace, _workspace, NullLogger.Instance);
+		WorkspaceExecutor executor = new (executorOptions, _workspace, _autoSdkWorkspace, NullLogger.Instance);
 		await executor.InitializeAsync(_logger, CancellationToken.None);
 
 		WorkspaceMaterializerSettings settings = await _workspace.GetSettingsAsync(CancellationToken.None);
@@ -116,7 +114,7 @@ public sealed class WorkspaceExecutorTest : IDisposable
 
 		BeginBatchResponse batch = new BeginBatchResponse { Change = 1, PreflightChange = 1000 };
 		JobExecutorOptions executorOptions = new JobExecutorOptions(_session, null!, "jobPreflight", "batch1", batch, default, "", null!, new JobOptions());
-		WorkspaceExecutor executor = new (executorOptions, null, _workspace, NullLogger.Instance);
+		WorkspaceExecutor executor = new (executorOptions, _workspace, null, NullLogger.Instance);
 		
 		await executor.InitializeAsync(_logger, CancellationToken.None);
 		AssertWorkspaceFile(_workspace, "main.cpp", "main");
@@ -131,7 +129,7 @@ public sealed class WorkspaceExecutorTest : IDisposable
 		_server.AddJob("jobNoChange", StreamId, 0, 0);
 		BeginBatchResponse batch = new BeginBatchResponse { };
 		JobExecutorOptions executorOptions = new JobExecutorOptions(_session, null!, "jobNoChange", "batch1", batch, default, "", null!, new JobOptions());
-		WorkspaceExecutor executor = new (executorOptions, null, _workspace, NullLogger.Instance);
+		WorkspaceExecutor executor = new (executorOptions, _workspace, null, NullLogger.Instance);
 		await Assert.ThrowsExceptionAsync<WorkspaceMaterializationException>(() => executor.InitializeAsync(_logger, CancellationToken.None));
 	}
 
diff --git a/Engine/Source/Programs/Horde/Horde.Agent/Execution/JobExecutor.cs b/Engine/Source/Programs/Horde/Horde.Agent/Execution/JobExecutor.cs
--- a/Engine/Source/Programs/Horde/Horde.Agent/Execution/JobExecutor.cs
+++ b/Engine/Source/Programs/Horde/Horde.Agent/Execution/JobExecutor.cs
@@ -434,7 +434,7 @@ namespace Horde.Agent.Execution
 
 		protected abstract Task<bool> ExecuteAsync(BeginStepResponse step, ILogger logger, CancellationToken cancellationToken);
 
-		protected virtual async Task<bool> SetupAsync(BeginStepResponse step, DirectoryReference workspaceDir, DirectoryReference? sharedStorageDir, ILogger logger, CancellationToken cancellationToken)
+		protected virtual async Task<bool> SetupAsync(BeginStepResponse step, DirectoryReference workspaceDir, DirectoryReference? sharedStorageDir, bool? useP4, ILogger logger, CancellationToken cancellationToken)
 		{
 			FileReference definitionFile = FileReference.Combine(workspaceDir, "Engine", "Saved", "Horde", "Exported.json");
 
@@ -467,7 +467,7 @@ namespace Horde.Agent.Execution
 				arguments.Append($" CopyUAT -WithLauncher -TargetDir=\"{buildDir}\"");
 			}
 
-			int result = await ExecuteAutomationToolAsync(step, workspaceDir, arguments.ToString(), logger, cancellationToken);
+			int result = await ExecuteAutomationToolAsync(step, workspaceDir, arguments.ToString(), useP4, logger, cancellationToken);
 			if (result != 0)
 			{
 				return false;
@@ -686,7 +686,7 @@ namespace Horde.Agent.Execution
 			}
 		}
 
-		protected async Task<bool> ExecuteAsync(BeginStepResponse step, DirectoryReference workspaceDir, DirectoryReference? sharedStorageDir, ILogger logger, CancellationToken cancellationToken)
+		protected async Task<bool> ExecuteAsync(BeginStepResponse step, DirectoryReference workspaceDir, DirectoryReference? sharedStorageDir, bool? useP4, ILogger logger, CancellationToken cancellationToken)
 		{
 			StringBuilder arguments = new StringBuilder("BuildGraph");
 			if (_preprocessScript)
@@ -717,7 +717,7 @@ namespace Horde.Agent.Execution
 			_xgeMetadataExtractor?.ClearLocalIbMonFiles();
 			if (_jobOptions.UseNewTempStorage ?? false)
 			{
-				bool result = await ExecuteWithTempStorageAsync(step, workspaceDir, arguments.ToString(), logger, cancellationToken);
+				bool result = await ExecuteWithTempStorageAsync(step, workspaceDir, arguments.ToString(), useP4, logger, cancellationToken);
 				return result;
 			}
 			else
@@ -727,7 +727,7 @@ namespace Horde.Agent.Execution
 					arguments.AppendArgument("-SharedStorageDir=", sharedStorageDir.FullName);
 				}
 				
-				bool result = await ExecuteAutomationToolAsync(step, workspaceDir, arguments.ToString(), logger, cancellationToken) == 0;
+				bool result = await ExecuteAutomationToolAsync(step, workspaceDir, arguments.ToString(), useP4, logger, cancellationToken) == 0;
 				return result;
 			}
 		}
@@ -767,7 +767,7 @@ namespace Horde.Agent.Execution
 			}
 		}
 
-		private async Task<bool> ExecuteWithTempStorageAsync(BeginStepResponse step, DirectoryReference workspaceDir, string arguments, ILogger logger, CancellationToken cancellationToken)
+		private async Task<bool> ExecuteWithTempStorageAsync(BeginStepResponse step, DirectoryReference workspaceDir, string arguments, bool? useP4, ILogger logger, CancellationToken cancellationToken)
 		{
 			DirectoryReference manifestDir = DirectoryReference.Combine(workspaceDir, "Engine", "Saved", "BuildGraph");
 
@@ -829,7 +829,7 @@ namespace Horde.Agent.Execution
 			}
 
 			// Run UAT
-			if (await ExecuteAutomationToolAsync(step, workspaceDir, arguments, logger, cancellationToken) != 0)
+			if (await ExecuteAutomationToolAsync(step, workspaceDir, arguments, useP4, logger, cancellationToken) != 0)
 			{
 				return false;
 			}
@@ -993,7 +993,7 @@ namespace Horde.Agent.Execution
 			return true;
 		}
 
-		protected async Task<int> ExecuteAutomationToolAsync(BeginStepResponse step, DirectoryReference workspaceDir, string arguments, ILogger logger, CancellationToken cancellationToken)
+		protected async Task<int> ExecuteAutomationToolAsync(BeginStepResponse step, DirectoryReference workspaceDir, string arguments, bool? useP4, ILogger logger, CancellationToken cancellationToken)
 		{
 			int result;
 			using IScope scope = GlobalTracer.Instance.BuildSpan("BuildGraph").StartActive();
@@ -1002,6 +1002,11 @@ namespace Horde.Agent.Execution
 			{
 				arguments += " -NoCompile";
 			}
+			
+			if (useP4 is false)
+			{
+				arguments += " -NoP4";
+			}
 
 			if (RuntimeInformation.IsOSPlatform(OSPlatform.Windows))
 			{
diff --git a/Engine/Source/Programs/Horde/Horde.Agent/Execution/LocalExecutor.cs b/Engine/Source/Programs/Horde/Horde.Agent/Execution/LocalExecutor.cs
--- a/Engine/Source/Programs/Horde/Horde.Agent/Execution/LocalExecutor.cs
+++ b/Engine/Source/Programs/Horde/Horde.Agent/Execution/LocalExecutor.cs
@@ -55,14 +55,14 @@ namespace Horde.Agent.Execution
 
 		protected override Task<bool> SetupAsync(BeginStepResponse step, ILogger logger, CancellationToken cancellationToken)
 		{
-			return SetupAsync(step, _localWorkspaceDir, null, logger, cancellationToken);
+			return SetupAsync(step, _localWorkspaceDir, null, null, logger, cancellationToken);
 		}
 
 		protected override Task<bool> ExecuteAsync(BeginStepResponse step, ILogger logger, CancellationToken cancellationToken)
 		{
 			if (_settings.RunSteps)
 			{
-				return ExecuteAsync(step, _localWorkspaceDir, null,  logger, cancellationToken);
+				return ExecuteAsync(step, _localWorkspaceDir, null, null, logger, cancellationToken);
 			}
 			else
 			{
diff --git a/Engine/Source/Programs/Horde/Horde.Agent/Execution/PerforceExecutor.cs b/Engine/Source/Programs/Horde/Horde.Agent/Execution/PerforceExecutor.cs
--- a/Engine/Source/Programs/Horde/Horde.Agent/Execution/PerforceExecutor.cs
+++ b/Engine/Source/Programs/Horde/Horde.Agent/Execution/PerforceExecutor.cs
@@ -214,13 +214,13 @@ namespace Horde.Agent.Execution
 		protected override async Task<bool> SetupAsync(BeginStepResponse step, ILogger logger, CancellationToken cancellationToken)
 		{
 			PerforceLogger perforceLogger = CreatePerforceLogger(logger);
-			return await SetupAsync(step, _workspace.WorkspaceDir, _sharedStorageDir, perforceLogger, cancellationToken);
+			return await SetupAsync(step, _workspace.WorkspaceDir, _sharedStorageDir, null, perforceLogger, cancellationToken);
 		}
 
 		protected override async Task<bool> ExecuteAsync(BeginStepResponse step, ILogger logger, CancellationToken cancellationToken)
 		{
 			PerforceLogger perforceLogger = CreatePerforceLogger(logger);
-			return await ExecuteAsync(step, _workspace.WorkspaceDir, _sharedStorageDir, perforceLogger, cancellationToken);
+			return await ExecuteAsync(step, _workspace.WorkspaceDir, _sharedStorageDir, null, perforceLogger, cancellationToken);
 		}
 
 		public override async Task FinalizeAsync(ILogger logger, CancellationToken cancellationToken)
diff --git a/Engine/Source/Programs/Horde/Horde.Agent/Execution/WorkspaceExecutor.cs b/Engine/Source/Programs/Horde/Horde.Agent/Execution/WorkspaceExecutor.cs
--- a/Engine/Source/Programs/Horde/Horde.Agent/Execution/WorkspaceExecutor.cs
+++ b/Engine/Source/Programs/Horde/Horde.Agent/Execution/WorkspaceExecutor.cs
@@ -18,14 +18,14 @@ namespace Horde.Agent.Execution
 		public const string Name = "Workspace";
 
 		private DirectoryReference? _sharedStorageDir;
-		private readonly IWorkspaceMaterializer? _autoSdkWorkspace;
 		private readonly IWorkspaceMaterializer _workspace;
+		private readonly IWorkspaceMaterializer? _autoSdkWorkspace;
 
-		public WorkspaceExecutor(JobExecutorOptions options, IWorkspaceMaterializer? autoSdkWorkspace, IWorkspaceMaterializer workspace, ILogger logger)
+		public WorkspaceExecutor(JobExecutorOptions options, IWorkspaceMaterializer workspace, IWorkspaceMaterializer? autoSdkWorkspace, ILogger logger)
 			: base(options, logger)
 		{
-			_autoSdkWorkspace = autoSdkWorkspace;
 			_workspace = workspace;
+			_autoSdkWorkspace = autoSdkWorkspace;
 		}
 
 		public override async Task InitializeAsync(ILogger logger, CancellationToken cancellationToken)
@@ -104,7 +104,7 @@ namespace Horde.Agent.Execution
 		{
 			// Loop back to JobExecutor's SetupAsync again, but with workspace and shared storage dir set
 			DirectoryReference workspaceDir = (await _workspace.GetSettingsAsync(cancellationToken)).DirectoryPath;
-			return await SetupAsync(step, workspaceDir, _sharedStorageDir, logger, cancellationToken);
+			return await SetupAsync(step, workspaceDir, _sharedStorageDir, false, logger, cancellationToken);
 		}
 
 		/// <inheritdoc/>
@@ -112,7 +112,7 @@ namespace Horde.Agent.Execution
 		{
 			// Loop back to JobExecutor's ExecuteAsync again, but with workspace and shared storage dir set
 			DirectoryReference workspaceDir = (await _workspace.GetSettingsAsync(cancellationToken)).DirectoryPath;
-			return await ExecuteAsync(step, workspaceDir, _sharedStorageDir, logger, cancellationToken);
+			return await ExecuteAsync(step, workspaceDir, _sharedStorageDir, false, logger, cancellationToken);
 		}
 
 		/// <inheritdoc/>
@@ -153,7 +153,7 @@ namespace Horde.Agent.Execution
 					autoSdkWorkspaceInfo, options.Session.WorkingDir, true, true, _loggerFactory.CreateLogger<ManagedWorkspaceMaterializer>());
 			}
 			
-			return new WorkspaceExecutor(options, autoSdkMaterializer, workspaceMaterializer, _loggerFactory.CreateLogger<WorkspaceExecutor>());
+			return new WorkspaceExecutor(options, workspaceMaterializer, autoSdkMaterializer, _loggerFactory.CreateLogger<WorkspaceExecutor>());
 		}
 	}
 }
