Horde: Handle unset environment in BeginBatch

diff --git a/Engine/Source/Programs/Horde/Horde.Build/Jobs/JobRpcService.cs b/Engine/Source/Programs/Horde/Horde.Build/Jobs/JobRpcService.cs
--- a/Engine/Source/Programs/Horde/Horde.Build/Jobs/JobRpcService.cs
+++ b/Engine/Source/Programs/Horde/Horde.Build/Jobs/JobRpcService.cs
@@ -330,7 +330,7 @@ namespace Horde.Build.Jobs
 			response.PreflightChange = job.PreflightChange;
 			response.ClonedPreflightChange = job.ClonedPreflightChange;
 			response.Arguments.AddRange(job.Arguments);
-			response.Environment.Add(agentConfig.Environment);
+			response.Environment.Add(agentConfig.Environment ?? new Dictionary<string, string>());
 			response.ValidAgentTypes.Add(streamConfig.AgentTypes.Keys);
 
 			return response;
