Horde: Fix timeline init values, better colors for init times
#rnx
#skipci

diff --git a/Engine/Source/Programs/Horde/HordeDashboard/src/components/jobDetailsV2/JobDetailTimeline.tsx b/Engine/Source/Programs/Horde/HordeDashboard/src/components/jobDetailsV2/JobDetailTimeline.tsx
--- a/Engine/Source/Programs/Horde/HordeDashboard/src/components/jobDetailsV2/JobDetailTimeline.tsx
+++ b/Engine/Source/Programs/Horde/HordeDashboard/src/components/jobDetailsV2/JobDetailTimeline.tsx
@@ -86,10 +86,14 @@ class TimelineDataView extends JobDataView {
    }
 
    detailsUpdated() {
-      if (this.details?.timing) {
-         this.generateSpans();
-         this.updateReady();
+
+      if (!this.details?.jobData) {
+         return;
       }
+
+      this.generateSpans();
+      this.updateReady();
+
    }
 
    generateSpans() {
@@ -102,7 +106,7 @@ class TimelineDataView extends JobDataView {
 
       const job = jobDetails.jobData!;
 
-      let bresponses = job.batches?.filter(b => !!b.agentId && (!!b.startTime && !!b.finishTime) && b.steps?.length)!.map(b => { return { ...b } as GetBatchResponse })!;
+      let bresponses = job.batches?.filter(b => !!b.agentId && (!!b.startTime && !!b.readyTime && !!b.finishTime) && b.steps?.length)!.map(b => { return { ...b } as GetBatchResponse })!;
 
       if (jobDetails.filter) {
          const label = jobDetails.filter.label;
@@ -110,6 +114,9 @@ class TimelineDataView extends JobDataView {
 
             bresponses?.forEach(b => {
                b.steps = b.steps?.filter(s => {
+                  if (!s.readyTime || !s.startTime || !s.finishTime) {
+                     return false;
+                  }
                   const node = jobDetails.nodeByStepId(s.id);
                   return label.includedNodes.indexOf(node?.name ?? "") !== -1;
                });
@@ -135,56 +142,28 @@ class TimelineDataView extends JobDataView {
 
       }
 
-      const timing = jobDetails.timing!;
-
       let maxTime = 0;
 
       if (filterStep && filterStep.finishTime) {
          maxTime = new Date(filterStep.finishTime).getTime();
       }
 
-      let batches: Batch[] = [];
+      let batches: Batch[] = [];      
 
       bresponses.forEach(b => {
 
-         const utcStart = new Date(b.startTime as string);
-         const utcFinish = new Date(b.finishTime as string);
+         const utcBatchStart = new Date(b.startTime as string);
+         const utcBatchFinish = new Date(b.finishTime as string);
+         const utcBatchReady = new Date(b.readyTime as string)
+         const utcStepStart = new Date(b.steps[0].startTime as string)
 
          if (!filterStep || !maxTime) {
-            if (utcFinish.getTime() > maxTime) {
-               maxTime = utcFinish.getTime();
+            if (utcBatchFinish.getTime() > maxTime) {
+               maxTime = utcBatchFinish.getTime();
             }
          }
 
-         let utcWaitStart = utcStart;
-         let utcInitStart = utcStart;
-         let utcInitFinish = utcFinish;
-         let utcWaitFinish = utcFinish;
-
-         const stepTiming = timing.steps[b.steps[0].id];
-         const anchor = b.steps[0]?.startTime ? new Date(b.steps[0]?.startTime) : utcStart;
-
-         /*
-         // verify batch timing is consistent, lots of redundant data from step timing
-         console.assert(stepTiming);
-         b.steps.forEach(s => {
-            const otime = timing.steps[s.id];
-            console.assert(stepTiming.totalInitTime == otime.totalInitTime);
-            console.assert(stepTiming.totalWaitTime == otime.totalWaitTime);
-         });
-         */
-
-         //console.assert(stepTiming && stepTiming.totalWaitTime && stepTiming.totalInitTime)
-
-         if (stepTiming && stepTiming.totalWaitTime && stepTiming.totalInitTime) {
-
-            utcWaitStart = new Date(anchor.getTime() - (stepTiming.totalWaitTime /*- stepTiming.totalInitTime*/) * 1000);
-            utcInitStart = new Date(anchor.getTime() - stepTiming.totalInitTime * 1000);
-            utcWaitFinish = utcInitStart;
-            utcInitFinish = anchor;
-         }
-
-         batches.push({ ...b, utcStart: utcStart, utcFinish: utcFinish, utcInitStart: utcInitStart, utcInitFinish: utcInitFinish, utcWaitStart: utcWaitStart, utcWaitFinish: utcWaitFinish });
+         batches.push({ ...b, utcStart: utcBatchStart, utcFinish: utcBatchFinish, utcWaitStart: utcBatchReady, utcWaitFinish: utcBatchStart, utcInitStart: utcBatchStart, utcInitFinish: utcStepStart });
 
       });
 
@@ -447,13 +426,15 @@ class TimelineRenderer {
          .data(d3.group(I, i => Y[i]))
          .join("g")
 
+      const waitColor = dashboard.darktheme ? "#5B6367" : "#D3D2D1";
+      const initColor = dashboard.darktheme ? "#2E74B3" : "#7CB5EB";
 
       g.append("g").selectAll("line")
          .data(([, I]: any) => I)
          .join("line")
          .attr("class", "separator")
-         .attr("x1", i => xScale(X[i as any].utcStart.getTime() / 1000))
-         .attr("x2", i => xScale(X[i as any].utcStart.getTime() / 1000))
+         .attr("x1", i => xScale(X[i as any].utcStart.getTime() / 1000) + 1)
+         .attr("x2", i => xScale(X[i as any].utcStart.getTime() / 1000) + 1)
          .attr("y1", i => (yScale(X[i as any].lane) as number - 5 + 16))
          .attr("y2", i => (yScale(X[i as any].lane) as number + 5 + 16))
          .attr("stroke-linecap", 0)
@@ -461,10 +442,10 @@ class TimelineRenderer {
          .attr("stroke", i => {
             const span = spans[i as any];
             if (span.type === 0) {
-               return dashboard.darktheme ? "#5B6367" : "#D3D2D1";
+               return waitColor;
             }
             if (span.type === 1) {
-               return dashboard.darktheme ? "#506F7C" : "#BAD1DB";
+               return initColor;
             }
 
             return colors[span.step?.outcome ?? "Unspecified"];
@@ -475,8 +456,8 @@ class TimelineRenderer {
          .data(([, I]: any) => I)
          .join("line")
          .attr("class", "spanline")
-         .attr("x1", i => xScale(X[i as any].utcStart.getTime() / 1000))
-         .attr("x2", i => xScale(X[i as any].utcFinish.getTime() / 1000))
+         .attr("x1", i => xScale(X[i as any].utcStart.getTime() / 1000) - 1)
+         .attr("x2", i => xScale(X[i as any].utcFinish.getTime() / 1000) - 1)
          .attr("stroke-width", i => X[i as any].filtered ? 0 : 5)
          .attr("y1", i => (yScale(X[i as any].lane) as number + 16))
          .attr("y2", i => (yScale(X[i as any].lane) as number + 16))
@@ -484,10 +465,10 @@ class TimelineRenderer {
          .attr("stroke", i => {
             const span = spans[i as any];
             if (span.type === 0) {
-               return dashboard.darktheme ? "#5B6367" : "#D3D2D1";
+               return waitColor;
             }
             if (span.type === 1) {
-               return dashboard.darktheme ? "#506F7C" : "#BAD1DB";
+               return initColor;
             }
 
             return colors[span.step?.outcome ?? "Unspecified"];
@@ -723,10 +704,6 @@ const TimelineGraph: React.FC<{ dataView: TimelineDataView, filterTime: number,
    const [state, setState] = useState<{ graph?: TimelineRenderer }>({});
    const navigate = useNavigate();
 
-   if (!dataView.details?.timing) {
-      return null;
-   }
-
    // todo: we can probably cache
    dataView.filterTime = filterTime;
    dataView.filterCost = filterCost;
