Horde: Fix missing TempStorageDir parameter from batch object.

diff --git a/Engine/Source/Programs/Horde/Horde.Build/Jobs/JobRpcService.cs b/Engine/Source/Programs/Horde/Horde.Build/Jobs/JobRpcService.cs
--- a/Engine/Source/Programs/Horde/Horde.Build/Jobs/JobRpcService.cs
+++ b/Engine/Source/Programs/Horde/Horde.Build/Jobs/JobRpcService.cs
@@ -330,6 +330,7 @@ namespace Horde.Build.Jobs
 			response.PreflightChange = job.PreflightChange;
 			response.ClonedPreflightChange = job.ClonedPreflightChange;
 			response.Arguments.AddRange(job.Arguments);
+			response.TempStorageDir = agentConfig.TempStorageDir;
 			response.Environment.Add(agentConfig.Environment ?? new Dictionary<string, string>());
 			response.ValidAgentTypes.Add(streamConfig.AgentTypes.Keys);
 
