Horde: Fix null references in batch responses.

diff --git a/Engine/Source/Programs/Horde/Horde.Build/Jobs/JobRpcService.cs b/Engine/Source/Programs/Horde/Horde.Build/Jobs/JobRpcService.cs
--- a/Engine/Source/Programs/Horde/Horde.Build/Jobs/JobRpcService.cs
+++ b/Engine/Source/Programs/Horde/Horde.Build/Jobs/JobRpcService.cs
@@ -330,8 +330,14 @@ namespace Horde.Build.Jobs
 			response.PreflightChange = job.PreflightChange;
 			response.ClonedPreflightChange = job.ClonedPreflightChange;
 			response.Arguments.AddRange(job.Arguments);
-			response.TempStorageDir = agentConfig.TempStorageDir;
-			response.Environment.Add(agentConfig.Environment ?? new Dictionary<string, string>());
+			if (agentConfig.TempStorageDir != null)
+			{
+				response.TempStorageDir = agentConfig.TempStorageDir;
+			}
+			if (agentConfig.Environment != null)
+			{
+				response.Environment.Add(agentConfig.Environment);
+			}
 			response.ValidAgentTypes.Add(streamConfig.AgentTypes.Keys);
 
 			return response;
