Horde - Differenciate Gauntlet event fingerprints based on event type

diff --git a/Engine/Source/Programs/Horde/Horde.Build.Tests/IssueServiceTests.cs b/Engine/Source/Programs/Horde/Horde.Build.Tests/IssueServiceTests.cs
--- a/Engine/Source/Programs/Horde/Horde.Build.Tests/IssueServiceTests.cs
+++ b/Engine/Source/Programs/Horde/Horde.Build.Tests/IssueServiceTests.cs
@@ -1563,6 +1563,141 @@ namespace Horde.Build.Tests
 			Assert.AreEqual("Errors in Update Version Files", issue.Summary);
 		}
 
+		[TestMethod]
+		public async Task GauntletTest()
+		{
+			// #1
+			// Scenario: Gauntlet test event with Name property
+			// Expected: Gauntlet fingerprint using Name prefix
+			{
+				IJob job = CreateJob(_mainStreamId, 110, "Test Build", _graph);
+				await using (TestJsonLogger logger = CreateLogger(job, 0, 0))
+				{
+					logger.LogError(KnownLogEvents.Gauntlet_TestEvent, "    Test {Name} failed", "Bar.Foo.Test");
+				}
+				await UpdateCompleteStep(job, 0, 0, JobStepOutcome.Failure);
+
+				List<IIssue> issues = await IssueCollection.FindIssuesAsync();
+				Assert.AreEqual(1, issues.Count);
+				Assert.AreEqual("Gauntlet", issues[0].Fingerprints[0].Type);
+				Assert.AreEqual("test:Bar.Foo.Test", issues[0].Fingerprints[0].Keys.First());
+				Assert.AreEqual("Gauntlet test errors with Bar.Foo.Test", issues[0].Summary);
+			}
+			// #2
+			// Scenario: Gauntlet device event with Name property
+			// Expected: Gauntlet fingerprint using Device prefix
+			{
+				IJob job = CreateJob(_mainStreamId, 120, "Test Build", _graph);
+				await using (TestJsonLogger logger = CreateLogger(job, 0, 0))
+				{
+					logger.LogWarning(KnownLogEvents.Gauntlet_DeviceEvent, "    Device {Name} reported an issue", "Foo");
+				}
+				await UpdateCompleteStep(job, 0, 0, JobStepOutcome.Failure);
+
+				List<IIssue> issues = await IssueCollection.FindIssuesAsync();
+				Assert.AreEqual(1, issues.Count);
+				Assert.AreEqual("Gauntlet", issues[0].Fingerprints[0].Type);
+				Assert.AreEqual("device:Foo", issues[0].Fingerprints[0].Keys.First());
+				Assert.AreEqual("Gauntlet device warnings with Foo", issues[0].Summary);
+			}
+			// #3
+			// Scenario: Gauntlet build drop event with File and Directory property
+			// Expected: Gauntlet fingerprint using Access prefix
+			{
+				IJob job = CreateJob(_mainStreamId, 130, "Test Build", _graph);
+				await using (TestJsonLogger logger = CreateLogger(job, 0, 0))
+				{
+					logger.LogError(KnownLogEvents.Gauntlet_BuildDropEvent, "    File {File} reported an issue", "/Bar/Foo.txt");
+					logger.LogError(KnownLogEvents.Gauntlet_BuildDropEvent, "    Folder {Directory} reported an issue", "/Bar/Foo");
+				}
+				await UpdateCompleteStep(job, 0, 0, JobStepOutcome.Failure);
+
+				List<IIssue> issues = await IssueCollection.FindIssuesAsync();
+				Assert.AreEqual(1, issues.Count);
+				Assert.AreEqual("Gauntlet", issues[0].Fingerprints[0].Type);
+				Assert.AreEqual("access:/Bar/Foo.txt", issues[0].Fingerprints[0].Keys.First());
+				Assert.AreEqual("Gauntlet access errors with /Bar/Foo.txt and with /Bar/Foo", issues[0].Summary);
+			}
+			// #4
+			// Scenario: Gauntlet Fatal event
+			// Expected: Gauntlet fingerprint using hash prefix
+			{
+				string LogMessage =
+				  "    Engine encountered a critical failure.\n"
+				+ @"Assertion failed: State.bGfxPSOSet [File:D:\build\U5M+Inc\Sync\Engine\Source\Runtime\RHI\Public\RHIValidationContext.h] [Line: 809]"
+				+ @"A Graphics PSO has to be set to set resources into a shader!"
+				+ @"	0x00007fff4b43a1f4 UnrealEditor-RHI.dll!FValidationContext::RHISetShaderParameters() [Unknown File]"
+				+ @"	0x00007fff4b3e5a01 UnrealEditor-RHI.dll!FRHICommandSetShaderParameters<FRHIGraphicsShader>::Execute() [Unknown File]"
+				+ @"	0x00007fff4b3e9a5a UnrealEditor-RHI.dll!FRHICommand<FRHICommandSetShaderParameters<FRHIGraphicsShader>,FRHICommandSetShaderParametersString1159>::ExecuteAndDestruct() [Unknown File]"
+				+ @"	0x00007fff4b3e74f7 UnrealEditor - RHI.dll!FRHICommandListBase::Execute()[Unknown File]"
+				+ @"	0x00007fff4b3f3228 UnrealEditor-RHI.dll!FRHICommandListImmediate::ExecuteAndReset() [Unknown File]"
+				+ @"	0x00007fff4b464ebf UnrealEditor-RHI.dll!FRHIComputeCommandList::SubmitCommandsHint() [Unknown File]"
+				+ @"	0x00007fff3d3066fc UnrealEditor-Renderer.dll!TBaseStaticDelegateInstance<void [Unknown File]"
+				+ @"	0x00007fff48bd7813 UnrealEditor-RenderCore.dll!FRDGBuilder::ExecutePass() [Unknown File]"
+				+ @"	0x00007fff48bd2a8a UnrealEditor-RenderCore.dll!FRDGBuilder::Execute() [Unknown File]"
+				+ @"	0x00007fff3d31fef4 UnrealEditor-Renderer.dll!FSceneRenderer::RenderThreadEnd() [Unknown File]"
+				+ @"	0x00007fff3d2ed57a UnrealEditor-Renderer.dll!`FPixelShaderUtils::AddFullscreenPass<FHZBTestPS>'::`2'::<lambda_1>::operator()() [Unknown File]"
+				+ @"	0x00007fff3d3043f9 UnrealEditor - Renderer.dll!FSceneRenderer::DoOcclusionQueries()[Unknown File]"
+				+ @"	0x00007fff3d30c888 UnrealEditor-Renderer.dll!TBaseStaticDelegateInstance<void [Unknown File]"
+				+ @"	0x00007fff501fad42 UnrealEditor-Core.dll!FNamedTaskThread::ProcessTasksNamedThread() [Unknown File]"
+				+ @"	0x00007fff501fb25e UnrealEditor-Core.dll!FNamedTaskThread::ProcessTasksUntilQuit() [Unknown File]"
+				+ @"	0x00007fff48ccbb54 UnrealEditor-RenderCore.dll!RenderingThreadMain() [Unknown File]"
+				+ @"	0x00007fff48ccff44 UnrealEditor-RenderCore.dll!FRenderingThread::Run() [Unknown File]"
+				+ @"	0x00007fff5089c862 UnrealEditor-Core.dll!FRunnableThreadWin::Run() [Unknown File]"
+				+ @"	0x00007fff5089a7bf UnrealEditor-Core.dll!FRunnableThreadWin::GuardedRun() [Unknown File]"
+				+ @"	0x00007fff8e304ed0 KERNEL32.DLL!UnknownFunction [Unknown File]"
+				+ @"	0x00007fff8f26e39b ntdll.dll!UnknownFunction [Unknown File]"
+				+ @"    Test did not run until completion. The test exited prematurely.";
+
+				IJob job = CreateJob(_mainStreamId, 140, "Test Build", _graph);
+				await using (TestJsonLogger logger = CreateLogger(job, 0, 0))
+				{
+					logger.LogError(KnownLogEvents.Gauntlet_FatalEvent, "{Message}", LogMessage);
+				}
+				await UpdateCompleteStep(job, 0, 0, JobStepOutcome.Failure);
+
+				List<IIssue> issues = await IssueCollection.FindIssuesAsync();
+				Assert.AreEqual(1, issues.Count);
+				Assert.AreEqual("Gauntlet", issues[0].Fingerprints[0].Type);
+				Assert.AreEqual("hash:", issues[0].Fingerprints[0].Keys.First().Substring(0, 5));
+				Assert.AreEqual("Gauntlet fatal errors in Update Version Files", issues[0].Summary);
+			}
+			// #5
+			// Scenario: Gauntlet Test event
+			// Expected: Gauntlet fingerprint using hash prefix
+			{
+				string[] logErrors =
+				{
+				"	Expected 'SimpleValue::Get meta equality' to be true.",
+				"	Expected 'SimpleValueSkipData::Get meta equality' to be true.",
+				"	Expected 'SimpleValueZen::Get meta equality' to be true.",
+				"	Expected 'SimpleValueZenAndDirect::Get meta equality' to be true.",
+				"	Expected 'SimpleValueSkipDataZen::Get meta equality' to be true.",
+				"	Expected 'SimpleValueSkipDataZenAndDirect::Get meta equality' to be true.",
+				"	Expected 'SimpleValueWithMeta::Get meta equality' to be true.",
+				"	Expected 'SimpleValueWithMetaSkipData::Get meta equality' to be true.",
+				"	Expected 'SimpleValueWithMetaZenAndDirect::Get meta equality' to be true.",
+				"	Expected 'SimpleValueWithMetaSkipDataZenAndDirect::Get meta equality' to be true."
+				};
+
+				IJob job = CreateJob(_mainStreamId, 150, "Test Build", _graph);
+				await using (TestJsonLogger logger = CreateLogger(job, 0, 0))
+				{
+					foreach (string error in logErrors)
+					{
+						logger.LogError(KnownLogEvents.Gauntlet_TestEvent, "{Error}", error);
+					}
+				}
+				await UpdateCompleteStep(job, 0, 0, JobStepOutcome.Failure);
+
+				List<IIssue> issues = await IssueCollection.FindIssuesAsync();
+				Assert.AreEqual(1, issues.Count);
+				Assert.AreEqual("Gauntlet", issues[0].Fingerprints[0].Type);
+				Assert.AreEqual("hash:", issues[0].Fingerprints[0].Keys.First().Substring(0, 5));
+				Assert.AreEqual("Gauntlet test errors in Update Version Files", issues[0].Summary);
+			}
+		}
+
 		[TestMethod]
 		public async Task MaskIssueTest()
 		{
diff --git a/Engine/Source/Programs/Horde/Horde.Build/Issues/Handlers/GauntletIssueHandler.cs b/Engine/Source/Programs/Horde/Horde.Build/Issues/Handlers/GauntletIssueHandler.cs
--- a/Engine/Source/Programs/Horde/Horde.Build/Issues/Handlers/GauntletIssueHandler.cs
+++ b/Engine/Source/Programs/Horde/Horde.Build/Issues/Handlers/GauntletIssueHandler.cs
@@ -3,7 +3,8 @@
 using System;
 using System.Collections.Generic;
 using System.Linq;
-using System.Text.Json;
+using System.Text;
+using System.Text.RegularExpressions;
 using EpicGames.Core;
 using Horde.Build.Jobs;
 using Horde.Build.Jobs.Graphs;
@@ -19,14 +20,29 @@ namespace Horde.Build.Issues.Handlers
 	class GauntletIssueHandler : IssueHandler
 	{
 		/// <summary>
-		/// Prefix for unit test keys
+		/// Prefix for test keys
 		/// </summary>
-		const string UnitTestPrefix = "UnitTest:";
+		const string TestPrefix = "test";
 
 		/// <summary>
-		/// Prefix for screenshot test keys
+		/// Prefix for device keys
 		/// </summary>
-		const string ScreenshotTestPrefix = "Screenshot:";
+		const string DevicePrefix = "device";
+
+		/// <summary>
+		/// Prefix for access keys
+		/// </summary>
+		const string AccessPrefix = "access";
+
+		/// <summary>
+		/// Prefix for fatal failure keys
+		/// </summary>
+		const string FatalPrefix = "fatal";
+
+		/// <summary>
+		/// Max Message Length to hash
+		/// </summary>
+		const int MaxMessageLength = 2000;
 
 		/// <inheritdoc/>
 		public override string Type => "Gauntlet";
@@ -34,6 +50,18 @@ namespace Horde.Build.Issues.Handlers
 		/// <inheritdoc/>
 		public override int Priority => 10;
 
+		/// <summary>
+		///  Known Gauntlet events
+		/// </summary>
+		static readonly Dictionary<EventId, string> s_knownGauntletEvents = new Dictionary<EventId, string> {
+			{ KnownLogEvents.Gauntlet_TestEvent, TestPrefix},
+			{ KnownLogEvents.Gauntlet_DeviceEvent, DevicePrefix},
+			{ KnownLogEvents.Gauntlet_UnrealEngineTestEvent, TestPrefix},
+			{ KnownLogEvents.Gauntlet_BuildDropEvent, AccessPrefix},
+			{ KnownLogEvents.Gauntlet_FatalEvent, FatalPrefix}
+		};
+
+
 		/// <summary>
 		/// Determines if the given event id matches
 		/// </summary>
@@ -41,64 +69,122 @@ namespace Horde.Build.Issues.Handlers
 		/// <returns>True if the given event id matches</returns>
 		public static bool IsMatchingEventId(EventId eventId)
 		{
-			return eventId == KnownLogEvents.Gauntlet_TestEvent ||
-			       eventId == KnownLogEvents.Gauntlet_DeviceEvent ||
-			       eventId == KnownLogEvents.Gauntlet_UnrealEngineTestEvent ||
-			       eventId == KnownLogEvents.Gauntlet_BuildDropEvent;
+			return s_knownGauntletEvents.ContainsKey(eventId);
+		}
+
+		/// <summary>
+		/// Return the prefix string associate with the event id
+		/// </summary>
+		/// <param name="eventId">The event id to get the information from</param>
+		/// <returns>The corresponding prefix as a string</returns>
+		public static string GetEventPrefix(EventId eventId)
+		{
+			return s_knownGauntletEvents[eventId];
 		}
 
 		/// <summary>
 		/// Parses symbol names from a log event
 		/// </summary>
 		/// <param name="eventData">The log event data</param>
-		/// <param name="unitTestNames">Receives a set of the unit test names</param>
-		private static void GetUnitTestNames(ILogEventData eventData, HashSet<string> unitTestNames)
+		/// <param name="testNames">Receives a set of the test names</param>
+		/// <param name="metadata"></param>
+		private static void GetNames(ILogEventData eventData, HashSet<string> testNames, HashSet<string> metadata)
 		{
 			foreach (ILogEventLine line in eventData.Lines)
 			{
-				string? group = null;
 				string? name = null;
 
 				string? value;
-				if (line.Data.TryGetNestedProperty("properties.group", out value))
-				{
-					group = value;
-				}
-				if (line.Data.TryGetNestedProperty("properties.name", out value))
+				if (line.Data.TryGetNestedProperty("properties.Name", out value))
 				{
 					name = value;
 				}
 
-				if (group != null && name != null)
+				if (name != null)
 				{
-					unitTestNames.Add($"{UnitTestPrefix}:{group}/{name}");
+					string prefix = GetEventPrefix(eventData.EventId!.Value);
+					testNames.Add($"{prefix}:{name}");
+					metadata.Add($"context=with {name}");
 				}
 			}
 		}
 
 		/// <summary>
-		/// Parses screenshot test names from a log event
+		/// Parses symbol file or directory from a log event
 		/// </summary>
-		/// <param name="eventData">The event data</param>
-		/// <param name="screenshotTestNames">Receives the parsed screenshot test names</param>
-		private static void GetScreenshotTestNames(ILogEventData eventData, HashSet<string> screenshotTestNames)
+		/// <param name="eventData">The log event data</param>
+		/// <param name="paths">Receives a set of the paths</param>
+		/// <param name="metadata"></param>
+		private static void GetPaths(ILogEventData eventData, HashSet<string> paths, HashSet<string> metadata)
 		{
-			foreach (ILogEventLine line in eventData.Lines)
+			if(eventData.EventId == KnownLogEvents.Gauntlet_BuildDropEvent)
 			{
-				if (line.Data.TryGetProperty("properties", JsonValueKind.Object, out JsonElement properties))
+				foreach (ILogEventLine line in eventData.Lines)
 				{
-					foreach (JsonProperty property in properties.EnumerateObject())
+					string? path = null;
+
+					string? value;
+					if (line.Data.TryGetNestedProperty("properties.File", out value))
+					{
+						path = value;
+					}
+					else if (line.Data.TryGetNestedProperty("properties.Directory", out value))
 					{
-						JsonElement value = property.Value;
-						if (value.ValueKind == JsonValueKind.Object && value.HasStringProperty("type", "Screenshot") && value.TryGetStringProperty("name", out string? name))
-						{
-							screenshotTestNames.Add($"{ScreenshotTestPrefix}:{name}");
-						}
+						path = value;
+					}
+
+					if (path != null)
+					{
+						paths.Add($"{AccessPrefix}:{path}");
+						metadata.Add($"context=with {path}");
 					}
 				}
 			}
 		}
 
+		/// <summary>
+		/// Produce a hash from error message
+		/// </summary>
+		/// <param name="job">The job that spawned the event</param>
+		/// <param name="eventData">The log event data</param>
+		/// <param name="keys">Receives a set of the keys</param>
+		/// <param name="metadata">Receives a set of metadata</param>
+		/// <param name="node">Node that was executed</param>
+		private static void GetHash(IJob job, ILogEventData eventData, HashSet<string> keys, HashSet<string> metadata, INode node)
+		{
+
+			string error = eventData.Message.Length > MaxMessageLength? eventData.Message.Substring(0, MaxMessageLength): eventData.Message;
+
+			if (TryGetHash(error, out Md5Hash hash))
+			{
+				keys.Add($"hash:{hash}:stream:{job.StreamId}");
+			}
+			else
+			{
+				keys.Add($"{node.Name}");
+			}
+			metadata.Add($"context=in {node.Name}");
+		}
+
+		private static bool TryGetHash(string message, out Md5Hash hash)
+		{
+			string sanitized = message.ToUpperInvariant();
+			sanitized = Regex.Replace(sanitized, @"(?<![a-zA-Z])(?:[A-Z]:|/)[^ :]+[/\\]SYNC[/\\]", "{root}/"); // Redact things that look like workspace roots; may be different between agents
+			sanitized = Regex.Replace(sanitized, @"0[xX][0-9a-fA-F]+", "H"); // Redact hex strings
+			sanitized = Regex.Replace(sanitized, @"\d[\d.,:]*", "n"); // Redact numbers and timestamp like things
+
+			if (sanitized.Length > 30)
+			{
+				hash = Md5Hash.Compute(Encoding.UTF8.GetBytes(sanitized));
+				return true;
+			}
+			else
+			{
+				hash = Md5Hash.Zero;
+				return false;
+			}
+		}
+
 		/// <inheritdoc/>
 		public override void TagEvents(IJob job, INode node, IReadOnlyNodeAnnotations annotations, IReadOnlyList<IssueEvent> stepEvents)
 		{
@@ -107,13 +193,16 @@ namespace Horde.Build.Issues.Handlers
 				if (stepEvent.EventId != null && IsMatchingEventId(stepEvent.EventId.Value))
 				{
 					HashSet<string> keys = new HashSet<string>();
-					GetUnitTestNames(stepEvent.EventData, keys);
-					GetScreenshotTestNames(stepEvent.EventData, keys);
-
-					if (keys.Count > 0)
+					HashSet<string> metadata = new HashSet<string>();
+					GetNames(stepEvent.EventData, keys, metadata);
+					GetPaths(stepEvent.EventData, keys, metadata);
+					if(keys.Count == 0)
 					{
-						stepEvent.Fingerprint = new NewIssueFingerprint(Type, keys, null, null);
+						GetHash(job, stepEvent.EventData, keys, metadata, node);
 					}
+					metadata.Add($"type={GetEventPrefix(stepEvent.EventId.Value)}");
+
+					stepEvent.Fingerprint = new NewIssueFingerprint(Type, keys, null, metadata);
 				}
 			}
 		}
@@ -133,38 +222,10 @@ namespace Horde.Build.Issues.Handlers
 		/// <inheritdoc/>
 		public override string GetSummary(IIssueFingerprint fingerprint, IssueSeverity severity)
 		{
-			List<string> unitTestNames = fingerprint.Keys.Where(x => x.StartsWith(UnitTestPrefix, StringComparison.Ordinal)).Select(x => x.Substring(UnitTestPrefix.Length + 1)).ToList();
-			if (unitTestNames.Count > 0)
-			{
-				HashSet<string> groupNames = new HashSet<string>(StringComparer.OrdinalIgnoreCase);
-				HashSet<string> testNames = new HashSet<string>(StringComparer.OrdinalIgnoreCase);
-				foreach (string unitTestName in unitTestNames)
-				{
-					int idx = unitTestName.IndexOf('/', StringComparison.OrdinalIgnoreCase);
-					if (idx != -1)
-					{
-						groupNames.Add(unitTestName.Substring(0, idx));
-					}
-					testNames.Add(unitTestName.Substring(idx + 1));
-				}
-
-				if (groupNames.Count == 1)
-				{
-					return $"{groupNames.First()} test failures: {StringUtils.FormatList(testNames.ToArray(), 3)}";
-				}
-				else
-				{
-					return $"{StringUtils.FormatList(groupNames.OrderBy(x => x).ToArray(), 100)} test failures";
-				}
-			}
-
-			List<string> screenshotTestNames = fingerprint.Keys.Where(x => x.StartsWith(ScreenshotTestPrefix, StringComparison.Ordinal)).Select(x => x.Substring(ScreenshotTestPrefix.Length + 1)).ToList();
-			if (screenshotTestNames.Count > 0)
-			{
-				return $"Screenshot test failures: {StringUtils.FormatList(screenshotTestNames.ToArray(), 3)}";
-			}
-
-			return "Test failures";
+			string Title = fingerprint.GetMetadataValues("type").FirstOrDefault() ?? "unknown";
+			string severityText = (severity == IssueSeverity.Warning) ? "warnings" : "errors";
+			string[] Errors = fingerprint.GetMetadataValues("context").ToArray();
+			return $"{Type} {Title} {severityText} {StringUtils.FormatList(Errors, 2)}";
 		}
 	}
 }
