Horde: Stop instances stuck in pending state for longer than 5 mins

Check is applied during scale-out of fleet.

diff --git a/Engine/Source/Programs/Horde/Horde.Build.Tests/Fleet/AwsRecyclingFleetManagerTest.cs b/Engine/Source/Programs/Horde/Horde.Build.Tests/Fleet/AwsRecyclingFleetManagerTest.cs
--- a/Engine/Source/Programs/Horde/Horde.Build.Tests/Fleet/AwsRecyclingFleetManagerTest.cs
+++ b/Engine/Source/Programs/Horde/Horde.Build.Tests/Fleet/AwsRecyclingFleetManagerTest.cs
@@ -13,6 +13,7 @@ using Microsoft.VisualStudio.TestTools.UnitTesting;
 using Horde.Build.Agents;
 using Horde.Build.Agents.Fleet.Providers;
 using Horde.Build.Utilities;
+using HordeCommon;
 using Microsoft.Extensions.Logging;
 using JsonSerializer = System.Text.Json.JsonSerializer;
 
@@ -21,6 +22,8 @@ namespace Horde.Build.Tests.Fleet
 	[TestClass]
 	public class AwsRecyclingFleetManagerTest : TestSetup
 	{
+		private readonly FakeClock _clock = new();
+		
 		[TestMethod]
 		public async Task ExpandAmazonEc2ExceptionsArePropagated()
 		{
@@ -152,6 +155,25 @@ namespace Horde.Build.Tests.Fleet
 			Assert.AreEqual(FakeAmazonEc2.StatePending.Name, ec2.GetInstance(i4.InstanceId)!.State.Name);
 			Assert.AreEqual(FakeAmazonEc2.StatePending.Name, ec2.GetInstance(i2.InstanceId)!.State.Name);
 		}
+		
+		[TestMethod]
+		public async Task StopInstancesStuckInPending()
+		{
+			FakeAmazonEc2 ec2 = new ();
+			Instance i1 = ec2.AddInstance(FakeAmazonEc2.StatePending, InstanceType.M5Large, launchTime: _clock.UtcNow.Subtract(TimeSpan.FromMinutes(50)));
+			Instance i2 = ec2.AddInstance(FakeAmazonEc2.StatePending, InstanceType.M5Large, launchTime: _clock.UtcNow.Subtract(TimeSpan.FromMinutes(3)));
+			Instance i3 = ec2.AddInstance(FakeAmazonEc2.StateStopped, InstanceType.M5Large);
+			Instance i4 = ec2.AddInstance(FakeAmazonEc2.StateStopped, InstanceType.M5Large);
+			ec2.SetCapacity(FakeAmazonEc2.AzUsEast1A, InstanceType.M5Large, 2);
+
+			ScaleResult result = await ExpandPoolAsync(ec2.Get(), 2, new(new List<string>()));
+			Assert.AreEqual(FleetManagerOutcome.Success, result.Outcome);
+			Assert.AreEqual(3, ec2.GetPendingInstanceCount());
+			Assert.AreEqual(FakeAmazonEc2.StateStopping.Name, ec2.GetInstance(i1.InstanceId)!.State.Name);
+			Assert.AreEqual(FakeAmazonEc2.StatePending.Name, ec2.GetInstance(i2.InstanceId)!.State.Name);
+			Assert.AreEqual(FakeAmazonEc2.StatePending.Name, ec2.GetInstance(i3.InstanceId)!.State.Name);
+			Assert.AreEqual(FakeAmazonEc2.StatePending.Name, ec2.GetInstance(i4.InstanceId)!.State.Name);
+		}
 
 		[TestMethod]
 		public void DistributesInstanceRequestsEvenly()
@@ -186,7 +208,7 @@ namespace Horde.Build.Tests.Fleet
 			}
 		}
 
-		private async Task ExpandPoolAsync(IAmazonEC2 ec2, int numRequestedInstances, AwsRecyclingFleetManagerSettings settings)
+		private async Task<ScaleResult> ExpandPoolAsync(IAmazonEC2 ec2, int numRequestedInstances, AwsRecyclingFleetManagerSettings settings)
 		{
 			using ILoggerFactory loggerFactory = LoggerFactory.Create(builder =>
 				{
@@ -197,8 +219,8 @@ namespace Horde.Build.Tests.Fleet
 			ILogger<AwsRecyclingFleetManager> logger = loggerFactory.CreateLogger<AwsRecyclingFleetManager>();
 			IPool pool = await PoolService.CreatePoolAsync("testPool", null, true, 0, 0, sizeStrategy: PoolSizeStrategy.NoOp);
 			using NoOpDogStatsd dogStatsd = new ();
-			AwsRecyclingFleetManager manager = new (ec2, AgentCollection, dogStatsd, settings, logger);
-			await manager.ExpandPoolAsync(pool, new List<IAgent>(), numRequestedInstances, CancellationToken.None);
+			AwsRecyclingFleetManager manager = new (ec2, AgentCollection, dogStatsd, _clock, settings, logger);
+			return await manager.ExpandPoolAsync(pool, new List<IAgent>(), numRequestedInstances, CancellationToken.None);
 		}
 	}
 }
\ No newline at end of file
diff --git a/Engine/Source/Programs/Horde/Horde.Build.Tests/Fleet/FakeAmazonEc2.cs b/Engine/Source/Programs/Horde/Horde.Build.Tests/Fleet/FakeAmazonEc2.cs
--- a/Engine/Source/Programs/Horde/Horde.Build.Tests/Fleet/FakeAmazonEc2.cs
+++ b/Engine/Source/Programs/Horde/Horde.Build.Tests/Fleet/FakeAmazonEc2.cs
@@ -20,6 +20,7 @@ namespace Horde.Build.Tests.Fleet;
 public class FakeAmazonEc2
 {
 	public static readonly InstanceState StatePending = new() { Code = 0, Name = "pending" };
+	public static readonly InstanceState StateStopping = new() { Code = 64, Name = "stopping" };
 	public static readonly InstanceState StateStopped = new() { Code = 80, Name = "stopped" };
 	public const string AzUsEast1A = "us-east-1a";
 	public const string AzUsEast1B = "us-east-1b";
@@ -53,6 +54,10 @@ public class FakeAmazonEc2
 			.Setup(x => x.StartInstancesAsync(It.IsAny<StartInstancesRequest>(), It.IsAny<CancellationToken>()))
 			.Returns(StartInstancesAsync);
 		
+		_mock
+			.Setup(x => x.StopInstancesAsync(It.IsAny<StopInstancesRequest>(), It.IsAny<CancellationToken>()))
+			.Returns(StopInstancesAsync);
+		
 		_mock
 			.Setup(x => x.DescribeInstancesAsync(It.IsAny<DescribeInstancesRequest>(), It.IsAny<CancellationToken>()))
 			.Returns(DescribeInstancesAsync);
@@ -110,9 +115,39 @@ public class FakeAmazonEc2
 
 	private Task<DescribeInstancesResponse> DescribeInstancesAsync(DescribeInstancesRequest request, CancellationToken cancellationToken)
 	{
+		List<Instance> filteredInstances = new();
+		if (request.Filters.Count > 0)
+		{
+			foreach (Instance i in _instances.Values)
+			{
+				foreach (Filter f in request.Filters)
+				{
+					if (f.Name == "instance-state-name")
+					{
+						if (i.State.Name == f.Values[0])
+						{
+							filteredInstances.Add(i);
+						}
+					}
+					else if (f.Name.StartsWith("tag:", StringComparison.Ordinal))
+					{
+						// TODO: respect tag filtering
+					}
+					else
+					{
+						throw new Exception("Unknown filter " + f.Name);
+					}
+				}
+			}
+		}
+		else
+		{
+			filteredInstances = _instances.Values.ToList();
+		}
+		
 		return Task.FromResult(new DescribeInstancesResponse
 		{
-			Reservations = new () { new Reservation { Instances = _instances.Values.Select(CopyInstance).ToList() } },
+			Reservations = new () { new Reservation { Instances = filteredInstances.Select(CopyInstance).ToList() } },
 			HttpStatusCode = HttpStatusCode.OK
 		});
 	}
@@ -166,6 +201,31 @@ public class FakeAmazonEc2
 		});
 	}
 	
+	private Task<StopInstancesResponse> StopInstancesAsync(StopInstancesRequest request, CancellationToken cancellationToken)
+	{
+		List<InstanceStateChange> stateChanges = new();
+		
+		foreach (string instanceId in request.InstanceIds)
+		{
+			if (_instances[instanceId].State.Name == StatePending.Name)
+			{
+				stateChanges.Add(new()
+				{
+					InstanceId = instanceId,
+					CurrentState = StateStopping,
+					PreviousState = _instances[instanceId].State
+				});
+				_instances[instanceId].State = StateStopping;
+			}
+		}
+		
+		return Task.FromResult(new StopInstancesResponse
+		{
+			StoppingInstances = stateChanges,
+			HttpStatusCode = HttpStatusCode.OK
+		});
+	}
+	
 	private async Task<ModifyInstanceAttributeResponse> ModifyInstanceAttributeAsync(ModifyInstanceAttributeRequest request, CancellationToken cancellationToken)
 	{
 		if (!_instances.ContainsKey(request.InstanceId))
diff --git a/Engine/Source/Programs/Horde/Horde.Build/Agents/Fleet/FleetManagerFactory.cs b/Engine/Source/Programs/Horde/Horde.Build/Agents/Fleet/FleetManagerFactory.cs
--- a/Engine/Source/Programs/Horde/Horde.Build/Agents/Fleet/FleetManagerFactory.cs
+++ b/Engine/Source/Programs/Horde/Horde.Build/Agents/Fleet/FleetManagerFactory.cs
@@ -5,6 +5,7 @@ using System.Text.Json;
 using Amazon.AutoScaling;
 using Amazon.EC2;
 using Horde.Build.Agents.Fleet.Providers;
+using HordeCommon;
 using Microsoft.Extensions.DependencyInjection;
 using Microsoft.Extensions.Logging;
 using Microsoft.Extensions.Options;
@@ -33,6 +34,7 @@ public interface IFleetManagerFactory
 public sealed class FleetManagerFactory : IFleetManagerFactory
 {
 	private readonly IAgentCollection _agentCollection;
+	private readonly IClock _clock;
 	private readonly IDogStatsd _dogStatsd;
 	private readonly IServiceProvider _provider;
 	private readonly IOptionsMonitor<ServerSettings> _settings;
@@ -44,9 +46,10 @@ public sealed class FleetManagerFactory : IFleetManagerFactory
 	/// <summary>
 	/// Constructor
 	/// </summary>
-	public FleetManagerFactory(IAgentCollection agentCollection, IDogStatsd dogStatsd, IServiceProvider provider, IOptionsMonitor<ServerSettings> settings, ILoggerFactory loggerFactory)
+	public FleetManagerFactory(IAgentCollection agentCollection, IClock clock, IDogStatsd dogStatsd, IServiceProvider provider, IOptionsMonitor<ServerSettings> settings, ILoggerFactory loggerFactory)
 	{
 		_agentCollection = agentCollection;
+		_clock = clock;
 		_dogStatsd = dogStatsd;
 		_provider = provider;
 		_settings = settings;
@@ -67,7 +70,7 @@ public sealed class FleetManagerFactory : IFleetManagerFactory
 			FleetManagerType.AwsReuse =>
 				new AwsReuseFleetManager(GetAwsEc2(type), _agentCollection, DeserializeSettings<AwsReuseFleetManagerSettings>(config), _loggerFactory.CreateLogger<AwsReuseFleetManager>()),
 			FleetManagerType.AwsRecycle =>
-				new AwsRecyclingFleetManager(GetAwsEc2(type), _agentCollection, _dogStatsd, DeserializeSettings<AwsRecyclingFleetManagerSettings>(config), _loggerFactory.CreateLogger<AwsRecyclingFleetManager>()),
+				new AwsRecyclingFleetManager(GetAwsEc2(type), _agentCollection, _dogStatsd, _clock, DeserializeSettings<AwsRecyclingFleetManagerSettings>(config), _loggerFactory.CreateLogger<AwsRecyclingFleetManager>()),
 			FleetManagerType.AwsAsg =>
 				new AwsAsgFleetManager(GetAwsAutoScaling(type), DeserializeSettings<AwsAsgSettings>(config), _loggerFactory.CreateLogger<AwsAsgFleetManager>()),
 			_ => throw new ArgumentException("Unknown fleet manager type " + type)
diff --git a/Engine/Source/Programs/Horde/Horde.Build/Agents/Fleet/Providers/AwsRecyclingFleetManager.cs b/Engine/Source/Programs/Horde/Horde.Build/Agents/Fleet/Providers/AwsRecyclingFleetManager.cs
--- a/Engine/Source/Programs/Horde/Horde.Build/Agents/Fleet/Providers/AwsRecyclingFleetManager.cs
+++ b/Engine/Source/Programs/Horde/Horde.Build/Agents/Fleet/Providers/AwsRecyclingFleetManager.cs
@@ -11,6 +11,7 @@ using Amazon.EC2.Model;
 using Amazon.Runtime;
 using EpicGames.Core;
 using Horde.Build.Agents.Pools;
+using HordeCommon;
 using Microsoft.Extensions.Logging;
 using OpenTracing;
 using OpenTracing.Util;
@@ -91,16 +92,18 @@ public sealed class AwsRecyclingFleetManager : IFleetManager
 	private readonly IAmazonEC2 _ec2;
 	private readonly IAgentCollection _agentCollection;
 	private readonly IDogStatsd _dogStatsd;
+	private readonly IClock _clock;
 	private readonly ILogger _logger;
 
 	/// <summary>
 	/// Constructor
 	/// </summary>
-	public AwsRecyclingFleetManager(IAmazonEC2 ec2, IAgentCollection agentCollection, IDogStatsd dogStatsd, AwsRecyclingFleetManagerSettings settings, ILogger<AwsRecyclingFleetManager> logger)
+	public AwsRecyclingFleetManager(IAmazonEC2 ec2, IAgentCollection agentCollection, IDogStatsd dogStatsd, IClock clock, AwsRecyclingFleetManagerSettings settings, ILogger<AwsRecyclingFleetManager> logger)
 	{
 		_ec2 = ec2;
 		_agentCollection = agentCollection;
 		_dogStatsd = dogStatsd;
+		_clock = clock;
 		_logger = logger;
 		Settings = settings;
 	}
@@ -116,6 +119,7 @@ public sealed class AwsRecyclingFleetManager : IFleetManager
 
 		using IDisposable logScope = _logger.BeginScope(new Dictionary<string, object> { ["PoolId"] = pool.Id });
 
+		await StopStuckPendingInstancesAsync(pool, cancellationToken);
 		Dictionary<string, List<Instance>> candidatesPerAz = await GetCandidateInstancesAsync(pool, cancellationToken);
 		Dictionary<string, int> instanceCountPerAz = candidatesPerAz.ToDictionary(x => x.Key, x => x.Value.Count);
 		Dictionary<string, int> requestCountPerAz = DistributeInstanceRequestsOverAzs(requestedInstancesCount, instanceCountPerAz, out int stoppedInstancesMissingCount);
@@ -389,6 +393,43 @@ public sealed class AwsRecyclingFleetManager : IFleetManager
 		}
 	}
 	
+	/// <summary>
+	/// Stop instances that have been stuck in state pending for longer than X time
+	/// </summary>
+	/// <param name="pool">Which pool instances must belong to</param>
+	/// <param name="cancellationToken">Cancellation token for the call</param>
+	private async Task StopStuckPendingInstancesAsync(IPool pool, CancellationToken cancellationToken)
+	{
+		TimeSpan pendingStateTimeout = TimeSpan.FromMinutes(5);
+		
+		using IScope scope = GlobalTracer.Instance.BuildSpan("StopStuckPendingInstances")
+			.WithTag("PoolId", pool.Id.ToString()).StartActive();
+		
+		// Find pending instances in the correct pool
+		DescribeInstancesRequest request = new()
+		{
+			Filters = new List<Filter>
+			{
+				new("instance-state-name", new List<string> { InstanceStateName.Pending.Value }),
+				new("tag:" + AwsFleetManager.PoolTagName, new List<string> { pool.Name })
+			}
+		};
+		DescribeInstancesResponse response = await _ec2.DescribeInstancesAsync(request, cancellationToken);
+		scope.Span.SetTag("res.statusCode", (int)response.HttpStatusCode);
+		scope.Span.SetTag("res.numReservations", response.Reservations.Count);
+
+		List<string> stuckInstanceIds = response.Reservations
+			.SelectMany(x => x.Instances)
+			.Where(x => x.LaunchTime < _clock.UtcNow - pendingStateTimeout)
+			.Select(x => x.InstanceId)
+			.ToList();
+
+		StopInstancesRequest stopRequest = new() { InstanceIds = stuckInstanceIds };
+		StopInstancesResponse stopResponse = await _ec2.StopInstancesAsync(stopRequest, cancellationToken);
+		scope.Span.SetTag("stopRes.statusCode", (int)stopResponse.HttpStatusCode);
+		scope.Span.SetTag("stopRes.numStoppingInstances", stopResponse.StoppingInstances.Count);
+	}
+	
 	/// <summary>
 	/// Get list of stopped instances from EC2 API that can be used scaling out a given pool
 	/// </summary>
