Horde: Prevent race in audit log test by proactively waiting on log flush

diff --git a/Engine/Source/Programs/Horde/Horde.Build.Tests/AgentServiceTest.cs b/Engine/Source/Programs/Horde/Horde.Build.Tests/AgentServiceTest.cs
--- a/Engine/Source/Programs/Horde/Horde.Build.Tests/AgentServiceTest.cs
+++ b/Engine/Source/Programs/Horde/Horde.Build.Tests/AgentServiceTest.cs
@@ -4,6 +4,7 @@ using System;
 using System.Collections.Generic;
 using System.Linq;
 using System.Security.Claims;
+using System.Threading;
 using System.Threading.Tasks;
 using EpicGames.Core;
 using Horde.Build.Agents;
@@ -63,9 +64,9 @@ public class AgentServiceTest : TestSetup
 		Fixture fixture = await CreateFixtureAsync();
 		IAuditLogChannel<AgentId> agentLogger = AgentCollection.GetLogger(fixture.Agent1.Id);
 
-		async Task<bool> AuditLogContains(string text)
+		async Task<bool> AuditLogContains(string text, int waitForNumMessages = -1)
 		{
-			await FlushAuditLogsAsync();
+			await FlushAuditLogsAsync(waitForNumMessages);
 			return await agentLogger.FindAsync().AnyAsync(x => x.Data.Contains(text, StringComparison.Ordinal));
 		}
 		
@@ -75,14 +76,33 @@ public class AgentServiceTest : TestSetup
 		
 		props = new () { $"{KnownPropertyNames.AwsInstanceType}=c6.xlarge" };
 		agent = await AgentService.CreateSessionAsync(agent, AgentStatus.Ok, props, new Dictionary<string, int>(), "test");
-		Assert.IsTrue(await AuditLogContains("AWS EC2 instance type changed"));
+		Assert.IsTrue(await AuditLogContains("AWS EC2 instance type changed", 1));
 	}
 
-	private Task FlushAuditLogsAsync()
+	private async Task FlushAuditLogsAsync(int waitForNumMessages = -1, int maxWaitMs = 3000)
 	{
 		IAuditLog<AgentId> auditLog = ServiceProvider.GetRequiredService<IAuditLog<AgentId>>();
 		AuditLog<AgentId> log = (AuditLog<AgentId>)auditLog;
-		return log.FlushMessagesInternalAsync();
+		using CancellationTokenSource cts = new (maxWaitMs);
+		
+		if (waitForNumMessages >= 1)
+		{
+			int numMessagesFlushed = 0;
+			for (;;)
+			{
+				numMessagesFlushed += await log.FlushMessagesInternalAsync();
+				if (numMessagesFlushed >= waitForNumMessages || cts.Token.IsCancellationRequested)
+				{
+					break;
+				}
+				
+				await Task.Delay(100, cts.Token);
+			}
+		}
+		else
+		{
+			await log.FlushMessagesInternalAsync();	
+		}
 	}
 	
 	[TestMethod]
diff --git a/Engine/Source/Programs/Horde/Horde.Build/Auditing/AuditLog.cs b/Engine/Source/Programs/Horde/Horde.Build/Auditing/AuditLog.cs
--- a/Engine/Source/Programs/Horde/Horde.Build/Auditing/AuditLog.cs
+++ b/Engine/Source/Programs/Horde/Horde.Build/Auditing/AuditLog.cs
@@ -193,7 +193,7 @@ namespace Horde.Build.Auditing
 		/// Flush any pending messages to database
 		/// Exposed as internal for use in tests
 		/// </summary>
-		internal async Task FlushMessagesInternalAsync()
+		internal async Task<int> FlushMessagesInternalAsync()
 		{
 			List<AuditLogMessage> newMessages = new ();
 			while (_messageChannel.Reader.TryRead(out AuditLogMessage? newMessage))
@@ -204,6 +204,8 @@ namespace Horde.Build.Auditing
 			{
 				await _messages.InsertManyAsync(newMessages);
 			}
+			
+			return newMessages.Count;
 		}
 		
 		async Task WriteMessagesAsync()
