Horde: Better logging of fleet scaling results

diff --git a/Engine/Source/Programs/Horde/Horde.Build/Agents/Fleet/FleetService.cs b/Engine/Source/Programs/Horde/Horde.Build/Agents/Fleet/FleetService.cs
--- a/Engine/Source/Programs/Horde/Horde.Build/Agents/Fleet/FleetService.cs
+++ b/Engine/Source/Programs/Horde/Horde.Build/Agents/Fleet/FleetService.cs
@@ -273,7 +273,8 @@ namespace Horde.Build.Agents.Fleet
 
 			if (pool.LastScaleResult == null || !pool.LastScaleResult.Equals(result))
 			{
-				_logger.LogInformation("Scale result: {Result}", result.Message);
+				_logger.LogInformation("Scale result: Outcome={Outcome} AgentsAdded={AgentsAdded} AgentsRemoved={AgentsRemoved} Message={Message}", 
+					result.Outcome, result.AgentsAddedCount, result.AgentsAddedCount, result.Message);
 			}
 
 			await _poolCollection.TryUpdateAsync(
diff --git a/Engine/Source/Programs/Horde/Horde.Build/Agents/Fleet/Providers/AwsRecyclingFleetManager.cs b/Engine/Source/Programs/Horde/Horde.Build/Agents/Fleet/Providers/AwsRecyclingFleetManager.cs
--- a/Engine/Source/Programs/Horde/Horde.Build/Agents/Fleet/Providers/AwsRecyclingFleetManager.cs
+++ b/Engine/Source/Programs/Horde/Horde.Build/Agents/Fleet/Providers/AwsRecyclingFleetManager.cs
@@ -466,8 +466,8 @@ public sealed class AwsRecyclingFleetManager : IFleetManager
 	public async Task<ScaleResult> ShrinkPoolAsync(IPool pool, IReadOnlyList<IAgent> agents, int count, CancellationToken cancellationToken)
 	{
 		await AwsFleetManager.ShrinkPoolViaAgentShutdownRequestAsync(_agentCollection, pool, agents, count, cancellationToken);
-		return new ScaleResult(FleetManagerOutcome.Success, 0, count);
-	} 
+		return new ScaleResult(FleetManagerOutcome.Success, 0, count, "Pool scaled in");
+	}
 
 	/// <inheritdoc/>
 	public async Task<int> GetNumStoppedInstancesAsync(IPool pool, CancellationToken cancellationToken)
