Horde: New step history panel
#skipci
#rnx

diff --git a/Engine/Source/Programs/Horde/HordeDashboard/src/components/ChangeButton.tsx b/Engine/Source/Programs/Horde/HordeDashboard/src/components/ChangeButton.tsx
--- a/Engine/Source/Programs/Horde/HordeDashboard/src/components/ChangeButton.tsx
+++ b/Engine/Source/Programs/Horde/HordeDashboard/src/components/ChangeButton.tsx
@@ -225,7 +225,7 @@ function getJobSummary(job: GetJobResponse): { text: string, color: string } {
 }
 
 
-export const ChangeButton: React.FC<{ job?: JobData, stepRef?: GetJobStepRefResponse, commit?: GetChangeSummaryResponse, hideAborted?: boolean, rangeCL?: number, pinned?: boolean }> = ({ job, stepRef, commit, hideAborted, rangeCL, pinned }) => {
+export const ChangeButton: React.FC<{ job?: JobData, stepRef?: GetJobStepRefResponse, commit?: GetChangeSummaryResponse, hideAborted?: boolean, rangeCL?: number, pinned?: boolean, prefix?: string }> = ({ job, stepRef, commit, hideAborted, rangeCL, pinned, prefix }) => {
 
    const [menuShown, setMenuShown] = useState(false);
 
@@ -247,6 +247,10 @@ export const ChangeButton: React.FC<{ job?: JobData, stepRef?: GetJobStepRefResp
    if (stepRef) {
       change = stepRef.change.toString();
    }
+
+   if (prefix) {
+      change = `${prefix} ${change}`;
+   }
    
    return (<Stack verticalAlign="center" verticalFill={true} horizontalAlign="start"> <div style={{ paddingBottom: "1px" }}>
       <Stack tokens={{ childrenGap: 4 }}>
diff --git a/Engine/Source/Programs/Horde/HordeDashboard/src/components/jobDetailsV2/JobDetailStepHistoryV2.tsx b/Engine/Source/Programs/Horde/HordeDashboard/src/components/jobDetailsV2/JobDetailStepHistoryV2.tsx
new file mode 100644
--- /dev/null
+++ b/Engine/Source/Programs/Horde/HordeDashboard/src/components/jobDetailsV2/JobDetailStepHistoryV2.tsx
@@ -0,0 +1,765 @@
+// Copyright Epic Games, Inc. All Rights Reserved.
+
+import { DefaultButton, Spinner, SpinnerSize, Stack, Text } from "@fluentui/react";
+import * as d3 from "d3";
+import { action, makeObservable, observable } from "mobx";
+import { observer } from "mobx-react-lite";
+import moment, { Moment } from "moment";
+import { useEffect, useState } from "react";
+import { Link } from "react-router-dom";
+import backend from "../../backend";
+import { GetJobStepRefResponse, JobStepOutcome } from "../../backend/Api";
+import dashboard, { StatusColor } from "../../backend/Dashboard";
+import { ISideRailLink } from "../../base/components/SideRail";
+import { displayTimeZone, getElapsedString, getHumanTime, msecToElapsed } from "../../base/utilities/timeUtils";
+import { hordeClasses, modeColors } from "../../styles/Styles";
+import { ChangeButton } from "../ChangeButton";
+import { HistoryModal } from "../HistoryModal";
+import { StepRefStatusIcon } from "../StatusIcon";
+import { JobDataView, JobDetailsV2 } from "./JobDetailsViewCommon";
+
+const sideRail: ISideRailLink = { text: "History", url: "rail_step_history" };
+
+class Tooltip {
+
+   constructor() {
+      makeObservable(this);
+   }
+
+   @observable
+   updated = 0
+
+   @action
+   update(ref?: GetJobStepRefResponse, x?: number, y?: number, change?: number) {
+      this.ref = ref;
+      this.x = x ?? 0;
+      this.y = y ?? 0;
+      this.change = change ?? 0;
+      this.updated++;
+   }
+
+   @action
+   freeze(frozen: boolean) {
+      this.frozen = frozen;
+      this.updated++;
+   }
+
+
+   x: number = 0;
+   y: number = 0;
+   change: number = 0;
+   frozen: boolean = false;
+   ref?: GetJobStepRefResponse;
+}
+
+class StepHistoryDataView extends JobDataView {
+
+   constructor(details: JobDetailsV2) {
+      super(details, false)
+      makeObservable(this);
+   }
+
+   filterUpdated() {
+      // this.updateReady();
+   }
+
+   initData(history: GetJobStepRefResponse[]) {
+
+      this.minTime = this.maxTime = undefined;
+      this.maxMinutes = 0;
+      history = history.filter(r => {
+
+         if (!r.startTime) {
+            return false;
+         }
+
+         if (!r.finishTime) {
+            const hours = moment.duration(moment().diff(moment(r.startTime))).asHours();
+            return hours <= 24;
+         }
+
+         return true;
+
+      });
+
+      // durations in minutes
+      const durations = this.durations;
+
+      durations.clear();
+
+      history.forEach(h => {
+
+         this.times.set(h.startTime as string, new Date(h.startTime));
+         if (h.finishTime) {
+            this.times.set(h.finishTime as string, new Date(h.finishTime));
+         }
+
+         let end: Moment;
+         if (h.finishTime) {
+            end = moment(h.finishTime);
+         } else {
+            end = moment();
+         }
+
+         const start = moment(h.startTime);
+
+         if (!this.minTime || this.minTime.getTime() > (start.unix() * 1000)) {
+            this.minTime = start.toDate();
+         }
+
+         if (!this.maxTime || this.maxTime.getTime() < (end.unix() * 1000)) {
+            this.maxTime = end.toDate();
+         }
+
+         const minutes = moment.duration(end.diff(start)).asMinutes();
+         this.maxMinutes = Math.max(minutes, this.maxMinutes);
+         durations.set(h.jobId, minutes);
+      });
+
+      history = history.sort((a, b) => this.times.get(b.startTime as string)!.getTime() - this.times.get(a.startTime as string)!.getTime())
+
+      this.history = history;
+   }
+
+   set(stepId?: string) {
+
+      const details = this.details;
+
+      if (!details) {
+         return;
+      }
+
+      if (this.stepId === stepId || !details.jobId) {
+         return;
+      }
+
+      this.stepId = stepId;
+
+      if (!this.stepId) {
+         return;
+      }
+
+      const jobData = details.jobData!;
+      const stepName = details.getStepName(this.stepId);
+
+      backend.getJobStepHistory(jobData.streamId, stepName, 1024, jobData.templateId!).then(response => {
+         this.initData(response);
+         this.updateReady();
+      }).finally(() => {
+         this.initialize(this.history?.length ? [sideRail] : undefined);
+      })
+   }
+
+   clear() {
+      this.history = [];
+      this.stepId = undefined;
+      this.durations.clear();
+      this.minTime = this.maxTime = undefined;
+      this.maxMinutes = 0;
+      this.renderer = undefined;
+      super.clear();
+   }
+
+   get lastFailure(): GetJobStepRefResponse | undefined {
+
+      return this.history.find(r => r.outcome === JobStepOutcome.Failure);
+   }
+
+   get lastSuccess(): GetJobStepRefResponse | undefined {
+
+      return this.history.find(r => r.outcome === JobStepOutcome.Success);
+   }
+
+   get current(): GetJobStepRefResponse | undefined {
+
+      return this.history.find(r => r.stepId === this.stepId);
+   }
+
+   detailsUpdated() {
+
+   }
+
+   @observable
+   selectedAgentId?: string;
+
+   @action
+   setSelectedAgentId(agentId?: string) {
+      if (this.selectedAgentId === agentId) {
+         return;
+      }
+
+      this.selectedAgentId = agentId;
+   }
+
+   history: GetJobStepRefResponse[] = [];
+   times = new Map<string, Date>();
+
+   stepId?: string;
+
+   order = 5;
+
+   durations = new Map<string, number>();
+   maxMinutes = 0;
+
+   minTime?: Date;
+   maxTime?: Date;
+
+
+   tooltip = new Tooltip();
+
+   renderer?: StepHistoryRenderer;
+}
+
+type SelectionType = d3.Selection<SVGGElement, unknown, null, undefined>;
+type Zoom = d3.ZoomBehavior<Element, unknown>;
+type Scalar = d3.ScaleLinear<number, number, never>;
+
+enum ZoomType {
+   Current,
+   Success,
+   Failure,
+   Latest
+}
+
+class StepHistoryRenderer {
+
+   // find max active and factor in wait time
+
+   constructor(dataView: StepHistoryDataView) {
+
+      this.dataView = dataView;
+
+      this.margin = { top: 16, right: 32, bottom: 0, left: 48 };
+
+   }
+
+   zoomTo(type: ZoomType) {
+
+      const dataView = this.dataView;
+
+      if (!dataView) {
+         return;
+      }
+
+      let ref: GetJobStepRefResponse | undefined;
+
+      if (type === ZoomType.Failure) {
+         ref = dataView.lastFailure;
+      }
+
+      if (type === ZoomType.Success) {
+         ref = dataView.lastSuccess;
+      }
+
+      if (type === ZoomType.Latest) {
+         ref = dataView.history[0];
+      }
+
+      if (type === ZoomType.Current) {
+         ref = dataView.current;
+      }
+
+
+      this.zoomToRef(ref);
+
+   }
+
+   zoomToRef(ref?: GetJobStepRefResponse, zoomLevel = 8) {
+
+      if (!ref) {
+         return;
+      }
+
+      const dataView = this.dataView;
+
+      if (!dataView) {
+         return;
+      }
+
+      const svg = this.svg;
+      const zoom = this.zoom;
+
+      if (!svg || !zoom) {
+         return;
+      }
+
+      //@fixme, width needs to be set elsewhere
+      const width = 1800;
+      const margin = this.margin;
+
+      // based on zoom of 8
+      const mysteriousZoomFactor = 221;
+      svg.call(zoom.transform as any, d3.zoomIdentity.scale(8));
+      let trans = this.scaleX!(new Date(ref.startTime!).getTime() / 1000) + 128;
+      trans /= 8;
+      if (trans > (width - margin.right)) {
+         trans = width - margin.right;
+      }
+      svg.call(zoom.transform as any, d3.zoomIdentity.scale(8).translate(-(trans - mysteriousZoomFactor), 0))
+
+      const posx = this.scaleX!(new Date(ref.startTime!).getTime() / 1000);
+      const posy = this.scaleY!(dataView.durations.get(ref.jobId)!);
+
+      dataView.tooltip.update(ref, posx - 100, 108, ref.change);
+      dataView.tooltip.freeze(true);
+
+      svg!.selectAll(".cline")
+         .attr("x1", () => posx)
+         .attr("x2", () => posx)
+         .attr("y2", d => posy - 8)
+   }
+
+   render(container: HTMLDivElement) {
+
+      if (this.hasRendered && !this.forceRender) {
+         return;
+      }
+
+      const dataView = this.dataView;
+      dataView.renderer = this;
+
+      this.hasRendered = true;
+
+      const width = 1800;
+      const margin = this.margin;
+
+      let svg = this.svg;
+
+      let height = 340;
+
+      const scolors = dashboard.getStatusColors();
+      const colors: Record<string, string> = {
+         "Success": scolors.get(StatusColor.Success)!,
+         "Failure": scolors.get(StatusColor.Failure)!,
+         "Warnings": scolors.get(StatusColor.Warnings)!,
+         "Unspecified": scolors.get(StatusColor.Skipped)!,
+      };
+
+      const x = this.scaleX = d3.scaleLinear()
+         .domain([dataView.minTime!, dataView.maxTime!].map(d => d.getTime() / 1000))
+         .range([margin.left, width - margin.right])
+
+      const y = this.scaleY = d3.scaleLinear()
+         .domain([0, dataView.maxMinutes]).nice()
+         .range([height - margin.bottom, margin.top])
+
+
+      if (!svg) {
+         svg = d3.select(container)
+            .append("svg") as any as SelectionType
+
+         this.svg = svg;
+      } else {
+         // remove tooltip
+         d3.select(container).selectAll('div').remove();
+         svg.selectAll("*").remove();
+      }
+
+      svg.attr("viewBox", [0, 0, width, height] as any);
+
+      svg.append("rect")
+         .attr("width", "100%")
+         .attr("height", "100%")
+         .attr("fill", modeColors.background);
+
+      const clipId = `step_history_clip`;
+
+      svg.append("clipPath")
+         .attr("id", clipId)
+         .append("rect")
+         .attr("x", margin.left)
+         .attr("y", 0)
+         .attr("width", width - margin.left - margin.right + 2)
+         .attr("height", height);
+
+      const data = dataView.history;
+
+      svg.append("g")
+         .attr("class", "bars")
+         .attr("clip-path", `url(#${clipId})`)
+         .selectAll("rect")
+         .data(data)
+         .join("rect")
+         .attr("fill", (d) => {
+            if (!d.finishTime) {
+               return scolors.get(StatusColor.Running)!;
+            }
+            if (!d.outcome) {
+               return colors["Unspecified"];
+            }
+            return colors[d.outcome];
+         })
+         .attr("x", d => x(new Date(d.startTime!).getTime() / 1000) + 8)
+         .attr("y", d => y(dataView.durations.get(d.jobId)!))
+         .attr("height", d => y(0) - y(dataView.durations.get(d.jobId)!))
+         .attr("width", 2);
+
+      const arrowPoints = [[0, 0], [0, 10], [10, 5]] as any;
+      svg.append("marker")
+         .attr("id", "cmarker")
+         .attr('viewBox', [0, 0, 10, 10])
+         .attr("refX", 5)
+         .attr("refY", 5)
+         .attr("markerWidth", 5)
+         .attr("markerHeight", 5)
+         .attr("orient", 'auto-start-reverse')
+         .append("path")
+         .attr('d', d3.line()(arrowPoints))
+         .style("fill", dashboard.darktheme ? "#6D6C6B" : "#4D4C4B");
+
+      //svg.append("g")
+      //  .attr("class", "cline")
+      svg.append("line")
+         .attr("class", "cline")
+         .attr("x1", () => 0)
+         .attr("x2", () => 0)
+         .attr("y1", () => 0)
+         .attr("y2", () => 0)
+         .attr("stroke-width", () => 2)
+         .style("stroke-dasharray", "3, 3")
+         .attr("marker-end", "url(#cmarker)")
+         .attr("stroke", dashboard.darktheme ? "#6D6C6B" : "#4D4C4B")
+
+      // left axis
+      svg.append("g")
+         .attr("transform", `translate(${this.margin.left},0)`)
+         .style("font-family", "Horde Open Sans Regular")
+         .style("font-size", "9px")
+         .call(d3.axisLeft(y)
+            .ticks(7)
+            .tickFormat((d) => {
+
+               if (!d) {
+                  return "";
+               }
+
+               return msecToElapsed((d as number) * 60000, true, false);
+
+            }))
+      
+      svg.append("line")
+         .attr("class", "bline")         
+         .attr("x1", () => margin.left)
+         .attr("x2", () => width - margin.right) 
+         .attr("y1", () => height)
+         .attr("y2", () => height)
+         .attr("stroke-width", () => 1)
+         .attr("stroke", () => dashboard.darktheme ? "#6D6C6B" : "#4D4C4B")
+      
+
+      const xAxis = (g: SelectionType) => {
+
+         g.attr("transform", `translate(0,32)`)
+            .style("font-family", "Horde Open Sans Regular")
+            .style("font-size", "9px")
+            .call(d3.axisTop(x)
+               .tickFormat(d => {
+                  return getHumanTime(new Date((d as number) * 1000));
+               })
+               .tickSizeOuter(0))
+            .call(g => g.select(".domain").remove())
+            .call(g => g.selectAll(".tick line").attr("stroke-opacity", dashboard.darktheme ? 0.35 : 0.25)
+               .attr("stroke", dashboard.darktheme ? "#6D6C6B" : "#4D4C4B")
+               .attr("y2", height - this.margin.bottom))
+      }
+
+      // top axis
+      svg.append("g").attr("class", "x-axis").call(xAxis)
+
+      const zoom = this.zoom = d3.zoom()
+         .scaleExtent([1, 8])
+         .extent([[margin.left, 0], [width - margin.right, height]])
+         .translateExtent([[margin.left, -Infinity], [width - margin.right, Infinity]])
+         .on("zoom", zoomed);
+
+      function zoomed(event: any) {
+         dataView.tooltip.update(undefined)
+         dataView.tooltip.freeze(false);
+
+         const zoomLevel = event.transform.k ? event.transform.k : 1;
+         const barWidth = Math.max(2, zoomLevel / 2);
+
+         x.range([margin.left, width - margin.right].map(d => event.transform.applyX(d)));
+         svg!.selectAll(".bars rect")
+            .attr("x", (d: any) => x(new Date(d.startTime!).getTime() / 1000) - (barWidth / 2))
+            .attr("width", barWidth);
+
+         svg!.selectAll(".x-axis").call(xAxis as any);
+      }
+
+      svg.call(zoom as any);
+
+      this.zoomToRef(dataView.current ?? dataView.history[0]);
+
+      const closestData = (xpos: number, ypos: number): GetJobStepRefResponse | undefined => {
+
+         if (!dataView.history.length) {
+            return undefined;
+         }
+
+         let first = dataView.history[0];
+
+         let closest = dataView.history.reduce<GetJobStepRefResponse>((best: GetJobStepRefResponse, ref: GetJobStepRefResponse) => {
+
+            const bestx = x(dataView.times.get(best.startTime as string)!.getTime() / 1000);
+            const refx = x(dataView.times.get(ref.startTime as string)!.getTime() / 1000);
+
+            if (Math.abs(xpos - bestx) > Math.abs(xpos - refx)) {
+               return ref
+            }
+
+            return best;
+
+         }, first);
+
+         return closest;
+
+      }
+
+      const handleMouseMove = (event: any) => {
+
+         if (dataView.tooltip.frozen) {
+            return;
+         }
+
+         let mouseX = d3.pointer(event)[0];
+         let mouseY = d3.pointer(event)[1];
+
+         const ref = closestData(mouseX, mouseY);
+
+         if (ref) {
+            const posx = x(new Date(ref.startTime!).getTime() / 1000);
+            const posy = y(dataView.durations.get(ref.jobId)!);
+
+            svg!.selectAll(".cline")
+               .attr("x1", () => posx)
+               .attr("x2", () => posx)
+               .attr("y1", d => 0)
+               .attr("y2", d => posy - 8)
+
+         } else {
+            //svg!.selectAll(".cline rect")
+            //.attr("width", 0);
+         }
+
+         // relative to container
+         if (ref) {
+            this.showToolTip()
+         }
+         dataView.tooltip.update(ref, d3.pointer(event, container)[0], mouseY, ref?.change);
+      }      
+
+      // events
+      svg.on("wheel", (event) => { event.preventDefault(); })
+
+      svg.on("mousemove", (event) => { this.showToolTip(true); handleMouseMove(event);  });
+      svg.on("mouseleave", (event) => { if (!dataView.tooltip.frozen) dataView.tooltip.update(undefined);})
+
+      svg.on("pointerdown", (event) => {  });
+      svg.on("pointerup", (event) => { });
+
+      svg.on("click", (event) => {                  
+         dataView.tooltip.freeze(!dataView.tooltip.frozen);
+      });
+
+   }
+
+   showToolTip(shown: boolean = true) {
+      const svg = this.svg;
+      if (!svg) {
+         return;
+      }
+
+      svg!.selectAll(".cline")
+         .style("display", shown ? "block" : "none")
+
+
+   }
+
+   margin: { top: number, right: number, bottom: number, left: number }
+
+   dataView: StepHistoryDataView;
+
+   svg?: SelectionType;
+
+   zoom?: Zoom;
+   scaleX?: Scalar;
+   scaleY?: Scalar;
+
+   hasRendered = false;
+   forceRender = false;
+}
+
+const GraphTooltip: React.FC<{ dataView: StepHistoryDataView }> = observer(({ dataView }) => {
+
+   // subscribe
+   if (dataView.tooltip.updated) { }
+
+   const tooltip = dataView.tooltip;
+   const ref = tooltip.ref;
+   const renderer = dataView.renderer;
+
+   if (!ref) {
+      if (renderer) {
+         renderer.showToolTip(false);
+      }
+      return null;
+   }
+
+   renderer?.showToolTip(true);
+
+   let tipX = tooltip.x;
+   let offsetX = 32;
+   let translateX = "0%";
+
+   if (tipX > 1500) {
+      offsetX = -32;
+      translateX = "-100%";
+   }
+
+   const translateY = "-50%";
+
+   const sindex = dataView.history.indexOf(ref);
+
+   const start = moment(ref.startTime);
+   let end = moment(Date.now());
+
+   if (ref.finishTime) {
+      end = moment(ref.finishTime);
+   }
+
+   const textSize = "small";
+
+   const duration = getElapsedString(start, end);
+
+   const displayTime = moment(ref.startTime).tz(displayTimeZone());
+   const format = dashboard.display24HourClock ? "HH:mm:ss z" : "LT z";
+   let displayTimeStr = displayTime.format('MMM Do') + ` at ${displayTime.format(format)}`;
+
+   return <div style={{
+      position: "absolute",
+      display: "block",
+      top: `${tooltip.y}px`,
+      left: `${tooltip.x + offsetX}px`,
+      backgroundColor: modeColors.background,
+      zIndex: 1,
+      border: "solid",
+      borderWidth: "1px",
+      borderRadius: "3px",
+      width: "max-content",
+      borderColor: dashboard.darktheme ? "#413F3D" : "#2D3F5F",
+      padding: 16,
+      pointerEvents: tooltip.frozen ? undefined : "none",
+      transform: `translate(${translateX}, ${translateY})`
+   }}>
+      <Stack>
+         <Link to={`/job/${ref.jobId}?step=${ref.stepId}`}><Stack horizontal>
+            <StepRefStatusIcon stepRef={ref} />
+            <Text variant={textSize}>{dataView.details?.nodeByStepId(dataView.stepId)?.name}</Text>
+         </Stack>
+         </Link>
+         <Stack style={{ paddingLeft: 2, paddingTop: 8 }} tokens={{ childrenGap: 8 }}>
+            <Stack>
+               <ChangeButton prefix="CL" job={dataView.details!.jobData!} stepRef={ref} hideAborted={true} rangeCL={sindex < (dataView.history.length - 1) ? dataView.history[sindex + 1].change : undefined} />
+            </Stack>
+            <Stack><Text variant={textSize}>{displayTimeStr}</Text></Stack>
+            <Stack><Text variant={textSize}>Duration: {duration}</Text></Stack>
+            <Stack style={{ cursor: "pointer" }} onClick={() => dataView.setSelectedAgentId(ref.agentId)}>
+               <Text variant={textSize}>{ref.agentId!}</Text>
+            </Stack>
+         </Stack>
+      </Stack>
+   </div>
+})
+
+const StepHistoryGraph: React.FC<{ dataView: StepHistoryDataView }> = ({ dataView }) => {
+
+   const graph_container_id = `timeline_graph_container`;
+
+   const [container, setContainer] = useState<HTMLDivElement | null>(null);
+   const [state, setState] = useState<{ graph?: StepHistoryRenderer }>({});
+
+   if (!state.graph) {
+      setState({ ...state, graph: new StepHistoryRenderer(dataView) })
+      return null;
+   }
+
+   if (container) {
+      try {
+         state.graph?.render(container);
+
+      } catch (err) {
+         console.error(err);
+      }
+
+   }
+
+   return <Stack className={hordeClasses.horde}>
+      <Stack style={{ paddingLeft: 8, paddingTop: 8 }}>
+         <div style={{ position: "relative" }}>
+            <GraphTooltip dataView={dataView} />
+         </div>
+         <Stack style={{paddingTop: 16,paddingBottom: 16,paddingLeft: 16, backgroundColor: modeColors.background}}>
+            <div id={graph_container_id} className="horde-no-darktheme" style={{ shapeRendering: "crispEdges", userSelect: "none", position: "relative" }} ref={(ref: HTMLDivElement) => setContainer(ref)} onMouseEnter={() => { }} onMouseLeave={() => { }} />
+         </Stack>
+      </Stack>
+   </Stack>;
+
+}
+
+
+JobDetailsV2.registerDataView("StepHistoryDataView", (details: JobDetailsV2) => new StepHistoryDataView(details));
+
+export const StepHistoryPanelV2: React.FC<{ jobDetails: JobDetailsV2; stepId: string }> = observer(({ jobDetails, stepId }) => {
+
+   const dataView = jobDetails.getDataView<StepHistoryDataView>("StepHistoryDataView");
+
+   useEffect(() => {
+      return () => {
+         dataView?.clear();
+      };
+   }, [dataView]);
+
+   dataView.subscribe();
+
+   if (!jobDetails.jobData) {
+      return null;
+   }
+
+
+
+   dataView.set(stepId);
+
+
+   return (<Stack id={sideRail.url} styles={{ root: { paddingTop: 18, paddingRight: 12 } }}>
+      <Stack className={hordeClasses.raised}>
+         {!!dataView.selectedAgentId && <HistoryModal agentId={dataView.selectedAgentId} onDismiss={() => dataView.setSelectedAgentId(undefined)} />}
+         <Stack tokens={{ childrenGap: 12 }}>
+            <Stack horizontal>
+               <Stack>
+                  <Text variant="mediumPlus" styles={{ root: { fontFamily: "Horde Open Sans SemiBold" } }}>{"History"}</Text>
+               </Stack>
+               <Stack grow />
+               <Stack horizontal tokens={{ childrenGap: 12 }}>
+                  {!!dataView.current && <Stack>
+                     <DefaultButton text="Current" onClick={() => dataView.renderer?.zoomTo(ZoomType.Current)} />
+                  </Stack>}
+                  {!!dataView.history[0] && <Stack>
+                     <DefaultButton text="Latest" onClick={() => dataView.renderer?.zoomTo(ZoomType.Latest)} />
+                  </Stack>}
+                  {!!dataView.lastSuccess && <Stack>
+                     <DefaultButton text="Last Success" onClick={() => dataView.renderer?.zoomTo(ZoomType.Success)} />
+                  </Stack>}
+                  {!!dataView.lastFailure && <Stack>
+                     <DefaultButton text="Last Failure" onClick={() => dataView.renderer?.zoomTo(ZoomType.Failure)} />
+                  </Stack>}
+               </Stack>
+               <Stack>
+               </Stack>
+            </Stack>
+            {!!dataView.history.length && <StepHistoryGraph dataView={dataView} />}
+            {!dataView.history.length && <Spinner size={SpinnerSize.large} />}
+         </Stack>
+      </Stack>
+   </Stack>);
+});
\ No newline at end of file
diff --git a/Engine/Source/Programs/Horde/HordeDashboard/src/components/jobDetailsV2/JobDetailsViewCommon.tsx b/Engine/Source/Programs/Horde/HordeDashboard/src/components/jobDetailsV2/JobDetailsViewCommon.tsx
--- a/Engine/Source/Programs/Horde/HordeDashboard/src/components/jobDetailsV2/JobDetailsViewCommon.tsx
+++ b/Engine/Source/Programs/Horde/HordeDashboard/src/components/jobDetailsV2/JobDetailsViewCommon.tsx
@@ -17,8 +17,10 @@ import { LabelStatusIcon, StepStatusIcon } from "../StatusIcon";
 
 export abstract class JobDataView {
 
-   constructor(details: JobDetailsV2) {
-      makeObservable(this);
+   constructor(details: JobDetailsV2, initObservable?:boolean) {
+      if (initObservable === undefined || initObservable) {
+         makeObservable(this);   
+      }      
       this.details = details;
    }
 
