Horde: Actually enable the new step history view
#rnx
#skipci

diff --git a/Engine/Source/Programs/Horde/HordeDashboard/src/components/jobDetailsV2/StepDetailView.tsx b/Engine/Source/Programs/Horde/HordeDashboard/src/components/jobDetailsV2/StepDetailView.tsx
--- a/Engine/Source/Programs/Horde/HordeDashboard/src/components/jobDetailsV2/StepDetailView.tsx
+++ b/Engine/Source/Programs/Horde/HordeDashboard/src/components/jobDetailsV2/StepDetailView.tsx
@@ -9,7 +9,7 @@ import { HistoryModal } from "../HistoryModal";
 import { useQuery } from "../JobDetailCommon";
 import { JobDetailArtifactsV2 } from "./JobDetailArtifactsV2";
 import { HealthPanel } from "./JobDetailHealthV2";
-import { StepHistoryPanel } from "./JobDetailStepHistory";
+import { StepHistoryPanelV2 } from "./JobDetailStepHistoryV2";
 import { StepTrendsPanel } from "./JobDetailStepTrends";
 import { JobDataView, JobDetailsV2 } from "./JobDetailsViewCommon";
 import { TimelinePanel } from "./JobDetailTimeline";
@@ -84,7 +84,7 @@ const StepDetailViewInner: React.FC<{ jobDetails: JobDetailsV2, stepId: string }
          <StepTestReportPanel jobDetails={jobDetails} stepId={stepId} />
       </Stack>
       {<Stack>
-         <StepHistoryPanel jobDetails={jobDetails} stepId={stepId} />
+         <StepHistoryPanelV2 jobDetails={jobDetails} stepId={stepId} />
       </Stack>}
       <Stack>
          <JobDetailArtifactsV2 jobDetails={jobDetails} stepId={stepId} />
