UnrealAutomatedTests - Avoid displaying missing reference when image comparison was successfull

diff --git a/Engine/Source/Programs/Horde/HordeDashboard/src/Plugins/UnrealAutomatedTests/components/UnrealAutomatedTests.tsx b/Engine/Source/Programs/Horde/HordeDashboard/src/Plugins/UnrealAutomatedTests/components/UnrealAutomatedTests.tsx
--- a/Engine/Source/Programs/Horde/HordeDashboard/src/Plugins/UnrealAutomatedTests/components/UnrealAutomatedTests.tsx
+++ b/Engine/Source/Programs/Horde/HordeDashboard/src/Plugins/UnrealAutomatedTests/components/UnrealAutomatedTests.tsx
@@ -223,14 +223,18 @@ const EntryPane: React.FC<{ entry: TestEntry, testArtifacts: TestEntryArtifact[]
             <Stack styles={{ root: { paddingLeft: 16 } }}>
                <Stack><Text variant="medium" styles={{ root: { fontWeight: "bold" } }}>Image comparison: {artifact?.Name}</Text></Stack>
                <Stack horizontal>
+               {(is_error || is_warning) &&
                   <Stack styles={{ root: { padding: 5 } }}>
                      <a href={imageLinks.approved}><Image width={400} src={imageLinks.approved || missingImage} alt={artifact?.Files.Approved} /></a>
                      <Stack.Item align="center">Approved{!imageLinks.approved && MissingImageLabel()}</Stack.Item>
                   </Stack>
+               }
+               {(is_error || is_warning) &&
                   <Stack styles={{ root: { padding: 5 } }}>
                      <a href={imageLinks.difference}><Image width={400} src={imageLinks.difference || missingImage} alt={artifact?.Files.Difference} /></a>
                      <Stack.Item align="center">Difference{!imageLinks.difference && MissingImageLabel()}</Stack.Item>
                   </Stack>
+               }
                   <Stack styles={{ root: { padding: 5 } }}>
                      <a href={imageLinks.unapproved}><Image width={400} src={imageLinks.unapproved || missingImage} alt={artifact?.Files.Unapproved} /></a>
                      <Stack.Item align="center">Unapproved{!imageLinks.unapproved && MissingImageLabel()}</Stack.Item>
