Horde: Always return fleet scaling results with a message

diff --git a/Engine/Source/Programs/Horde/Horde.Build/Agents/Fleet/FleetService.cs b/Engine/Source/Programs/Horde/Horde.Build/Agents/Fleet/FleetService.cs
--- a/Engine/Source/Programs/Horde/Horde.Build/Agents/Fleet/FleetService.cs
+++ b/Engine/Source/Programs/Horde/Horde.Build/Agents/Fleet/FleetService.cs
@@ -203,7 +203,6 @@ namespace Horde.Build.Agents.Fleet
 				return new ScaleResult(FleetManagerOutcome.NoOp, 0, 0, "Auto-scaling disabled");
 			}
 
-			ScaleResult result = new (FleetManagerOutcome.NoOp, 0, 0);
 			int currentAgentCount = poolSizeResult.CurrentAgentCount;
 			int desiredAgentCount = poolSizeResult.DesiredAgentCount;
 			int deltaAgentCount = desiredAgentCount - currentAgentCount;
@@ -251,18 +250,40 @@ namespace Horde.Build.Agents.Fleet
 			
 			scope.Span.SetTag("IsDowntimeActive", _downtimeService.IsDowntimeActive);
 			
+			ScaleResult result = new (FleetManagerOutcome.NoOp, 0, 0);
 			try
 			{
-				if (deltaAgentCount > 0 && !_downtimeService.IsDowntimeActive && !isScaleOutCoolingDown)
+				if (_downtimeService.IsDowntimeActive)
+				{
+					result = new (FleetManagerOutcome.NoOp, 0, 0, "Downtime is active");
+				}
+				else
 				{
-					result = await fleetManager.ExpandPoolAsync(pool, agents, deltaAgentCount, cancellationToken);
-					scaleOutTime = _clock.UtcNow;
+					if (deltaAgentCount > 0)
+					{
+						if (isScaleOutCoolingDown)
+						{
+							result = new(FleetManagerOutcome.NoOp, 0, 0, "Cannot scale out, cooldown active");
+						}
+						else
+						{
+							result = await fleetManager.ExpandPoolAsync(pool, agents, deltaAgentCount, cancellationToken);
+							scaleOutTime = _clock.UtcNow;
+						}
+					}
 				}
 
-				if (deltaAgentCount < 0 && !isScaleInCoolingDown)
+				if (deltaAgentCount < 0)
 				{
-					result = await fleetManager.ShrinkPoolAsync(pool, agents, -deltaAgentCount, cancellationToken);
-					scaleInTime = _clock.UtcNow;
+					if (isScaleInCoolingDown)
+					{
+						result = new (FleetManagerOutcome.NoOp, 0, 0, "Cannot scale in, cooldown active");
+					}
+					else
+					{
+						result = await fleetManager.ShrinkPoolAsync(pool, agents, -deltaAgentCount, cancellationToken);
+						scaleInTime = _clock.UtcNow;	
+					}
 				}
 			}
 			catch (Exception ex)
