Horde: Shorten verbose log message for fleet scaling errors

Having a too long log body made the log entry not propagate.

diff --git a/Engine/Source/Programs/Horde/Horde.Build/Agents/Fleet/FleetService.cs b/Engine/Source/Programs/Horde/Horde.Build/Agents/Fleet/FleetService.cs
--- a/Engine/Source/Programs/Horde/Horde.Build/Agents/Fleet/FleetService.cs
+++ b/Engine/Source/Programs/Horde/Horde.Build/Agents/Fleet/FleetService.cs
@@ -253,27 +253,23 @@ namespace Horde.Build.Agents.Fleet
 			ScaleResult result = new (FleetManagerOutcome.NoOp, 0, 0);
 			try
 			{
-				if (_downtimeService.IsDowntimeActive)
+				if (deltaAgentCount > 0)
 				{
-					result = new (FleetManagerOutcome.NoOp, 0, 0, "Downtime is active");
-				}
-				else
-				{
-					if (deltaAgentCount > 0)
+					if (_downtimeService.IsDowntimeActive)
 					{
-						if (isScaleOutCoolingDown)
-						{
-							result = new(FleetManagerOutcome.NoOp, 0, 0, "Cannot scale out, cooldown active");
-						}
-						else
-						{
-							result = await fleetManager.ExpandPoolAsync(pool, agents, deltaAgentCount, cancellationToken);
-							scaleOutTime = _clock.UtcNow;
-						}
+						result = new (FleetManagerOutcome.NoOp, 0, 0, "Downtime is active");
+					}
+					else if (isScaleOutCoolingDown)
+					{
+						result = new(FleetManagerOutcome.NoOp, 0, 0, "Cannot scale out, cooldown active");
+					}
+					else
+					{
+						result = await fleetManager.ExpandPoolAsync(pool, agents, deltaAgentCount, cancellationToken);
+						scaleOutTime = _clock.UtcNow;
 					}
 				}
-
-				if (deltaAgentCount < 0)
+				else if (deltaAgentCount < 0)
 				{
 					if (isScaleInCoolingDown)
 					{
@@ -288,7 +284,7 @@ namespace Horde.Build.Agents.Fleet
 			}
 			catch (Exception ex)
 			{
-				_logger.LogInformation(ex, "Failed to scale {PoolName}:\n{Exception}", pool.Name, ex);
+				_logger.LogInformation(ex, "Failed to scale {PoolName}", pool.Name);
 				return new ScaleResult(FleetManagerOutcome.Failure, 0, 0);
 			}
 
