Horde: Avoid stopping instances with an empty list

Calling EC2 StopInstances with an empty list is not a valid call and will result in an exception. Also move log metadata to tracing span.

diff --git a/Engine/Source/Programs/Horde/Horde.Build/Agents/Fleet/Providers/AwsRecyclingFleetManager.cs b/Engine/Source/Programs/Horde/Horde.Build/Agents/Fleet/Providers/AwsRecyclingFleetManager.cs
--- a/Engine/Source/Programs/Horde/Horde.Build/Agents/Fleet/Providers/AwsRecyclingFleetManager.cs
+++ b/Engine/Source/Programs/Horde/Horde.Build/Agents/Fleet/Providers/AwsRecyclingFleetManager.cs
@@ -127,48 +127,20 @@ public sealed class AwsRecyclingFleetManager : IFleetManager
 		List<InstanceType>? instanceTypePriority = Settings.InstanceTypes?.Select(InstanceType.FindValue).ToList();
 		int instancesToStartCount = requestCountPerAz.Values.Sum(x => x);
 		
-		Dictionary<string, object?> logScopeMetadata = new()
-		{
-			["InstancesToStartCount"] = instancesToStartCount,
-			["RequestedInstancesCount"] = requestedInstancesCount,
-			["CurrentAgentCount"] = agents.Count,
-			["StoppedInstancesMissingCount"] = stoppedInstancesMissingCount,
-			["CandidatesPerAz"] = candidatesPerAz.ToDictionary(x => x.Key, x => x.Value.Select(y => y.InstanceId)),
-			["RequestCountPerAz"] = requestCountPerAz,
-			["StoppedInstancesPerAz"] = stoppedInstancesPerAz.ToDictionary(x => x.Key, x => x.Value.Select(y => y.InstanceId)),
-			["InstanceTypePriority"] = Settings.InstanceTypes,
-		};
-		
+		scope.Span.SetTag("InstancesToStartCount", instancesToStartCount);
+		scope.Span.SetTag("RequestedInstancesCount", requestedInstancesCount);
+		scope.Span.SetTag("CurrentAgentCount", agents.Count);
+		scope.Span.SetTag("StoppedInstancesMissingCount", stoppedInstancesMissingCount);
 		scope.Span.SetTag("InstancesToStartCount", instancesToStartCount);
 		scope.Span.SetTag("StoppedInstancesMissingCount", stoppedInstancesMissingCount);
 
-		ScaleResult result;
-		using (_logger.BeginScope(logScopeMetadata))
+		if (instancesToStartCount == 0)
 		{
-			if (instancesToStartCount == 0)
-			{
-				result = new ScaleResult(
-					FleetManagerOutcome.Failure,
-					instancesToStartCount,
-					0,
-					"Unable to start any instance(s). No stopped instances are available");
-			}
-			else if (stoppedInstancesMissingCount > 0)
-			{
-				result = new ScaleResult(
-					FleetManagerOutcome.PartialSuccess,
-					instancesToStartCount,
-					0,
-					$"Starting {instancesToStartCount} instance(s) but not enough stopped instances to accommodate the full pool scale-out");
-			}
-			else
-			{
-				result = new ScaleResult(
-					FleetManagerOutcome.Success,
-					instancesToStartCount,
-					0,
-					$"Starting {instancesToStartCount} instance(s)");
-			}
+			return new ScaleResult(
+				FleetManagerOutcome.Failure,
+				0,
+				0,
+				"Unable to start any instance(s). No stopped instances are available");
 		}
 
 		foreach ((string az, List<Instance> instances) in stoppedInstancesPerAz)
@@ -180,7 +152,20 @@ public sealed class AwsRecyclingFleetManager : IFleetManager
 			}
 		}
 
-		return result;
+		if (stoppedInstancesMissingCount > 0)
+		{
+			return new ScaleResult(
+				FleetManagerOutcome.PartialSuccess,
+				instancesToStartCount,
+				0,
+				$"Starting {instancesToStartCount} instance(s) but not enough stopped instances to accommodate the full pool scale-out");
+		}
+		
+		return new ScaleResult(
+			FleetManagerOutcome.Success,
+			instancesToStartCount,
+			0,
+			$"Starting {instancesToStartCount} instance(s)");
 	}
 
 	private static Dictionary<string, List<Instance>> GetInstancesToLaunch(
@@ -424,10 +409,13 @@ public sealed class AwsRecyclingFleetManager : IFleetManager
 			.Select(x => x.InstanceId)
 			.ToList();
 
-		StopInstancesRequest stopRequest = new() { InstanceIds = stuckInstanceIds };
-		StopInstancesResponse stopResponse = await _ec2.StopInstancesAsync(stopRequest, cancellationToken);
-		scope.Span.SetTag("stopRes.statusCode", (int)stopResponse.HttpStatusCode);
-		scope.Span.SetTag("stopRes.numStoppingInstances", stopResponse.StoppingInstances.Count);
+		if (stuckInstanceIds.Count > 0)
+		{
+			StopInstancesRequest stopRequest = new() { InstanceIds = stuckInstanceIds };
+			StopInstancesResponse stopResponse = await _ec2.StopInstancesAsync(stopRequest, cancellationToken);
+			scope.Span.SetTag("stopRes.statusCode", (int)stopResponse.HttpStatusCode);
+			scope.Span.SetTag("stopRes.numStoppingInstances", stopResponse.StoppingInstances.Count);
+		}
 	}
 	
 	/// <summary>
