Horde: Fix overflow on job details view
#skipci
#rnx

diff --git a/Engine/Source/Programs/Horde/HordeDashboard/src/components/jobDetailsV2/JobDetailViewV2.tsx b/Engine/Source/Programs/Horde/HordeDashboard/src/components/jobDetailsV2/JobDetailViewV2.tsx
--- a/Engine/Source/Programs/Horde/HordeDashboard/src/components/jobDetailsV2/JobDetailViewV2.tsx
+++ b/Engine/Source/Programs/Horde/HordeDashboard/src/components/jobDetailsV2/JobDetailViewV2.tsx
@@ -368,7 +368,7 @@ const DetailsView: React.FC<{ jobDetails: JobDetailsV2 }> = ({ jobDetails }) =>
             <Stack style={{ width: "100%", backgroundColor: modeColors.background }}>
                <Stack horizontal>
                   <Stack /*className={classNames.pointerSuppress} */ style={{ position: "relative", width: "100%", height: 'calc(100vh - 228px)' }}>
-                     <div id="hordeContentArea" ref={scrollRef} style={{ overflowX: "hidden", overflowY: "visible" }}>
+                     <div id="hordeContentArea" ref={scrollRef} style={{ overflowX: "auto", overflowY: "visible" }}>
                         <Stack horizontal style={{paddingBottom: 48}}>
                            <Stack key={`${key}_2`} style={{ paddingLeft: centerAlign }} />
                            <Stack style={{ width: rootWidth }}>
