Horde: Better information density in job details view
#skipci
#rnx

diff --git a/Engine/Source/Programs/Horde/HordeDashboard/src/components/jobDetailsV2/JobDetailHealthV2.tsx b/Engine/Source/Programs/Horde/HordeDashboard/src/components/jobDetailsV2/JobDetailHealthV2.tsx
--- a/Engine/Source/Programs/Horde/HordeDashboard/src/components/jobDetailsV2/JobDetailHealthV2.tsx
+++ b/Engine/Source/Programs/Horde/HordeDashboard/src/components/jobDetailsV2/JobDetailHealthV2.tsx
@@ -201,13 +201,13 @@ export const HealthPanel: React.FC<{ jobDetails: JobDetailsV2 }> = observer(({ j
    });
 
    const columns = [
-      { key: 'health_column1', name: 'Summary', minWidth: 500, maxWidth: 500, isResizable: false },
+      { key: 'health_column1', name: 'Summary', minWidth: 460, maxWidth: 460, isResizable: false },
       { key: 'health_column2', name: 'Quarantine', minWidth: 80, maxWidth: 80, isResizable: false },
       { key: 'health_column3', name: 'Jira', minWidth: 80, maxWidth: 80, isResizable: false },
       { key: 'health_column4', name: 'JiraPriority', minWidth: 64, maxWidth: 64, isResizable: false },
-      { key: 'health_column5', name: 'JiraAssignee', minWidth: 240, maxWidth: 240, isResizable: false },
+      { key: 'health_column5', name: 'JiraAssignee', minWidth: 180, maxWidth: 180, isResizable: false },
       { key: 'health_column6', name: 'JiraStatus', minWidth: 64, maxWidth: 64, isResizable: false },
-      { key: 'health_column7', name: 'Status', minWidth: 160, maxWidth: 160, isResizable: false },
+      { key: 'health_column7', name: 'Status', minWidth: 140, maxWidth: 140, isResizable: false },
       { key: 'health_column8', name: 'Opened', minWidth: 140, maxWidth: 150, isResizable: false },
    ];
 
diff --git a/Engine/Source/Programs/Horde/HordeDashboard/src/components/jobDetailsV2/JobDetailStepTrends.tsx b/Engine/Source/Programs/Horde/HordeDashboard/src/components/jobDetailsV2/JobDetailStepTrends.tsx
--- a/Engine/Source/Programs/Horde/HordeDashboard/src/components/jobDetailsV2/JobDetailStepTrends.tsx
+++ b/Engine/Source/Programs/Horde/HordeDashboard/src/components/jobDetailsV2/JobDetailStepTrends.tsx
@@ -881,7 +881,7 @@ const StepTrendsPanelInner: React.FC<{ jobDetails: JobDetailsV2; dataView: Trend
             <Stack styles={{ root: { position: 'relative', top: -10 } }}>
                {statusDiv}
             </Stack>
-            <Stack styles={{ root: { width: 1720, minHeight: 400, height: 400 } }} onWheel={onChartScroll} onMouseEnter={onChartWheelEnter} onMouseLeave={onChartWheelLeave}>
+            <Stack styles={{ root: { width: 1370, minHeight: 400, height: 400 } }} onWheel={onChartScroll} onMouseEnter={onChartWheelEnter} onMouseLeave={onChartWheelLeave}>
                {placeholderDiv}
             </Stack>
 
diff --git a/Engine/Source/Programs/Horde/HordeDashboard/src/components/jobDetailsV2/JobDetailTimeline.tsx b/Engine/Source/Programs/Horde/HordeDashboard/src/components/jobDetailsV2/JobDetailTimeline.tsx
--- a/Engine/Source/Programs/Horde/HordeDashboard/src/components/jobDetailsV2/JobDetailTimeline.tsx
+++ b/Engine/Source/Programs/Horde/HordeDashboard/src/components/jobDetailsV2/JobDetailTimeline.tsx
@@ -583,7 +583,7 @@ class TimelineRenderer {
 
       const handleMouseMove = (event: any) => {
 
-         let mouseX = d3.pointer(event)[0];
+         let mouseX = d3.pointer(event, container)[0];
          let mouseY = d3.pointer(event)[1];
 
          const span = closestData(mouseX, mouseY);
@@ -636,8 +636,10 @@ class TimelineRenderer {
             let tipX = mouseX;
             let translateX = "0%";
 
-            if (tipX > 1500) {
-               translateX = "-180%";
+            if (tipX > 1100) {
+               translateX = "-120%";
+            } else {
+               tipX += 48;
             }
 
             tooltip
diff --git a/Engine/Source/Programs/Horde/HordeDashboard/src/components/jobDetailsV2/JobDetailViewSteps.tsx b/Engine/Source/Programs/Horde/HordeDashboard/src/components/jobDetailsV2/JobDetailViewSteps.tsx
--- a/Engine/Source/Programs/Horde/HordeDashboard/src/components/jobDetailsV2/JobDetailViewSteps.tsx
+++ b/Engine/Source/Programs/Horde/HordeDashboard/src/components/jobDetailsV2/JobDetailViewSteps.tsx
@@ -183,7 +183,7 @@ export const StepsPanelInner: React.FC<{ jobDetails: JobDetailsV2, depStepId?: s
    const buildColumns = (): IColumn[] => {
 
       const widths: Record<string, number> = {
-         "Name": 1074,
+         "Name": 732,
          "Progress": 120,
          "Start": 100,
          "End": 100,
diff --git a/Engine/Source/Programs/Horde/HordeDashboard/src/components/jobDetailsV2/JobDetailViewV2.tsx b/Engine/Source/Programs/Horde/HordeDashboard/src/components/jobDetailsV2/JobDetailViewV2.tsx
--- a/Engine/Source/Programs/Horde/HordeDashboard/src/components/jobDetailsV2/JobDetailViewV2.tsx
+++ b/Engine/Source/Programs/Horde/HordeDashboard/src/components/jobDetailsV2/JobDetailViewV2.tsx
@@ -314,7 +314,7 @@ const DetailsViewOverview: React.FC<{ jobDetails: JobDetailsV2 }> = ({ jobDetail
    </Stack>
 };
 
-const rootWidth = 1776;
+const rootWidth = 1440;
 
 const ScrollRestore: React.FC<{ jobDetails: JobDetailsV2, scrollRef: React.RefObject<HTMLDivElement> }> = ({ jobDetails, scrollRef }) => {
 
@@ -345,7 +345,7 @@ const DetailsView: React.FC<{ jobDetails: JobDetailsV2 }> = ({ jobDetails }) =>
 
    const details = jobDetails;
 
-   const centerAlign = vw / 2 - 890;
+   const centerAlign = vw / 2 - 720;
    const key = `windowsize_jobdetail_view_${windowSize.width}_${windowSize.height}`;
 
    return <Stack id="rail_job_overview">
