Horde: Rework C++ remote execution to use direct connection to initiator.

diff --git a/Engine/Source/Programs/Horde/Horde.Agent.Tests/ComputeTaskExecutorTest.cs b/Engine/Source/Programs/Horde/Horde.Agent.Tests/ComputeTaskExecutorTest.cs
deleted file mode 100644
diff --git a/Engine/Source/Programs/Horde/Horde.Agent/Commands/Compute/ComputeCommand.cs b/Engine/Source/Programs/Horde/Horde.Agent/Commands/Compute/ComputeCommand.cs
deleted file mode 100644
diff --git a/Engine/Source/Programs/Horde/Horde.Agent/Commands/Compute/ComputeCommandV2.cs b/Engine/Source/Programs/Horde/Horde.Agent/Commands/Compute/ComputeCommandV2.cs
deleted file mode 100644
diff --git a/Engine/Source/Programs/Horde/Horde.Agent/Commands/Compute/CppComputeCommand.cs b/Engine/Source/Programs/Horde/Horde.Agent/Commands/Compute/CppComputeCommand.cs
new file mode 100644
--- /dev/null
+++ b/Engine/Source/Programs/Horde/Horde.Agent/Commands/Compute/CppComputeCommand.cs
@@ -0,0 +1,166 @@
+// Copyright Epic Games, Inc. All Rights Reserved.
+
+using System;
+using System.Collections.Generic;
+using System.Net.Http;
+using System.Text.Json;
+using System.Threading;
+using System.Threading.Tasks;
+using EpicGames.Core;
+using EpicGames.Horde.Compute;
+using EpicGames.Horde.Compute.Cpp;
+using EpicGames.Horde.Logs;
+using EpicGames.Horde.Storage;
+using EpicGames.Horde.Storage.Backends;
+using EpicGames.Horde.Storage.Nodes;
+using Horde.Agent.Leases.Handlers;
+using Microsoft.Extensions.DependencyInjection;
+using Microsoft.Extensions.Logging;
+using Microsoft.Extensions.Options;
+
+namespace Horde.Agent.Commands
+{
+	/// <summary>
+	/// Installs the agent as a service
+	/// </summary>
+	[Command("cppcompute", "Executes a command through the C++ Compute API")]
+	class CppComputeCommand : Command
+	{
+		class JsonComputeTask
+		{
+			public string Executable { get; set; } = null!;
+			public List<string> Arguments { get; set; } = new List<string>();
+			public string WorkingDir { get; set; } = String.Empty;
+			public Dictionary<string, string> EnvVars { get; set; } = new Dictionary<string, string>();
+			public List<string> OutputPaths { get; set; } = new List<string>();
+		}
+
+		[CommandLine("-Loopback")]
+		bool Loopback { get; set; }
+
+		[CommandLine("-Task=", Required = true)]
+		FileReference TaskFile { get; set; } = null!;
+
+		readonly IServiceProvider _serviceProvider;
+		readonly IHttpClientFactory _httpClientFactory;
+		readonly AgentSettings _settings;
+		readonly ILogger _logger;
+
+		/// <summary>
+		/// Constructor
+		/// </summary>
+		public CppComputeCommand(IServiceProvider serviceProvider, IHttpClientFactory httpClientFactory, IOptions<AgentSettings> settings, ILogger<CppComputeCommand> logger)
+		{
+			_serviceProvider = serviceProvider;
+			_httpClientFactory = httpClientFactory;
+			_settings = settings.Value;
+			_logger = logger;
+		}
+
+		/// <inheritdoc/>
+		public override async Task<int> ExecuteAsync(ILogger logger)
+		{
+			await using IComputeClient channel = CreateClient();
+
+			IComputeRequest<bool> request = await channel.AddRequestAsync<bool>(default, null, HandleRequestAsync);
+			bool result = await request.Result;
+
+			return result? 0 : 1;
+		}
+
+		async Task<bool> HandleRequestAsync(IComputeChannel channel, CancellationToken cancellationToken)
+		{
+			MemoryStorageClient storage = new MemoryStorageClient();
+			NodeLocator node = await CreateComputeNodeAsync(TaskFile, storage, cancellationToken);
+
+			CppComputeMessage message = new CppComputeMessage { Locator = node };
+			await channel.WriteCbMessageAsync(message, cancellationToken);
+
+			for (; ; )
+			{
+				object? obj = await channel.ReadCbMessageAsync(cancellationToken);
+				switch (obj)
+				{
+					case null:
+						return false;
+					case BlobReadMessage readBlob:
+						await channel.WriteBlobDataAsync(readBlob, storage, cancellationToken);
+						break;
+					case CppComputeOutputMessage cppOutput:
+						{
+							using ComputeStorageClient remoteStorage = new ComputeStorageClient(channel);
+							TreeReader reader = new TreeReader(remoteStorage, null, _logger);
+
+							CppComputeOutputNode computeOutput = await reader.ReadNodeAsync<CppComputeOutputNode>(cppOutput.Locator, cancellationToken);
+							_logger.LogInformation("Exit code: {ExitCode}", computeOutput.ExitCode);
+
+							LogNode logNode = await computeOutput.Log.ExpandAsync(reader, cancellationToken);
+							foreach (LogChunkRef logChunkRef in logNode.TextChunkRefs)
+							{
+								LogChunkNode logChunk = await logChunkRef.ExpandAsync(reader, cancellationToken);
+								foreach (Utf8String line in logChunk.Lines)
+								{
+									_logger.LogInformation("Output: {Line}", line);
+								}
+							}
+
+							await channel.WriteCbMessageAsync(new CppComputeFinishMessage(), cancellationToken);
+							await channel.WriteCbMessageAsync(new CloseMessage(), cancellationToken);
+						}
+						return true;
+					default:
+						throw new NotImplementedException();
+				}
+			}
+		}
+
+		static async Task<NodeLocator> CreateComputeNodeAsync(FileReference taskFile, IStorageClient storage, CancellationToken cancellationToken)
+		{
+			using TreeWriter writer = new TreeWriter(storage);
+
+			byte[] data = await FileReference.ReadAllBytesAsync(taskFile, cancellationToken);
+			JsonComputeTask jsonComputeTask = JsonSerializer.Deserialize<JsonComputeTask>(data, new JsonSerializerOptions { AllowTrailingCommas = true, PropertyNameCaseInsensitive = true, PropertyNamingPolicy = JsonNamingPolicy.CamelCase })!;
+
+			DirectoryNode sandbox = new DirectoryNode();
+			await sandbox.CopyFromDirectoryAsync(taskFile.Directory.ToDirectoryInfo(), new ChunkingOptions(), writer, null, cancellationToken);
+
+			CppComputeNode computeNode = new CppComputeNode(jsonComputeTask.Executable, jsonComputeTask.Arguments, jsonComputeTask.WorkingDir, new TreeNodeRef<DirectoryNode>(sandbox));
+			foreach ((string key, string value) in jsonComputeTask.EnvVars)
+			{
+				computeNode.EnvVars.Add(key, value);
+			}
+			foreach (string outputPath in jsonComputeTask.OutputPaths)
+			{
+				computeNode.OutputPaths.Add(outputPath);
+			}
+
+			NodeHandle handle = await writer.FlushAsync(computeNode, cancellationToken);
+			return handle.Locator;
+		}
+
+		IComputeClient CreateClient()
+		{
+			if (Loopback)
+			{
+				return new LoopbackComputeClient(RunListenerAsync);
+			}
+			else
+			{
+				return new ServerComputeClient(CreateHttpClient, _logger);
+			}
+		}
+
+		HttpClient CreateHttpClient()
+		{
+			HttpClient client = _httpClientFactory.CreateClient();
+			client.BaseAddress = _settings.GetCurrentServerProfile().Url;
+			return client;
+		}
+
+		async Task RunListenerAsync(IComputeChannel channel, CancellationToken cancellationToken)
+		{
+			ComputeHandler handler = ActivatorUtilities.CreateInstance<ComputeHandler>(_serviceProvider);
+			await handler.RunAsync(channel, cancellationToken);
+		}
+	}
+}
diff --git a/Engine/Source/Programs/Horde/Horde.Agent/Commands/Compute/DdcComputeCommand.cs b/Engine/Source/Programs/Horde/Horde.Agent/Commands/Compute/DdcComputeCommand.cs
new file mode 100644
--- /dev/null
+++ b/Engine/Source/Programs/Horde/Horde.Agent/Commands/Compute/DdcComputeCommand.cs
@@ -0,0 +1,68 @@
+// Copyright Epic Games, Inc. All Rights Reserved.
+
+using System.Net.Http;
+using System.Threading;
+using System.Threading.Tasks;
+using EpicGames.Core;
+using EpicGames.Horde.Compute;
+using Microsoft.Extensions.Logging;
+using Microsoft.Extensions.Options;
+
+namespace Horde.Agent.Commands
+{
+	/// <summary>
+	/// Installs the agent as a service
+	/// </summary>
+	[Command("ddccompute", "Executes a command through the Horde Compute API")]
+	class DdcComputeCommand : Command
+	{
+		[CommandLine("-Cluster=", Description = "Cluster to execute on")]
+		public ClusterId ClusterId { get; set; } = new ClusterId("default");
+
+		[CommandLine("-Condition=", Description = "Match the agent to run on")]
+		public string? Condition { get; set; }
+
+		readonly AgentSettings _settings;
+
+		/// <summary>
+		/// Constructor
+		/// </summary>
+		public DdcComputeCommand(IOptions<AgentSettings> settings)
+		{
+			_settings = settings.Value;
+		}
+
+		/// <inheritdoc/>
+		public override async Task<int> ExecuteAsync(ILogger logger)
+		{
+			ServerProfile profile = _settings.GetCurrentServerProfile();
+			logger.LogInformation("Connecting to server: {Server}", profile.Name);
+
+			await using ServerComputeClient client = new ServerComputeClient(CreateHttpClient, logger);
+			client.Start();
+
+			IComputeRequest<object?> request = await client.AddRequestAsync(ClusterId, null, TestCommandAsync);
+			await request.Result;
+
+			return 0;
+		}
+
+		HttpClient CreateHttpClient()
+		{
+			HttpClient client = new HttpClient();
+			client.BaseAddress = _settings.GetCurrentServerProfile().Url;
+			return client;
+		}
+
+		async Task<object?> TestCommandAsync(IComputeChannel channel, CancellationToken cancellationToken)
+		{
+			ComputeMessageWriter writer = new ComputeMessageWriter(channel);
+			await writer.WriteCbMessageAsync(new XorRequestMessage { Value = 123, Payload = new byte[] { 1, 2, 3, 4, 5 } }, cancellationToken);
+
+			XorResponseMessage response = await channel.ReadCbMessageAsync<XorResponseMessage>(cancellationToken);
+
+			await writer.WriteCbMessageAsync(new CloseMessage(), cancellationToken);
+			return null;
+		}
+	}
+}
diff --git a/Engine/Source/Programs/Horde/Horde.Agent/Horde.Agent.csproj b/Engine/Source/Programs/Horde/Horde.Agent/Horde.Agent.csproj
--- a/Engine/Source/Programs/Horde/Horde.Agent/Horde.Agent.csproj
+++ b/Engine/Source/Programs/Horde/Horde.Agent/Horde.Agent.csproj
@@ -156,7 +156,6 @@
     <Protobuf Include="..\Horde.Build\Protos\horde\messages\lease.proto" ProtoRoot="..\Horde.Build\Protos" Access="Public" />
     <Protobuf Include="..\Horde.Build\Protos\horde\messages\aws_lambda.proto" ProtoRoot="..\Horde.Build\Protos" Access="Public" />
     <Protobuf Include="..\Horde.Build\Protos\horde\tasks\compute_task.proto" ProtoRoot="..\Horde.Build\Protos" Access="Public" />
-    <Protobuf Include="..\Horde.Build\Protos\horde\tasks\compute_task_v2.proto" ProtoRoot="..\Horde.Build\Protos" Access="Public" />
     <Protobuf Include="..\Horde.Build\Protos\horde\tasks\conform_task.proto" ProtoRoot="..\Horde.Build\Protos" Access="Public" />
     <Protobuf Include="..\Horde.Build\Protos\horde\tasks\job_task.proto" ProtoRoot="..\Horde.Build\Protos" Access="Public" />
     <Protobuf Include="..\Horde.Build\Protos\horde\tasks\restart_task.proto" ProtoRoot="..\Horde.Build\Protos" Access="Public" />
diff --git a/Engine/Source/Programs/Horde/Horde.Agent/Leases/Handlers/ComputeHandler.cs b/Engine/Source/Programs/Horde/Horde.Agent/Leases/Handlers/ComputeHandler.cs
--- a/Engine/Source/Programs/Horde/Horde.Agent/Leases/Handlers/ComputeHandler.cs
+++ b/Engine/Source/Programs/Horde/Horde.Agent/Leases/Handlers/ComputeHandler.cs
@@ -3,361 +3,227 @@
 using System;
 using System.Collections.Generic;
 using System.Diagnostics;
-using System.IO;
 using System.Linq;
-using System.Runtime.InteropServices;
-using System.Text;
+using System.Net.Sockets;
 using System.Threading;
 using System.Threading.Tasks;
+using Amazon.EC2;
 using EpicGames.Core;
 using EpicGames.Horde.Compute;
+using EpicGames.Horde.Compute.Cpp;
+using EpicGames.Horde.Logs;
 using EpicGames.Horde.Storage;
-using EpicGames.Serialization;
-using Google.Protobuf;
-using Grpc.Core;
+using EpicGames.Horde.Storage.Backends;
+using EpicGames.Horde.Storage.Nodes;
 using Horde.Agent.Services;
-using Horde.Agent.Utility;
-using HordeCommon.Rpc;
 using HordeCommon.Rpc.Tasks;
+using Microsoft.CodeAnalysis.CSharp.Syntax;
 using Microsoft.Extensions.Logging;
 
 namespace Horde.Agent.Leases.Handlers
 {
-	class ComputeHandler : LeaseHandler<ComputeTaskMessage>
+	/// <summary>
+	/// Handler for compute tasks
+	/// </summary>
+	class ComputeHandler : LeaseHandler<ComputeTask>
 	{
-		class OutputTree
-		{
-			public Dictionary<string, Task<IoHash>> _files = new Dictionary<string, Task<IoHash>>();
-			public Dictionary<string, OutputTree> _subDirs = new Dictionary<string, OutputTree>();
-		}
-
-		readonly ILegacyStorageClient _storageClient;
 		readonly ILogger _logger;
 
 		/// <summary>
 		/// Constructor
 		/// </summary>
-		public ComputeHandler(ILegacyStorageClient legacyStorageClient, ILogger<ComputeHandler> logger)
+		public ComputeHandler(ILogger<ComputeHandler> logger)
 		{
-			_storageClient = legacyStorageClient;
 			_logger = logger;
 		}
 
 		/// <inheritdoc/>
-		public override async Task<LeaseResult> ExecuteAsync(ISession session, string leaseId, ComputeTaskMessage computeTask, CancellationToken cancellationToken)
+		public override async Task<LeaseResult> ExecuteAsync(ISession session, string leaseId, ComputeTask computeTask, CancellationToken cancellationToken)
 		{
 			_logger.LogInformation("Starting compute task (lease {LeaseId})", leaseId);
 
-			ComputeTaskResultMessage result;
 			try
 			{
-				DateTimeOffset actionTaskStartTime = DateTimeOffset.UtcNow;
-				using (IRpcClientRef<HordeRpc.HordeRpcClient> client = await session.RpcConnection.GetClientRefAsync<HordeRpc.HordeRpcClient>(cancellationToken))
-				{
-					DirectoryReference leaseDir = DirectoryReference.Combine(session.WorkingDir, "Compute", leaseId);
-					DirectoryReference.CreateDirectory(leaseDir);
-
-					try
-					{
-						result = await ExecuteAsync(leaseId, computeTask, leaseDir, cancellationToken);
-					}
-					finally
-					{
-						try
-						{
-							DirectoryReference.Delete(leaseDir, true);
-						}
-						catch
-						{
-						}
-					}
-				}
-			}
-			catch (LegacyBlobNotFoundException ex)
-			{
-				_logger.LogError(ex, "Blob not found: {Hash}", ex.Hash);
-				result = new ComputeTaskResultMessage(ComputeTaskOutcome.BlobNotFound, ex.Hash.ToString());
+				await ConnectAsync(computeTask, cancellationToken);
 			}
 			catch (Exception ex)
 			{
 				_logger.LogError(ex, "Exception while executing compute task");
-				result = new ComputeTaskResultMessage(ComputeTaskOutcome.Exception, ex.ToString());
 			}
 
-			return new LeaseResult(result.ToByteArray());
+			return LeaseResult.Success;
 		}
 
-		/// <summary>
-		/// Read stdout/stderr and wait until process completes
-		/// </summary>
-		/// <param name="process">The process reading from</param>
-		/// <param name="timeout">Max execution timeout for process</param>
-		/// <param name="cancelToken">Cancellation token</param>
-		/// <returns>stdout and stderr data</returns>
-		/// <exception cref="RpcException">Raised for either timeout or cancel</exception>
-		private static async Task<(byte[] stdOutData, byte[] stdOutErr)> ReadProcessStreams(ManagedProcess process, TimeSpan timeout, CancellationToken cancelToken)
+		async Task ConnectAsync(ComputeTask computeTask, CancellationToken cancellationToken)
 		{
-			// Read stdout/stderr without cancellation token.
-			using MemoryStream stdOutStream = new MemoryStream();
-			Task stdOutReadTask = process.StdOut.CopyToAsync(stdOutStream, cancelToken);
+			using TcpClient tcpClient = new TcpClient();
+			await tcpClient.ConnectAsync(computeTask.RemoteIp, computeTask.RemotePort);
 
-			using MemoryStream stdErrStream = new MemoryStream();
-			Task stdErrReadTask = process.StdErr.CopyToAsync(stdErrStream, cancelToken);
-
-			Task outputTask = Task.WhenAll(stdOutReadTask, stdErrReadTask);
-
-			// Instead, create a separate task that will wait for either timeout or cancellation
-			// as cancel interruptions are not reliable with Stream.CopyToAsync()
-			Task timeoutTask = Task.Delay(timeout, cancelToken);
-
-			Task waitTask = await Task.WhenAny(outputTask, timeoutTask);
-			if (waitTask == timeoutTask)
-			{
-				throw new RpcException(new Grpc.Core.Status(StatusCode.DeadlineExceeded, $"Action timed out after {timeout.TotalMilliseconds} ms"));
-			}
+			Socket socket = tcpClient.Client;
+			await socket.SendAsync(computeTask.Nonce.Memory, SocketFlags.None, cancellationToken);
 
-			return (stdOutStream.ToArray(), stdErrStream.ToArray());
+			using SocketComputeChannel channel = new SocketComputeChannel(socket, computeTask.AesKey.Memory, computeTask.AesIv.Memory);
+			await RunAsync(channel, cancellationToken);
 		}
 
-		static int GetMarkerTime(Stopwatch timer)
+		public async Task RunAsync(IComputeChannel channel, CancellationToken cancellationToken)
 		{
-			int elapsedMs = (int)timer.ElapsedMilliseconds;
-			timer.Restart();
-			return elapsedMs;
+			for (; ; )
+			{
+				object? request = await channel.ReadCbMessageAsync(cancellationToken);
+				switch (request)
+				{
+					case null:
+					case CloseMessage _:
+						return;
+					case XorRequestMessage xorRequest:
+						{
+							XorResponseMessage response = new XorResponseMessage();
+							response.Payload = new byte[xorRequest.Payload.Length];
+							for (int idx = 0; idx < xorRequest.Payload.Length; idx++)
+							{
+								response.Payload[idx] = (byte)(xorRequest.Payload[idx] ^ xorRequest.Value);
+							}
+							await channel.WriteCbMessageAsync(response, cancellationToken);
+						}
+						break;
+					case CppComputeMessage cppCompute:
+						await RunCppAsync(channel, cppCompute.Locator, cancellationToken);
+						break;
+				}
+			}
 		}
 
-		/// <summary>
-		/// Execute an action
-		/// </summary>
-		/// <param name="leaseId">The lease id</param>
-		/// <param name="computeTaskMessage">Task to execute</param>
-		/// <param name="sandboxDir">Directory to use as a sandbox for execution</param>
-		/// <param name="cancellationToken">Token used to cancel the operation</param>
-		/// <returns>The action result</returns>
-		public async Task<ComputeTaskResultMessage> ExecuteAsync(string leaseId, ComputeTaskMessage computeTaskMessage, DirectoryReference sandboxDir, CancellationToken cancellationToken)
+		public async Task RunCppAsync(IComputeChannel channel, NodeLocator locator, CancellationToken cancellationToken)
 		{
-			Stopwatch timer = Stopwatch.StartNew();
-			ComputeTaskExecutionStatsMessage stats = new ComputeTaskExecutionStatsMessage();
-			stats.StartTime = Google.Protobuf.WellKnownTypes.Timestamp.FromDateTime(DateTime.UtcNow);
-
-			NamespaceId namespaceId = new NamespaceId(computeTaskMessage.NamespaceId);
-			BucketId inputBucketId = new BucketId(computeTaskMessage.InputBucketId);
-			BucketId outputBucketId = new BucketId(computeTaskMessage.OutputBucketId);
+			DirectoryReference sandboxDir = DirectoryReference.Combine(Program.DataDir, "Sandbox");
 
-			ComputeTask task = await _storageClient.GetRefAsync<ComputeTask>(namespaceId, inputBucketId, computeTaskMessage.TaskRefId.AsRefId(), cancellationToken);
-			_logger.LogInformation("Executing task {Hash} for lease ID {LeaseId}", computeTaskMessage.TaskRefId, leaseId);
-			stats.DownloadRefMs = GetMarkerTime(timer);
+			using ComputeStorageClient store = new ComputeStorageClient(channel);
+			TreeReader reader = new TreeReader(store, null, _logger);
 
-			DirectoryReference.CreateDirectory(sandboxDir);
-			FileUtils.ForceDeleteDirectoryContents(sandboxDir);
+			CppComputeNode node = await reader.ReadNodeAsync<CppComputeNode>(locator, cancellationToken);
+			DirectoryNode directoryNode = await node.Sandbox.ExpandAsync(reader, cancellationToken);
 
-			DirectoryTree inputDirectory = await _storageClient.ReadBlobAsync<DirectoryTree>(namespaceId, task.SandboxHash, cancellationToken: cancellationToken);
-			await SetupSandboxAsync(namespaceId, inputDirectory, sandboxDir);
-			stats.DownloadInputMs = GetMarkerTime(timer);
+			await directoryNode.CopyToDirectoryAsync(reader, sandboxDir.ToDirectoryInfo(), _logger, cancellationToken);
 
-			using (ManagedProcessGroup processGroup = new ManagedProcessGroup())
+			MemoryStorageClient storage = new MemoryStorageClient();
+			using (TreeWriter writer = new TreeWriter(storage))
 			{
-				string fileName = FileReference.Combine(sandboxDir, task.WorkingDirectory.ToString(), task.Executable.ToString()).FullName;
-				string workingDirectory = DirectoryReference.Combine(sandboxDir, task.WorkingDirectory.ToString()).FullName;
-				EnsureFileIsExecutable(fileName);
+				string executable = FileReference.Combine(sandboxDir, node.Executable).FullName;
+				string commandLine = CommandLineArguments.Join(node.Arguments);
+				string workingDir = DirectoryReference.Combine(sandboxDir, node.WorkingDirectory).FullName;
 
-				Dictionary<string, string> newEnvironment = new Dictionary<string, string>();
-				foreach (System.Collections.DictionaryEntry? entry in Environment.GetEnvironmentVariables())
+				Dictionary<string, string>? envVars = null;
+				if (node.EnvVars.Count > 0)
 				{
-					newEnvironment[entry!.Value.Key.ToString()!] = entry!.Value.Value!.ToString()!;
-				}
-				foreach (KeyValuePair<Utf8String, Utf8String> pair in task.EnvVars)
-				{
-					newEnvironment[pair.Key.ToString()] = pair.Value.ToString();
-				}
-
-				string arguments = CommandLineArguments.Join(task.Arguments.Select(x => x.ToString()));
-				_logger.LogInformation("Executing {FileName} with arguments {Arguments}", fileName, arguments);
-
-				TimeSpan timeout = TimeSpan.FromMinutes(5.0);// Action.Timeout == null ? TimeSpan.FromMinutes(5) : Action.Timeout.ToTimeSpan();
-				using (ManagedProcess process = new ManagedProcess(processGroup, fileName, arguments, workingDirectory, newEnvironment, ProcessPriorityClass.Normal, ManagedProcessFlags.None))
-				{
-					(byte[] stdOutData, byte[] stdErrData) = await ReadProcessStreams(process, timeout, cancellationToken);
-					stats.ExecMs = GetMarkerTime(timer);
-
-					foreach (string line in Encoding.UTF8.GetString(stdOutData).Split('\n'))
-					{
-						_logger.LogInformation("stdout: {Line}", line);
-					}
-					foreach (string line in Encoding.UTF8.GetString(stdErrData).Split('\n'))
+					envVars = ManagedProcess.GetCurrentEnvVars();
+					foreach ((string key, string value) in node.EnvVars)
 					{
-						_logger.LogInformation("stderr: {Line}", line);
+						envVars[key] = value;
 					}
+				}
 
-					process.WaitForExit();
-					_logger.LogInformation("exit: {ExitCode}", process.ExitCode);
-
-					ComputeTaskResult result = new ComputeTaskResult(process.ExitCode);
-					result.StdOutHash = await _storageClient.WriteBlobFromMemoryAsync(namespaceId, stdOutData, cancellationToken);
-					result.StdErrHash = await _storageClient.WriteBlobFromMemoryAsync(namespaceId, stdErrData, cancellationToken);
-					stats.UploadLogMs = GetMarkerTime(timer);
+				int exitCode;
 
-					FileReference[] outputFiles = ResolveOutputPaths(sandboxDir, task.OutputPaths.Select(x => x.ToString())).OrderBy(x => x.FullName, StringComparer.Ordinal).ToArray();
-					if (outputFiles.Length > 0)
+				LogWriter logWriter = new LogWriter(writer);
+				using (ManagedProcessGroup group = new ManagedProcessGroup())
+				{
+					using (ManagedProcess process = new ManagedProcess(group, executable, commandLine, workingDir, envVars, null, ProcessPriorityClass.Normal))
 					{
-						foreach (FileReference outputFile in outputFiles)
+						byte[] buffer = new byte[1024];
+						for (; ; )
 						{
-							_logger.LogInformation("output: {File}", outputFile.MakeRelativeTo(sandboxDir));
+							int length = await process.ReadAsync(buffer, 0, buffer.Length, cancellationToken);
+							if (length == 0)
+							{
+								process.WaitForExit();
+								exitCode = process.ExitCode;
+								break;
+							}
+							await logWriter.AppendAsync(buffer.AsMemory(0, length), cancellationToken);
 						}
-						result.OutputHash = await PutOutput(namespaceId, sandboxDir, outputFiles);
 					}
-					stats.UploadOutputMs = GetMarkerTime(timer);
+				}
 
-					CbObject resultObject = CbSerializer.Serialize(result);
-					await _storageClient.SetRefAsync(namespaceId, outputBucketId, computeTaskMessage.TaskRefId, resultObject, cancellationToken);
-					stats.UploadRefMs = GetMarkerTime(timer);
+				TreeNodeRef<LogNode> logNodeRef = await logWriter.CompleteAsync(cancellationToken);
 
-					return new ComputeTaskResultMessage(computeTaskMessage.TaskRefId, stats);
-				}
-			}
-		}
+				FileFilter filter = new FileFilter(node.OutputPaths.Select(x => x.ToString()));
+				List<FileReference> outputFiles = filter.ApplyToDirectory(sandboxDir, false);
 
-		/// <summary>
-		/// Downloads files to the sandbox
-		/// </summary>
-		/// <param name="namespaceId">Namespace for fetching data</param>
-		/// <param name="inputDirectory">The directory spec</param>
-		/// <param name="outputDir">Output directory on disk</param>
-		/// <returns>Async task</returns>
-		internal async Task SetupSandboxAsync(NamespaceId namespaceId, DirectoryTree inputDirectory, DirectoryReference outputDir)
-		{
-			DirectoryReference.CreateDirectory(outputDir);
+				DirectoryNode outputTree = new DirectoryNode();
+				await outputTree.CopyFilesAsync(sandboxDir, outputFiles, new ChunkingOptions(), writer, null, cancellationToken);
 
-			async Task DownloadFile(FileNode fileNode)
-			{
-				FileReference file = FileReference.Combine(outputDir, fileNode.Name.ToString());
-				_logger.LogInformation("Downloading {File} (digest: {Digest})", file, fileNode.Hash);
-				byte[] data = fileNode.IsCompressed
-					? await _storageClient.ReadCompressedBlobToMemoryAsync(namespaceId, fileNode.Hash)
-					: await _storageClient.ReadBlobToMemoryAsync(namespaceId, fileNode.Hash);
-				_logger.LogInformation("Writing {File} (digest: {Digest})", file, fileNode.Hash);
-				await FileReference.WriteAllBytesAsync(file, data);
+				CppComputeOutputNode outputNode = new CppComputeOutputNode(exitCode, logNodeRef, new TreeNodeRef<DirectoryNode>(outputTree));
+				NodeHandle outputHandle = await writer.FlushAsync(outputNode, cancellationToken);
+
+				CppComputeOutputMessage outputMessage = new CppComputeOutputMessage { Locator = outputHandle.Locator };
+				await channel.WriteCbMessageAsync(outputMessage, cancellationToken);
 			}
 
-			async Task DownloadDir(DirectoryNode directoryNode)
+			for (; ; )
 			{
-				DirectoryTree inputSubDirectory = await _storageClient.ReadBlobAsync<DirectoryTree>(namespaceId, directoryNode.Hash);
-				DirectoryReference outputSubDirectory = DirectoryReference.Combine(outputDir, directoryNode.Name.ToString());
-				await SetupSandboxAsync(namespaceId, inputSubDirectory, outputSubDirectory);
+				object? request = await channel.ReadCbMessageAsync(cancellationToken);
+				switch (request)
+				{
+					case null:
+					case CppComputeFinishMessage _:
+						return;
+					case BlobReadMessage blobRead:
+						await channel.WriteBlobDataAsync(blobRead, storage, cancellationToken);
+						break;
+				}
 			}
-
-			List<Task> tasks = new List<Task>();
-			tasks.AddRange(inputDirectory.Files.Select(x => Task.Run(() => DownloadFile(x))));
-			tasks.AddRange(inputDirectory.Directories.Select(x => Task.Run(() => DownloadDir(x))));
-			await Task.WhenAll(tasks);
 		}
 
-		async Task<IoHash> PutOutput(NamespaceId namespaceId, DirectoryReference baseDir, IEnumerable<FileReference> files)
+		class LogWriter
 		{
-			List<FileReference> sortedFiles = files.OrderBy(x => x.FullName, StringComparer.Ordinal).ToList();
-			(_, IoHash hash) = await PutDirectoryTree(namespaceId, baseDir.FullName.Length, sortedFiles, 0, sortedFiles.Count);
-			return hash;
-		}
+			const int MaxChunkSize = 128 * 1024;
+			const int FlushChunkSize = 100 * 1024;
 
-		async Task<(DirectoryTree, IoHash)> PutDirectoryTree(NamespaceId namespaceId, int baseDirLen, List<FileReference> sortedFiles, int minIdx, int maxIdx)
-		{
-			List<Task<FileNode>> files = new List<Task<FileNode>>();
-			List<Task<DirectoryNode>> trees = new List<Task<DirectoryNode>>();
+			readonly TreeWriter _writer;
+			int _lineCount;
+			long _offset;
+			readonly LogChunkBuilder _chunkBuilder = new LogChunkBuilder(MaxChunkSize);
+			readonly List<LogChunkRef> _chunks = new List<LogChunkRef>();
 
-			while (minIdx < maxIdx)
+			public IReadOnlyList<LogChunkRef> Chunks => _chunks;
+
+			public LogWriter(TreeWriter writer)
 			{
-				FileReference file = sortedFiles[minIdx];
+				_writer = writer;
+			}
 
-				int nextMinIdx = minIdx + 1;
+			public async ValueTask AppendAsync(ReadOnlyMemory<byte> data, CancellationToken cancellationToken)
+			{
+				_chunkBuilder.Append(data.Span);
 
-				int nextDirLen = file.FullName.IndexOf(Path.DirectorySeparatorChar, baseDirLen + 1);
-				if (nextDirLen == -1)
+				if (_chunkBuilder.Length > FlushChunkSize)
 				{
-					string name = file.FullName.Substring(baseDirLen + 1);
-					files.Add(CreateFileNode(namespaceId, name, file));
+					await FlushAsync(cancellationToken);
 				}
-				else
-				{
-					string name = file.FullName.Substring(baseDirLen + 1, nextDirLen - (baseDirLen + 1));
-					while (nextMinIdx < maxIdx)
-					{
-						string nextFile = sortedFiles[nextMinIdx].FullName;
-						if (nextFile.Length < nextDirLen || String.Compare(name, 0, nextFile, baseDirLen, name.Length, StringComparison.Ordinal) == 0)
-						{
-							break;
-						}
-						nextMinIdx++;
-					}
-					trees.Add(CreateDirectoryNode(namespaceId, name, nextDirLen, sortedFiles, minIdx, nextMinIdx));
-				}
-
-				minIdx = nextMinIdx;
 			}
 
-			DirectoryTree tree = new DirectoryTree();
-			tree.Files.AddRange(await Task.WhenAll(files));
-			tree.Directories.AddRange(await Task.WhenAll(trees));
-
-			IoHash hash = await _storageClient.WriteBlobAsync<DirectoryTree>(namespaceId, tree);
-			return (tree, hash);
-		}
-
-		private static void EnsureFileIsExecutable(string filePath)
-		{
-			if (RuntimeInformation.IsOSPlatform(OSPlatform.Linux))
+			async Task FlushAsync(CancellationToken cancellationToken)
 			{
-				ushort mode = 493; // 0755 octal Posix permission
-				FileUtils.SetFileMode_Linux(filePath, mode);
+				LogChunkNode chunkNode = _chunkBuilder.ToLogChunk();
+				LogChunkRef chunkRef = new LogChunkRef(_lineCount, _offset, chunkNode);
+				await _writer.WriteAsync(chunkRef, cancellationToken);
+
+				_chunks.Add(chunkRef);
+				_lineCount += chunkNode.LineCount;
+				_offset += chunkNode.Length;
 			}
-		}
 
-		async Task<DirectoryNode> CreateDirectoryNode(NamespaceId namespaceId, string name, int baseDirLen, List<FileReference> sortedFiles, int minIdx, int maxIdx)
-		{
-			(_, IoHash hash) = await PutDirectoryTree(namespaceId, baseDirLen, sortedFiles, minIdx, maxIdx);
-			return new DirectoryNode(name, hash);
-		}
+			public async Task<TreeNodeRef<LogNode>> CompleteAsync(CancellationToken cancellationToken)
+			{
+				await FlushAsync(cancellationToken);
+				LogIndexNode index = new LogIndexNode(new NgramSet(Array.Empty<ushort>()), 0, Array.Empty<LogChunkRef>());
 
-		async Task<FileNode> CreateFileNode(NamespaceId namespaceId, string name, FileReference file)
-		{
-			byte[] data = await FileReference.ReadAllBytesAsync(file);
-			IoHash hash = await _storageClient.WriteBlobFromMemoryAsync(namespaceId, data);
-			// TODO: Figure out how file can be marked as compressed
-			return new FileNode(name, hash, data.Length, (int)FileReference.GetAttributes(file), false);
-		}
+				TreeNodeRef<LogNode> logNodeRef = new TreeNodeRef<LogNode>(new LogNode(LogFormat.Text, _lineCount, _offset, _chunks, new TreeNodeRef<LogIndexNode>(index), true));
+				await _writer.WriteAsync(logNodeRef, cancellationToken);
 
-		/// <summary>
-		/// Resolves a list of output paths into file references
-		///
-		/// The REAPI spec allows directories to be specified as an output path which require all sub dirs and files
-		/// to be resolved.
-		/// </summary>
-		/// <param name="sandboxDir">Base directory where execution is taking place</param>
-		/// <param name="outputPaths">List of output paths relative to SandboxDir</param>
-		/// <returns>List of resolved paths (incl expanded dirs)</returns>
-		internal static HashSet<FileReference> ResolveOutputPaths(DirectoryReference sandboxDir, IEnumerable<string> outputPaths)
-		{
-			HashSet<FileReference> files = new HashSet<FileReference>();
-			foreach (string outputPath in outputPaths)
-			{
-				DirectoryReference dirRef = DirectoryReference.Combine(sandboxDir, outputPath);
-				if (DirectoryReference.Exists(dirRef))
-				{
-					IEnumerable<FileReference> listedFiles = DirectoryReference.EnumerateFiles(dirRef, "*", SearchOption.AllDirectories);
-					foreach (FileReference listedFileRef in listedFiles)
-					{
-						files.Add(listedFileRef);
-					}
-				}
-				else
-				{
-					FileReference fileRef = FileReference.Combine(sandboxDir, outputPath);
-					if (FileReference.Exists(fileRef))
-					{
-						files.Add(fileRef);
-					}
-				}
+				return logNodeRef;
 			}
-			return files;
 		}
 	}
 }
diff --git a/Engine/Source/Programs/Horde/Horde.Agent/Leases/Handlers/ComputeHandlerV2.cs b/Engine/Source/Programs/Horde/Horde.Agent/Leases/Handlers/ComputeHandlerV2.cs
deleted file mode 100644
diff --git a/Engine/Source/Programs/Horde/Horde.Agent/Program.cs b/Engine/Source/Programs/Horde/Horde.Agent/Program.cs
--- a/Engine/Source/Programs/Horde/Horde.Agent/Program.cs
+++ b/Engine/Source/Programs/Horde/Horde.Agent/Program.cs
@@ -189,7 +189,6 @@ namespace Horde.Agent
 			services.AddSingleton<JobHandler>();
 
 			services.AddSingleton<LeaseHandler, ComputeHandler>();
-			services.AddSingleton<LeaseHandler, ComputeHandlerV2>();
 			services.AddSingleton<LeaseHandler, ConformHandler>();
 			services.AddSingleton<LeaseHandler, JobHandler>(x => x.GetRequiredService<JobHandler>());
 			services.AddSingleton<LeaseHandler, RestartHandler>();
diff --git a/Engine/Source/Programs/Horde/Horde.Build.Tests/Compute/ComputeServiceTest.cs b/Engine/Source/Programs/Horde/Horde.Build.Tests/Compute/ComputeServiceTest.cs
deleted file mode 100644
diff --git a/Engine/Source/Programs/Horde/Horde.Build.Tests/Compute/RedisTaskSchedulerTest.cs b/Engine/Source/Programs/Horde/Horde.Build.Tests/Compute/RedisTaskSchedulerTest.cs
--- a/Engine/Source/Programs/Horde/Horde.Build.Tests/Compute/RedisTaskSchedulerTest.cs
+++ b/Engine/Source/Programs/Horde/Horde.Build.Tests/Compute/RedisTaskSchedulerTest.cs
@@ -2,7 +2,7 @@
 
 using System.Threading;
 using System.Threading.Tasks;
-using Horde.Build.Compute.V1;
+using Horde.Build.Compute;
 using Horde.Build.Server;
 using Microsoft.Extensions.Logging;
 using Microsoft.VisualStudio.TestTools.UnitTesting;
diff --git a/Engine/Source/Programs/Horde/Horde.Build.Tests/TestSetup.cs b/Engine/Source/Programs/Horde/Horde.Build.Tests/TestSetup.cs
--- a/Engine/Source/Programs/Horde/Horde.Build.Tests/TestSetup.cs
+++ b/Engine/Source/Programs/Horde/Horde.Build.Tests/TestSetup.cs
@@ -99,7 +99,7 @@ namespace Horde.Build.Tests
 		public AgentSoftwareService AgentSoftwareService => ServiceProvider.GetRequiredService<AgentSoftwareService>();
 		public AgentService AgentService => ServiceProvider.GetRequiredService<AgentService>();
 		public ICommitService CommitService => ServiceProvider.GetRequiredService<ICommitService>();
-		public Horde.Build.Compute.V1.ComputeService ComputeService => ServiceProvider.GetRequiredService<Horde.Build.Compute.V1.ComputeService>();
+		public ComputeService ComputeService => ServiceProvider.GetRequiredService<ComputeService>();
 		public GlobalsService GlobalsService => ServiceProvider.GetRequiredService<GlobalsService>();
 		public MongoService MongoService => ServiceProvider.GetRequiredService<MongoService>();
 		public ITemplateCollection TemplateCollection => ServiceProvider.GetRequiredService<ITemplateCollection>();
@@ -239,7 +239,7 @@ namespace Horde.Build.Tests
 			services.AddSingleton<AgentSoftwareService>();
 			services.AddSingleton<FleetService>();
 			services.AddSingleton<ConsistencyService>();
-			services.AddSingleton<Horde.Build.Compute.V1.ComputeService>();
+			services.AddSingleton<ComputeService>();
 			services.AddSingleton<RequestTrackerService>();
 			services.AddSingleton<GlobalsService>();
 			services.AddSingleton<CredentialService>();
diff --git a/Engine/Source/Programs/Horde/Horde.Build/Agents/Fleet/ComputeQueueAwsMetricStrategy.cs b/Engine/Source/Programs/Horde/Horde.Build/Agents/Fleet/ComputeQueueAwsMetricStrategy.cs
--- a/Engine/Source/Programs/Horde/Horde.Build/Agents/Fleet/ComputeQueueAwsMetricStrategy.cs
+++ b/Engine/Source/Programs/Horde/Horde.Build/Agents/Fleet/ComputeQueueAwsMetricStrategy.cs
@@ -9,7 +9,7 @@ using Amazon.CloudWatch;
 using Amazon.CloudWatch.Model;
 using EpicGames.Horde.Compute;
 using Horde.Build.Agents.Pools;
-using Horde.Build.Compute.V1;
+using Horde.Build.Compute;
 using Microsoft.Extensions.Logging;
 using OpenTracing;
 using OpenTracing.Util;
diff --git a/Engine/Source/Programs/Horde/Horde.Build/Compute/ComputeController.cs b/Engine/Source/Programs/Horde/Horde.Build/Compute/ComputeController.cs
new file mode 100644
--- /dev/null
+++ b/Engine/Source/Programs/Horde/Horde.Build/Compute/ComputeController.cs
@@ -0,0 +1,134 @@
+// Copyright Epic Games, Inc. All Rights Reserved.
+
+using System;
+using System.Collections.Generic;
+using System.Net;
+using EpicGames.Core;
+using EpicGames.Horde.Compute;
+using Google.Protobuf;
+using Horde.Build.Acls;
+using Horde.Build.Server;
+using Horde.Build.Utilities;
+using HordeCommon.Rpc.Tasks;
+using Microsoft.AspNetCore.Authorization;
+using Microsoft.AspNetCore.Mvc;
+using Microsoft.Extensions.Options;
+
+namespace Horde.Build.Compute
+{
+	/// <summary>
+	/// Request a machine to execute compute requests
+	/// </summary>
+	public class AddComputeTaskRequest
+	{
+		/// <summary>
+		/// Condition to identify machines that can execute the request
+		/// </summary>
+		public Requirements? Requirements { get; set; }
+
+		/// <summary>
+		/// Port to connect on
+		/// </summary>
+		public int? RemotePort { get; set; }
+
+		/// <summary>
+		/// Cryptographic nonce to identify the request, as a hex string
+		/// </summary>
+		public string Nonce { get; set; } = String.Empty;
+
+		/// <summary>
+		/// AES key for the channel, as a hex string
+		/// </summary>
+		public string AesKey { get; set; } = String.Empty;
+
+		/// <summary>
+		/// AES IV for the channel, as a hex string
+		/// </summary>
+		public string AesIv { get; set; } = String.Empty;
+	}
+
+	/// <summary>
+	/// Request a machine to execute compute requests
+	/// </summary>
+	public class AddComputeTasksRequest
+	{
+		/// <summary>
+		/// Condition to identify machines that can execute the request
+		/// </summary>
+		public Requirements? Requirements { get; set; }
+
+		/// <summary>
+		/// Port to connect on
+		/// </summary>
+		public int RemotePort { get; set; }
+
+		/// <summary>
+		/// List of tasks to add
+		/// </summary>
+		public List<AddComputeTaskRequest> Tasks { get; set; } = new List<AddComputeTaskRequest>();
+	}
+
+	/// <summary>
+	/// Controller for the /api/v2/compute endpoint
+	/// </summary>
+	[ApiController]
+	[Authorize]
+	[Route("[controller]")]
+	public class ComputeControllerV2 : HordeControllerBase
+	{
+		readonly ComputeService _computeService;
+		readonly IOptionsSnapshot<GlobalConfig> _globalConfig;
+
+		/// <summary>
+		/// Constructor
+		/// </summary>
+		public ComputeControllerV2(ComputeService computeService, IOptionsSnapshot<GlobalConfig> globalConfig)
+		{
+			_computeService = computeService;
+			_globalConfig = globalConfig;
+		}
+
+		/// <summary>
+		/// Add tasks to be executed remotely
+		/// </summary>
+		/// <param name="clusterId">Id of the compute cluster</param>
+		/// <param name="request">The request parameters</param>
+		/// <returns></returns>
+		[HttpPost]
+		[Authorize]
+		[Route("/api/v2/compute/{clusterId}")]
+		public ActionResult AddTasksAsync(ClusterId clusterId, [FromBody] AddComputeTasksRequest request)
+		{
+			ComputeClusterConfig? clusterConfig;
+			if (!_globalConfig.Value.TryGetComputeCluster(clusterId, out clusterConfig))
+			{
+				return NotFound(clusterId);
+			}
+			if(!clusterConfig.Authorize(AclAction.AddComputeTasks, User))
+			{
+				return Forbid(AclAction.AddComputeTasks, clusterId);
+			}
+
+			IPAddress? remoteIp = HttpContext.Connection.RemoteIpAddress;
+			if (remoteIp == null)
+			{
+				return BadRequest("Missing remote IP address");
+			}
+
+			foreach (AddComputeTaskRequest taskRequest in request.Tasks)
+			{
+				Requirements requirements = taskRequest.Requirements ?? request.Requirements ?? new Requirements();
+
+				ComputeTask computeTask = new ComputeTask();
+				computeTask.RemoteIp = remoteIp.ToString();
+				computeTask.RemotePort = taskRequest.RemotePort ?? request.RemotePort;
+				computeTask.Nonce = ByteString.CopyFrom(StringUtils.ParseHexString(taskRequest.Nonce));
+				computeTask.AesKey = ByteString.CopyFrom(StringUtils.ParseHexString(taskRequest.AesKey));
+				computeTask.AesIv = ByteString.CopyFrom(StringUtils.ParseHexString(taskRequest.AesIv));
+
+				_computeService.AddRequest(clusterId, requirements, computeTask);
+			}
+			return Ok();
+		}
+	}
+}
diff --git a/Engine/Source/Programs/Horde/Horde.Build/Compute/ComputeRequest.cs b/Engine/Source/Programs/Horde/Horde.Build/Compute/ComputeRequest.cs
new file mode 100644
--- /dev/null
+++ b/Engine/Source/Programs/Horde/Horde.Build/Compute/ComputeRequest.cs
@@ -0,0 +1,32 @@
+// Copyright Epic Games, Inc. All Rights Reserved.
+
+using EpicGames.Horde.Compute;
+using HordeCommon.Rpc.Tasks;
+
+namespace Horde.Build.Compute
+{
+	/// <summary>
+	/// Request for a compute resource
+	/// </summary>
+	public class ComputeRequest
+	{
+		/// <summary>
+		/// Requirements for the agent to satisfy this request
+		/// </summary>
+		public Requirements Requirements { get; }
+
+		/// <summary>
+		/// Connection settings for the channel
+		/// </summary>
+		public ComputeTask Task { get; }
+
+		/// <summary>
+		/// Constructor
+		/// </summary>
+		public ComputeRequest(Requirements requirements, ComputeTask task)
+		{
+			Requirements = requirements;
+			Task = task;
+		}
+	}
+}
diff --git a/Engine/Source/Programs/Horde/Horde.Build/Compute/ComputeService.cs b/Engine/Source/Programs/Horde/Horde.Build/Compute/ComputeService.cs
new file mode 100644
--- /dev/null
+++ b/Engine/Source/Programs/Horde/Horde.Build/Compute/ComputeService.cs
@@ -0,0 +1,131 @@
+// Copyright Epic Games, Inc. All Rights Reserved.
+
+using System.Collections.Generic;
+using EpicGames.Core;
+using EpicGames.Horde.Compute;
+using EpicGames.Serialization;
+using Horde.Build.Agents;
+using Horde.Build.Agents.Pools;
+using Horde.Build.Server;
+using HordeCommon.Rpc.Tasks;
+using Microsoft.Extensions.Options;
+
+namespace Horde.Build.Compute
+{
+	/// <summary>
+	/// Dispatches remote actions. Does not implement any cross-pod communication to satisfy leases; only agents connected to this server instance will be stored.
+	/// </summary>
+	public class ComputeService
+	{
+		class RequestQueue
+		{
+			public IoHash Hash { get; }
+			public Requirements Requirements { get; }
+			public Queue<ComputeRequest> Requests { get; } = new Queue<ComputeRequest>();
+
+			public RequestQueue(IoHash hash, Requirements requirements)
+			{
+				Hash = hash;
+				Requirements = requirements;
+			}
+		}
+
+		class ClusterInfo
+		{
+			public Dictionary<IoHash, RequestQueue> Queues { get; } = new Dictionary<IoHash, RequestQueue>();
+		}
+
+		readonly object _lockObject = new object();
+		readonly IOptionsMonitor<GlobalConfig> _globalConfig;
+		readonly Dictionary<ClusterId, ClusterInfo> _clusters = new Dictionary<ClusterId, ClusterInfo>();
+
+		/// <summary>
+		/// Constructor
+		/// </summary>
+		protected ComputeService(IOptionsMonitor<GlobalConfig> globalConfig)
+		{
+			_globalConfig = globalConfig;
+		}
+
+		/// <summary>
+		/// Adds a new compute request
+		/// </summary>
+		/// <param name="clusterId">The compute cluster id</param>
+		/// <param name="requirements">Requirements for the agent to serve the request</param>
+		/// <param name="task">Connection information for the compute channel</param>
+		public void AddRequest(ClusterId clusterId, Requirements requirements, ComputeTask task)
+		{
+			IoHash hash = IoHash.Compute(requirements.Serialize());
+			lock (_lockObject)
+			{
+				ClusterInfo? cluster;
+				if (!_clusters.TryGetValue(clusterId, out cluster))
+				{
+					cluster = new ClusterInfo();
+					_clusters.Add(clusterId, cluster);
+				}
+
+				RequestQueue? queue;
+				if (!cluster.Queues.TryGetValue(hash, out queue))
+				{
+					queue = new RequestQueue(hash, requirements);
+					cluster.Queues.Add(hash, queue);
+				}
+
+				queue.Requests.Enqueue(new ComputeRequest(requirements, task));
+			}
+		}
+
+		/// <summary>
+		/// Attempts to pop a request from the queue which can be met by the given agent
+		/// </summary>
+		/// <param name="agent">Agent to allocate a request for</param>
+		/// <returns>The allocated request</returns>
+		public ComputeRequest? PopRequest(IAgent agent)
+		{
+			ComputeRequest? request = null;
+			lock (_lockObject)
+			{
+				ClusterId? removeClusterId = null;
+
+				GlobalConfig globalConfig = _globalConfig.CurrentValue;
+				foreach ((ClusterId clusterId, ClusterInfo clusterInfo) in _clusters)
+				{
+					ComputeClusterConfig? clusterConfig;
+					if (!globalConfig.TryGetComputeCluster(clusterId, out clusterConfig))
+					{
+						removeClusterId ??= clusterId;
+					}
+					else if (clusterConfig.Condition == null || agent.SatisfiesCondition(clusterConfig.Condition))
+					{
+						foreach ((IoHash hash, RequestQueue queue) in clusterInfo.Queues)
+						{
+							if (agent.MeetsRequirements(queue.Requirements))
+							{
+								request = queue.Requests.Dequeue();
+
+								if (queue.Requests.Count == 0)
+								{
+									clusterInfo.Queues.Remove(hash);
+								}
+								break;
+							}
+						}
+
+						if (clusterInfo.Queues.Count == 0)
+						{
+							removeClusterId ??= clusterId;
+						}
+					}
+				}
+
+				if (removeClusterId != null)
+				{
+					_clusters.Remove(removeClusterId.Value);
+				}
+			}
+
+			return request;
+		}
+	}
+}
diff --git a/Engine/Source/Programs/Horde/Horde.Build/Compute/ComputeTaskSource.cs b/Engine/Source/Programs/Horde/Horde.Build/Compute/ComputeTaskSource.cs
new file mode 100644
--- /dev/null
+++ b/Engine/Source/Programs/Horde/Horde.Build/Compute/ComputeTaskSource.cs
@@ -0,0 +1,63 @@
+// Copyright Epic Games, Inc. All Rights Reserved.
+
+using System;
+using System.Threading;
+using System.Threading.Tasks;
+using EpicGames.Core;
+using EpicGames.Horde.Compute;
+using Google.Protobuf;
+using Google.Protobuf.WellKnownTypes;
+using Horde.Build.Agents;
+using Horde.Build.Agents.Leases;
+using Horde.Build.Tasks;
+using HordeCommon;
+using HordeCommon.Rpc.Tasks;
+using Microsoft.Extensions.Logging;
+
+namespace Horde.Build.Compute
+{
+	/// <summary>
+	/// Dispatches requests for compute resources
+	/// </summary>
+	public class ComputeTaskSource : TaskSourceBase<ComputeTask>
+	{
+		/// <inheritdoc/>
+		public override string Type => "Compute";
+
+		/// <inheritdoc/>
+		public override TaskSourceFlags Flags => TaskSourceFlags.None;
+
+		readonly ComputeService _computeService;
+		readonly ILogger _logger;
+
+		/// <summary>
+		/// Constructor
+		/// </summary>
+		public ComputeTaskSource(ComputeService computeService, ILogger<ComputeTaskSource> logger)
+		{
+			_computeService = computeService;
+			_logger = logger;
+		}
+
+		/// <inheritdoc/>
+		public override async Task<Task<AgentLease?>> AssignLeaseAsync(IAgent agent, CancellationToken cancellationToken)
+		{
+			while (!cancellationToken.IsCancellationRequested)
+			{
+				ComputeRequest? request = _computeService.PopRequest(agent);
+				if (request != null)
+				{
+					string leaseName = $"Compute for {request.Task.RemoteIp}";
+					byte[] payload = Any.Pack(request.Task).ToByteArray();
+
+					AgentLease lease = new AgentLease(LeaseId.GenerateNewId(), leaseName, null, null, null, LeaseState.Pending, request.Requirements.Resources, request.Requirements.Exclusive, payload);
+					_logger.LogInformation("Created compute lease for agent {AgentId} and remote {RemoteIp}:{RemotePort}", agent.Id, request.Task.RemoteIp, request.Task.RemotePort);
+					return Task.FromResult<AgentLease?>(lease);
+				}
+
+				await AsyncUtils.DelayNoThrow(TimeSpan.FromSeconds(5.0), cancellationToken);
+			}
+			return Task.FromResult<AgentLease?>(null);
+		}
+	}
+}
diff --git a/Engine/Source/Programs/Horde/Horde.Build/Compute/RedisMessageQueue.cs b/Engine/Source/Programs/Horde/Horde.Build/Compute/RedisMessageQueue.cs
new file mode 100644
--- /dev/null
+++ b/Engine/Source/Programs/Horde/Horde.Build/Compute/RedisMessageQueue.cs
@@ -0,0 +1,180 @@
+// Copyright Epic Games, Inc. All Rights Reserved.
+
+using System;
+using System.Collections.Generic;
+using System.Threading;
+using System.Threading.Tasks;
+using EpicGames.Redis;
+using StackExchange.Redis;
+
+namespace Horde.Build.Compute
+{
+	/// <summary>
+	/// Interface for a distributed message queue
+	/// </summary>
+	/// <typeparam name="T">The message type</typeparam>
+	interface IMessageQueue<T> where T : class
+	{
+		/// <summary>
+		/// Adds messages to the queue
+		/// </summary>
+		/// <param name="channelId">The channel to post to</param>
+		/// <param name="message">Message to post</param>
+		Task PostAsync(string channelId, T message);
+
+		/// <summary>
+		/// Waits for a message to be available on the given channel
+		/// </summary>
+		/// <param name="channelId">The channel to read from</param>
+		/// <returns>True if a message was available, false otherwise</returns>
+		Task<List<T>> ReadMessagesAsync(string channelId);
+
+		/// <summary>
+		/// Waits for a message to be available on the given channel
+		/// </summary>
+		/// <param name="channelId">The channel to read from</param>
+		/// <param name="cancellationToken">May be signalled to stop the wait, without throwing an exception</param>
+		/// <returns>True if a message was available, false otherwise</returns>
+		Task<List<T>> WaitForMessagesAsync(string channelId, CancellationToken cancellationToken);
+	}
+
+	/// <summary>
+	/// Implementation of a message queue using Redis
+	/// </summary>
+	/// <typeparam name="T">The message type. Must be serailizable by RedisSerializer.</typeparam>
+	class RedisMessageQueue<T> : IMessageQueue<T>, IDisposable where T : class
+	{
+		readonly IDatabase _redis;
+		readonly RedisKey _keyPrefix;
+		readonly RedisChannel _updateChannel;
+		readonly Dictionary<string, TaskCompletionSource<bool>> _channelWakeEvents = new Dictionary<string, TaskCompletionSource<bool>>();
+
+		// Time after which entries should be removed
+		TimeSpan ExpireTime { get; set; } = TimeSpan.FromSeconds(30);
+
+		/// <summary>
+		/// Constructor
+		/// </summary>
+		/// <param name="redis">The Redis database instance</param>
+		/// <param name="keyPrefix">Prefix for keys to use in the database</param>
+		public RedisMessageQueue(IDatabase redis, RedisKey keyPrefix)
+		{
+			_redis = redis;
+			_keyPrefix = keyPrefix;
+			_updateChannel = keyPrefix.Append("updates").ToString();
+
+			redis.Multiplexer.GetSubscriber().Subscribe(_updateChannel, OnChannelUpdate);
+		}
+
+		/// <inheritdoc/>
+		public void Dispose()
+		{
+			_redis.Multiplexer.GetSubscriber().Unsubscribe(_updateChannel, OnChannelUpdate);
+		}
+
+		/// <summary>
+		/// Callback for a message being posted to a channel
+		/// </summary>
+		/// <param name="channel"></param>
+		/// <param name="value"></param>
+		void OnChannelUpdate(RedisChannel channel, RedisValue value)
+		{
+			TaskCompletionSource<bool>? completionSource;
+			lock (_channelWakeEvents)
+			{
+				_channelWakeEvents.TryGetValue(value.ToString(), out completionSource);
+			}
+			if (completionSource != null)
+			{
+				completionSource.TrySetResult(true);
+			}
+		}
+
+		/// <summary>
+		/// Gets the set of messages for a channel
+		/// </summary>
+		/// <param name="channelId"></param>
+		/// <returns></returns>
+		RedisListKey<T> GetChannel(string channelId)
+		{
+			return new RedisListKey<T>(_keyPrefix.Append(channelId));
+		}
+
+		/// <inheritdoc/>
+		public async Task PostAsync(string channelId, T message)
+		{
+			RedisListKey<T> channel = GetChannel(channelId);
+
+			long length = await _redis.ListRightPushAsync(channel, message);
+			if (length == 1)
+			{
+				await _redis.PublishAsync(_updateChannel, channelId, CommandFlags.FireAndForget);
+			}
+			await _redis.KeyExpireAsync(channel, ExpireTime, flags: CommandFlags.FireAndForget);
+		}
+
+		async Task<bool> ReadMessagesAsync(RedisListKey<T> list, List<T> messages)
+		{
+			T? message = await _redis.ListLeftPopAsync(list);
+			while(message != null)
+			{
+				messages.Add(message);
+				message = await _redis.ListLeftPopAsync(list);
+			}
+			return messages.Count > 0;
+		}
+
+		/// <inheritdoc/>
+		public async Task<List<T>> ReadMessagesAsync(string channelId)
+		{
+			List<T> messages = new List<T>();
+			await ReadMessagesAsync(GetChannel(channelId), messages);
+			return messages;
+		}
+
+		/// <inheritdoc/>
+		public async Task<List<T>> WaitForMessagesAsync(string channelId, CancellationToken cancellationToken)
+		{
+			List<T> messages = new List<T>();
+
+			RedisListKey<T> channel = GetChannel(channelId);
+			while (!cancellationToken.IsCancellationRequested && !await ReadMessagesAsync(channel, messages))
+			{
+				// Register for notifications on this channel
+				TaskCompletionSource<bool>? completionSource;
+				lock (_channelWakeEvents)
+				{
+					if (!_channelWakeEvents.TryGetValue(channelId, out completionSource))
+					{
+						completionSource = new TaskCompletionSource<bool>();
+						_channelWakeEvents.Add(channelId, completionSource);
+					}
+				}
+
+				try
+				{
+					// Read the current queue state again, in case it was modified since we registered.
+					if(await ReadMessagesAsync(channel, messages))
+					{
+						break;
+					}
+
+					// Wait for messages to be available
+					using (cancellationToken.Register(() => completionSource.TrySetResult(false)))
+					{
+						await completionSource.Task;
+					}
+				}
+				finally
+				{
+					lock (_channelWakeEvents)
+					{
+						_channelWakeEvents.Remove(channelId);
+					}
+				}
+			}
+
+			return messages;
+		}
+	}
+}
diff --git a/Engine/Source/Programs/Horde/Horde.Build/Compute/RedisTaskScheduler.cs b/Engine/Source/Programs/Horde/Horde.Build/Compute/RedisTaskScheduler.cs
new file mode 100644
--- /dev/null
+++ b/Engine/Source/Programs/Horde/Horde.Build/Compute/RedisTaskScheduler.cs
@@ -0,0 +1,485 @@
+// Copyright Epic Games, Inc. All Rights Reserved.
+
+using System;
+using System.Collections.Generic;
+using System.Diagnostics;
+using System.Linq;
+using System.Threading;
+using System.Threading.Channels;
+using System.Threading.Tasks;
+using EpicGames.Core;
+using EpicGames.Redis;
+using Microsoft.Extensions.Logging;
+using StackExchange.Redis;
+
+namespace Horde.Build.Compute
+{
+	using Condition = StackExchange.Redis.Condition;
+
+	/// <summary>
+	/// Interface for a queue of compute operations, each of which may have different requirements for the executing machines.
+	/// </summary>
+	/// <typeparam name="TQueueId">Type used to identify a particular queue</typeparam>
+	/// <typeparam name="TTask">Type used to describe a task to be performed</typeparam>
+	interface ITaskScheduler<TQueueId, TTask> 
+		where TQueueId : notnull 
+		where TTask : class
+	{
+		/// <summary>
+		/// Adds a task to a queue
+		/// </summary>
+		/// <param name="queueId">The queue identifier</param>
+		/// <param name="taskId">The task to add</param>
+		/// <param name="atFront">Whether to insert at the front of the queue</param>
+		Task EnqueueAsync(TQueueId queueId, TTask taskId, bool atFront);
+
+		/// <summary>
+		/// Dequeue any task from a particular queue
+		/// </summary>
+		/// <param name="queueId">The queue to remove a task from</param>
+		/// <returns>Information about the task to be executed</returns>
+		Task<TTask?> DequeueAsync(TQueueId queueId);
+
+		/// <summary>
+		/// Dequeues a task that the given agent can execute
+		/// </summary>
+		/// <param name="predicate">Predicate for determining which queues can be removed from</param>
+		/// <param name="token">Cancellation token for the operation. Will return a null entry rather than throwing an exception.</param>
+		/// <returns>Information about the task to be executed</returns>
+		Task<Task<(TQueueId, TTask)?>> DequeueAsync(Func<TQueueId, ValueTask<bool>> predicate, CancellationToken token = default);
+
+		/// <summary>
+		/// Gets hashes of all the inactive task queues
+		/// </summary>
+		/// <returns></returns>
+		Task<List<TQueueId>> GetInactiveQueuesAsync();
+		
+		/// <summary>
+		/// Get number of tasks that a given pool/agent can execute
+		/// A read-only operation and will not affect any queue.
+		/// </summary>
+		/// <param name="predicate">Predicate for determining which queues to read</param>
+		/// <param name="token">Cancellation token for the operation</param>
+		/// <returns>Number of tasks matching the supplied predicate</returns>
+		Task<int> GetNumQueuedTasksAsync(Func<TQueueId, ValueTask<bool>> predicate, CancellationToken token = default);
+	}
+
+	/// <summary>
+	/// Implementation of <see cref="ITaskScheduler{TQueue, TTask}"/> using Redis for storage
+	/// </summary>
+	/// <typeparam name="TQueueId">Type used to identify a particular queue</typeparam>
+	/// <typeparam name="TTask">Type used to describe a task to be performed</typeparam>
+	class RedisTaskScheduler<TQueueId, TTask> : ITaskScheduler<TQueueId, TTask>, IDisposable 
+		where TQueueId : notnull
+		where TTask : class
+	{
+		class Listener
+		{
+			public Func<TQueueId, ValueTask<bool>> _predicate;
+			public TaskCompletionSource<(TQueueId, TTask)?> CompletionSource { get; }
+
+			public Listener(Func<TQueueId, ValueTask<bool>> predicate)
+			{
+				_predicate = predicate;
+				CompletionSource = new TaskCompletionSource<(TQueueId, TTask)?>();
+			}
+		}
+
+		readonly RedisConnectionPool _redisConnectionPool;
+		readonly RedisKey _baseKey;
+		readonly RedisSetKey<TQueueId> _queueIndex;
+		readonly RedisHashKey<TQueueId, DateTime> _activeQueues; // Queues which are actively being dequeued from
+		ReadOnlyHashSet<TQueueId> _localActiveQueues = new HashSet<TQueueId>();
+		readonly Stopwatch _resetActiveQueuesTimer = Stopwatch.StartNew();
+		readonly List<Listener> _listeners = new List<Listener>();
+		readonly RedisChannel<TQueueId> _newQueueChannel;
+		readonly Task _queueUpdateTask;
+		readonly CancellationTokenSource _cancellationSource = new CancellationTokenSource();
+		readonly ILogger _logger;
+
+		/// <summary>
+		/// Constructor
+		/// </summary>
+		/// <param name="redisConnectionPool">The Redis connection pool</param>
+		/// <param name="baseKey">Base key for all keys used by this scheduler</param>
+		/// <param name="logger"></param>
+		public RedisTaskScheduler(RedisConnectionPool redisConnectionPool, RedisKey baseKey, ILogger logger)
+		{
+			_redisConnectionPool = redisConnectionPool;
+			_baseKey = baseKey;
+			_queueIndex = new RedisSetKey<TQueueId>(baseKey.Append("index"));
+			_activeQueues = new RedisHashKey<TQueueId, DateTime>(baseKey.Append("active"));
+			_newQueueChannel = new RedisChannel<TQueueId>(baseKey.Append("new_queues").ToString());
+			_logger = logger;
+
+			_queueUpdateTask = Task.Run(() => UpdateQueuesAsync(_cancellationSource.Token));
+		}
+
+		/// <inheritdoc/>
+		public async ValueTask DisposeAsync()
+		{
+			if (!_queueUpdateTask.IsCompleted)
+			{
+				_cancellationSource.Cancel();
+				try
+				{
+					await _queueUpdateTask;
+				}
+				catch (OperationCanceledException)
+				{
+				}
+			}
+			_cancellationSource.Dispose();
+		}
+
+		/// <inheritdoc/>
+		public void Dispose()
+		{
+			DisposeAsync().AsTask().Wait();
+		}
+
+		/// <summary>
+		/// Gets the key for a list of tasks for a particular queue
+		/// </summary>
+		/// <param name="queueId">The queue identifier</param>
+		/// <returns></returns>
+		RedisListKey<TTask> GetQueueKey(TQueueId queueId)
+		{
+			return new RedisListKey<TTask>(_baseKey.Append(RedisSerializer.Serialize(queueId).AsKey()));
+		}
+
+		/// <summary>
+		/// Pushes a task onto either end of a queue
+		/// </summary>
+		static Task<long> PushTaskAsync(IDatabaseAsync database, RedisListKey<TTask> list, TTask task, When when, CommandFlags flags, bool atFront)
+		{
+			if (atFront)
+			{
+				return database.ListLeftPushAsync(list, task, when, flags);
+			}
+			else
+			{
+				return database.ListRightPushAsync(list, task, when, flags);
+			}
+		}
+
+		/// <summary>
+		/// Adds a task to a particular queue, creating and adding that queue to the index if necessary
+		/// </summary>
+		/// <param name="queueId">The queue to add the task to</param>
+		/// <param name="task">Task to be scheduled</param>
+		/// <param name="atFront">Whether to add to the front of the queue</param>
+		public async Task EnqueueAsync(TQueueId queueId, TTask task, bool atFront)
+		{
+			RedisListKey<TTask> list = GetQueueKey(queueId);
+			for (; ; )
+			{
+				IDatabase connection = _redisConnectionPool.GetDatabase();
+
+				long newLength = await PushTaskAsync(connection, list, task, When.Exists, CommandFlags.None, atFront);
+				if (newLength > 0)
+				{
+					_logger.LogInformation("Length of queue {QueueId} is {Length}", queueId, newLength);
+					break;
+				}
+
+				ITransaction transaction = connection.CreateTransaction();
+				_ = transaction.SetAddAsync(_queueIndex, queueId, CommandFlags.FireAndForget);
+				_ = PushTaskAsync(transaction, list, task, When.Always, CommandFlags.FireAndForget, atFront);
+
+				if (await transaction.ExecuteAsync())
+				{
+					_logger.LogInformation("Created queue {QueueId}", queueId);
+					await connection.PublishAsync(_newQueueChannel, queueId);
+					break;
+				}
+
+				_logger.LogDebug("EnqueueAsync() retrying...");
+			}
+		}
+
+		/// <summary>
+		/// Dequeues an item for execution by the given agent
+		/// </summary>
+		/// <param name="predicate">Predicate for queues that tasks can be removed from</param>
+		/// <param name="token">Cancellation token for waiting for an item</param>
+		/// <returns>The dequeued item, or null if no item is available</returns>
+		public async Task<Task<(TQueueId, TTask)?>> DequeueAsync(Func<TQueueId, ValueTask<bool>> predicate, CancellationToken token = default)
+		{
+			// Compare against all the list of cached queues to see if we can dequeue something from any of them
+			TQueueId[] queues = await _redisConnectionPool.GetDatabase().SetMembersAsync(_queueIndex);
+
+			// Try to dequeue an item from the list
+			(TQueueId, TTask)? entry = await TryAssignToLocalAgentAsync(queues, predicate);
+			if (entry != null)
+			{
+				return Task.FromResult(entry);
+			}
+
+			// Otherwise create a new task to do the wait
+			return WaitToDequeueAsync(queues, predicate, token);
+		}
+
+		private async Task<(TQueueId, TTask)?> WaitToDequeueAsync(TQueueId[] queues, Func<TQueueId, ValueTask<bool>> predicate, CancellationToken token)
+		{
+			Listener listener = new Listener(predicate);
+			try
+			{
+				// Add the listener to the global list
+				lock (_listeners)
+				{
+					_listeners.Add(listener);
+				}
+
+				// Try to dequeue an item from the list again. An item may have been made available 
+				(TQueueId, TTask)? entry = await TryAssignToLocalAgentAsync(queues, predicate);
+				if (entry != null)
+				{
+					if (!listener.CompletionSource.TrySetResult(entry.Value))
+					{
+						(TQueueId queue, TTask task) = entry.Value;
+						await EnqueueAsync(queue, task, true);
+					}
+				}
+
+				// Wait for an item to be available
+				using (IDisposable registration = token.Register(() => listener.CompletionSource.TrySetResult(null)))
+				{
+					return await listener.CompletionSource.Task;
+				}
+			}
+			finally
+			{
+				lock (_listeners)
+				{
+					_listeners.Remove(listener);
+				}
+			}
+		}
+
+		/// <summary>
+		/// Attempts to dequeue a task from a set of queue
+		/// </summary>
+		/// <param name="queueIds">The current array of queues</param>
+		/// <param name="predicate">Predicate for queues that tasks can be removed from</param>
+		/// <returns>The dequeued item, or null if no item is available</returns>
+		async Task<(TQueueId, TTask)?> TryAssignToLocalAgentAsync(TQueueId[] queueIds, Func<TQueueId, ValueTask<bool>> predicate)
+		{
+			foreach (TQueueId queueId in queueIds)
+			{
+				if (await predicate(queueId))
+				{
+					TTask? task = await DequeueAsync(queueId);
+					if (task != null)
+					{
+						return (queueId, task);
+					}
+				}
+			}
+			return null;
+		}
+
+		/// <summary>
+		/// Dequeues the item at the front of a queue
+		/// </summary>
+		/// <param name="queueId">The queue to dequeue from</param>
+		/// <returns>The dequeued item, or null if the queue is empty</returns>
+		public async Task<TTask?> DequeueAsync(TQueueId queueId)
+		{
+			await AddActiveQueue(queueId);
+
+			IDatabase database = _redisConnectionPool.GetDatabase();
+
+			TTask? item = await database.ListLeftPopAsync(GetQueueKey(queueId));
+			if (item == null)
+			{
+				ITransaction transaction = database.CreateTransaction();
+				transaction.AddCondition(Condition.KeyNotExists(GetQueueKey(queueId).Inner));
+				Task<bool> wasRemoved = transaction.SetRemoveAsync(_queueIndex, queueId);
+
+				if (await transaction.ExecuteAsync() && await wasRemoved)
+				{
+					_logger.LogInformation("Removed queue {QueueId} from index", queueId);
+				}
+			}
+
+			return item;
+		}
+
+		/// <summary>
+		/// Marks a queue key as being actively monitored, preventing it being returned by <see cref="GetInactiveQueuesAsync"/>
+		/// </summary>
+		/// <param name="queueId">The queue key</param>
+		/// <returns></returns>
+		async ValueTask AddActiveQueue(TQueueId queueId)
+		{
+			// Periodically clear out the set of active keys
+			TimeSpan resetTime = TimeSpan.FromSeconds(10.0);
+			if (_resetActiveQueuesTimer.Elapsed > resetTime)
+			{
+				lock (_resetActiveQueuesTimer)
+				{
+					if (_resetActiveQueuesTimer.Elapsed > resetTime)
+					{
+						_localActiveQueues = new HashSet<TQueueId>();
+						_resetActiveQueuesTimer.Restart();
+					}
+				}
+			}
+
+			// Check if the set of active keys already contains the key we're adding. In order to optimize the 
+			// common case under heavy load where the key is in the set, updating it creates a full copy of it. Any
+			// readers can thus access it without the need for any locking.
+			for(; ;)
+			{
+				ReadOnlyHashSet<TQueueId> localActiveQueuesCopy = _localActiveQueues;
+				if (localActiveQueuesCopy.Contains(queueId))
+				{
+					break;
+				}
+
+				HashSet<TQueueId> newLocalActiveQueues = new HashSet<TQueueId>(localActiveQueuesCopy);
+				newLocalActiveQueues.Add(queueId);
+
+				if (Interlocked.CompareExchange(ref _localActiveQueues, newLocalActiveQueues, localActiveQueuesCopy) == localActiveQueuesCopy)
+				{
+					_logger.LogInformation("Refreshing active queue {QueueId}", queueId);
+					await _redisConnectionPool.GetDatabase().HashSetAsync(_activeQueues, queueId, DateTime.UtcNow);
+					break;
+				}
+			}
+		}
+
+		/// <summary>
+		/// Find any inactive keys
+		/// </summary>
+		/// <returns></returns>
+		public async Task<List<TQueueId>> GetInactiveQueuesAsync()
+		{
+			HashSet<TQueueId> keys = new HashSet<TQueueId>(await _redisConnectionPool.GetDatabase().SetMembersAsync(_queueIndex));
+			HashSet<TQueueId> invalidKeys = new HashSet<TQueueId>();
+
+			DateTime minTime = DateTime.UtcNow - TimeSpan.FromMinutes(10.0);
+
+			HashEntry<TQueueId, DateTime>[] entries = await _redisConnectionPool.GetDatabase().HashGetAllAsync(_activeQueues);
+			foreach (HashEntry<TQueueId, DateTime> entry in entries)
+			{
+				if (entry.Value < minTime)
+				{
+					invalidKeys.Add(entry.Name);
+				}
+				else
+				{
+					keys.Remove(entry.Name);
+				}
+			}
+
+			if (invalidKeys.Count > 0)
+			{
+				await _redisConnectionPool.GetDatabase().HashDeleteAsync(_activeQueues, invalidKeys.ToArray());
+			}
+
+			return keys.ToList();
+		}
+
+		public async Task<int> GetNumQueuedTasksAsync(Func<TQueueId, ValueTask<bool>> predicate, CancellationToken token = default)
+		{
+			HashSet<TQueueId> queueIds = new (await _redisConnectionPool.GetDatabase().SetMembersAsync(_queueIndex));
+			long totalTaskCount = 0;
+			foreach (TQueueId queueId in queueIds)
+			{
+				if (await predicate(queueId))
+				{
+					totalTaskCount += await _redisConnectionPool.GetDatabase().ListLengthAsync(GetQueueKey(queueId));
+				}
+			}
+
+			return (int)totalTaskCount;
+		}
+
+		async Task UpdateQueuesAsync(CancellationToken cancellationToken)
+		{
+			Channel<TQueueId> newQueues = Channel.CreateUnbounded<TQueueId>();
+
+			await using RedisSubscription _ = await _redisConnectionPool.GetDatabase().Multiplexer.SubscribeAsync(_newQueueChannel, v => newQueues.Writer.TryWrite(v));
+
+			while (await newQueues.Reader.WaitToReadAsync(cancellationToken))
+			{
+				HashSet<TQueueId> newQueueIds = new HashSet<TQueueId>();
+				while (newQueues.Reader.TryRead(out TQueueId? queueId))
+				{
+					newQueueIds.Add(queueId);
+				}
+				foreach (TQueueId newQueueId in newQueueIds)
+				{
+					await TryDispatchToNewQueueAsync(newQueueId);
+				}
+			}
+		}
+
+		async Task<bool> TryDispatchToNewQueueAsync(TQueueId queueId)
+		{
+			RedisListKey<TTask> queue = GetQueueKey(queueId);
+
+			// Find a local listener that can execute the work
+			(TQueueId QueueId, TTask TaskId)? entry = null;
+			try
+			{
+				for (; ; )
+				{
+					Listener? listener = null;
+
+					// Look for a listener that can execute the task
+					HashSet<Listener> checkedListeners = new HashSet<Listener>();
+					while(listener == null)
+					{
+						// Find up to 10 listeners we haven't seen before
+						List<Listener> newListeners = new List<Listener>();
+						lock (_listeners)
+						{
+							newListeners.AddRange(_listeners.Where(x => !x.CompletionSource.Task.IsCompleted && checkedListeners.Add(x)).Take(10));
+						}
+						if (newListeners.Count == 0)
+						{
+							return false;
+						}
+
+						// Check predicates for each one against the new queue
+						foreach (Listener newListener in newListeners)
+						{
+							if (await newListener._predicate(queueId))
+							{
+								listener = newListener;
+								break;
+							}
+						}
+					}
+
+					// Pop an entry from the queue
+					if (entry == null)
+					{
+						TTask? task = await DequeueAsync(queueId);
+						if (task == null)
+						{
+							return false;
+						}
+						entry = (queueId, task);
+					}
+
+					// Assign it to the listener
+					if (listener.CompletionSource.TrySetResult(entry.Value))
+					{
+						entry = null;
+					}
+				}
+			}
+			finally
+			{
+				if (entry != null)
+				{
+					await EnqueueAsync(entry.Value.QueueId, entry.Value.TaskId, true);
+				}
+			}
+		}
+	}
+}
diff --git a/Engine/Source/Programs/Horde/Horde.Build/Compute/V1/ComputeController.cs b/Engine/Source/Programs/Horde/Horde.Build/Compute/V1/ComputeController.cs
deleted file mode 100644
diff --git a/Engine/Source/Programs/Horde/Horde.Build/Compute/V1/ComputeService.cs b/Engine/Source/Programs/Horde/Horde.Build/Compute/V1/ComputeService.cs
deleted file mode 100644
diff --git a/Engine/Source/Programs/Horde/Horde.Build/Compute/V1/IComputeService.cs b/Engine/Source/Programs/Horde/Horde.Build/Compute/V1/IComputeService.cs
deleted file mode 100644
diff --git a/Engine/Source/Programs/Horde/Horde.Build/Compute/V1/RedisMessageQueue.cs b/Engine/Source/Programs/Horde/Horde.Build/Compute/V1/RedisMessageQueue.cs
deleted file mode 100644
diff --git a/Engine/Source/Programs/Horde/Horde.Build/Compute/V1/RedisTaskScheduler.cs b/Engine/Source/Programs/Horde/Horde.Build/Compute/V1/RedisTaskScheduler.cs
deleted file mode 100644
diff --git a/Engine/Source/Programs/Horde/Horde.Build/Compute/V2/ComputeControllerV2.cs b/Engine/Source/Programs/Horde/Horde.Build/Compute/V2/ComputeControllerV2.cs
deleted file mode 100644
diff --git a/Engine/Source/Programs/Horde/Horde.Build/Compute/V2/ComputeServiceV2.cs b/Engine/Source/Programs/Horde/Horde.Build/Compute/V2/ComputeServiceV2.cs
deleted file mode 100644
diff --git a/Engine/Source/Programs/Horde/Horde.Build/Horde.Build.csproj b/Engine/Source/Programs/Horde/Horde.Build/Horde.Build.csproj
--- a/Engine/Source/Programs/Horde/Horde.Build/Horde.Build.csproj
+++ b/Engine/Source/Programs/Horde/Horde.Build/Horde.Build.csproj
@@ -96,7 +96,6 @@
     <Protobuf Include="Protos\horde\messages\agent_workspace.proto" ProtoRoot="Protos" />
     <Protobuf Include="Protos\horde\messages\lease.proto" ProtoRoot="Protos" />
     <Protobuf Include="Protos\horde\tasks\compute_task.proto" ProtoRoot="Protos" />
-    <Protobuf Include="Protos\horde\tasks\compute_task_v2.proto" ProtoRoot="Protos" />
     <Protobuf Include="Protos\horde\tasks\conform_task.proto" ProtoRoot="Protos" />
     <Protobuf Include="Protos\horde\tasks\job_task.proto" ProtoRoot="Protos" />
     <Protobuf Include="Protos\horde\tasks\restart_task.proto" ProtoRoot="Protos" />
diff --git a/Engine/Source/Programs/Horde/Horde.Build/Protos/horde/tasks/compute_task.cs b/Engine/Source/Programs/Horde/Horde.Build/Protos/horde/tasks/compute_task.cs
deleted file mode 100644
diff --git a/Engine/Source/Programs/Horde/Horde.Build/Startup.cs b/Engine/Source/Programs/Horde/Horde.Build/Startup.cs
--- a/Engine/Source/Programs/Horde/Horde.Build/Startup.cs
+++ b/Engine/Source/Programs/Horde/Horde.Build/Startup.cs
@@ -95,6 +95,7 @@ using Horde.Build.Server.Notices;
 using Horde.Build.Notifications.Sinks;
 using StatusCode = Grpc.Core.StatusCode;
 using Horde.Build.Artifacts;
+using Horde.Build.Compute;
 
 namespace Horde.Build
 {
@@ -684,18 +685,14 @@ namespace Horde.Build
 				services.AddHostedService<JobTaskSource>(provider => provider.GetRequiredService<JobTaskSource>());
 				services.AddSingleton<ConformTaskSource>();
 				services.AddHostedService<ConformTaskSource>(provider => provider.GetRequiredService<ConformTaskSource>());
-				services.AddSingleton<Compute.V1.ComputeService>();
-				services.AddSingleton<Compute.V1.IComputeService, Compute.V1.ComputeService>();
-				services.AddHostedService(provider => provider.GetRequiredService<Compute.V1.ComputeService>());
-				services.AddSingleton<Compute.V2.ComputeServiceV2>();
+				services.AddSingleton<ComputeService>();
 
 				services.AddSingleton<ITaskSource, UpgradeTaskSource>();
 				services.AddSingleton<ITaskSource, ShutdownTaskSource>();
 				services.AddSingleton<ITaskSource, RestartTaskSource>();
 				services.AddSingleton<ITaskSource, ConformTaskSource>(provider => provider.GetRequiredService<ConformTaskSource>());
 				services.AddSingleton<ITaskSource, JobTaskSource>(provider => provider.GetRequiredService<JobTaskSource>());
-				services.AddSingleton<ITaskSource, Compute.V1.ComputeService>();
-				services.AddSingleton<ITaskSource, Compute.V2.ComputeServiceV2>(provider => provider.GetRequiredService<Compute.V2.ComputeServiceV2>());
+				services.AddSingleton<ITaskSource, ComputeTaskSource>();
 			}
 
 			// Allow longer to shutdown so we can debug missing cancellation tokens
diff --git a/Engine/Source/Programs/Shared/EpicGames.Horde/Compute/ChannelId.cs b/Engine/Source/Programs/Shared/EpicGames.Horde/Compute/ChannelId.cs
deleted file mode 100644
diff --git a/Engine/Source/Programs/Shared/EpicGames.Horde/Compute/ComputeClient.cs b/Engine/Source/Programs/Shared/EpicGames.Horde/Compute/ComputeClient.cs
deleted file mode 100644
diff --git a/Engine/Source/Programs/Shared/EpicGames.Horde/Compute/ComputeExtensions.cs b/Engine/Source/Programs/Shared/EpicGames.Horde/Compute/ComputeExtensions.cs
deleted file mode 100644
diff --git a/Engine/Source/Programs/Shared/EpicGames.Horde/Compute/ComputeMessage.cs b/Engine/Source/Programs/Shared/EpicGames.Horde/Compute/ComputeMessage.cs
--- a/Engine/Source/Programs/Shared/EpicGames.Horde/Compute/ComputeMessage.cs
+++ b/Engine/Source/Programs/Shared/EpicGames.Horde/Compute/ComputeMessage.cs
@@ -1,8 +1,12 @@
 // Copyright Epic Games, Inc. All Rights Reserved.
 
 using System;
+using System.Buffers;
 using System.Collections.Generic;
+using System.IO;
 using System.Reflection;
+using System.Threading;
+using System.Threading.Tasks;
 using EpicGames.Core;
 using EpicGames.Serialization;
 
@@ -28,6 +32,47 @@ namespace EpicGames.Horde.Compute
 		}
 	}
 
+	/// <summary>
+	/// Utility class for writing cb messages to a stream
+	/// </summary>
+	public sealed class ComputeMessageWriter
+	{
+		readonly IComputeChannel _channel;
+		readonly CbWriter _writer;
+		byte[] _buffer;
+
+		/// <summary>
+		/// Constructor
+		/// </summary>
+		/// <param name="channel"></param>
+		/// <param name="reserve">Amount of data to reserve for serialization</param>
+		public ComputeMessageWriter(IComputeChannel channel, int reserve = CbWriter.DefaultChunkSize)
+		{
+			_channel = channel;
+			_writer = new CbWriter(reserve);
+			_buffer = new byte[reserve];
+		}
+
+		/// <summary>
+		/// Serializes a message to the channel
+		/// </summary>
+		/// <param name="message">Message to write</param>
+		/// <param name="cancellationToken">Cancellation token for the operation</param>
+		public async Task WriteCbMessageAsync(object message, CancellationToken cancellationToken)
+		{
+			ComputeMessage.Serialize(message, _writer);
+
+			int size = _writer.GetSize();
+			if (size > _buffer.Length)
+			{
+				_buffer = new byte[size + 1024];
+			}
+			_writer.CopyTo(_buffer);
+
+			await _channel.WriteMessageAsync(_buffer.AsMemory(0, size), cancellationToken);
+		}
+	}
+
 	/// <summary>
 	/// Utility methods for compute messages
 	/// </summary>
@@ -40,7 +85,12 @@ namespace EpicGames.Horde.Compute
 		{
 			typeof(CloseMessage),
 			typeof(XorRequestMessage),
-			typeof(XorResponseMessage)
+			typeof(XorResponseMessage),
+			typeof(CppComputeMessage),
+			typeof(CppComputeOutputMessage),
+			typeof(CppComputeFinishMessage),
+			typeof(BlobReadMessage),
+			typeof(BlobDataMessage)
 		};
 
 		static readonly Dictionary<Utf8String, ICbConverter> s_typeToConverter = CreateTypeToConverterMap();
@@ -121,5 +171,57 @@ namespace EpicGames.Horde.Compute
 			CbField dataField = obj.Find(s_dataField);
 			return converter.ReadObject(dataField)!;
 		}
+
+		/// <summary>
+		/// Receives a message from the remote
+		/// </summary>
+		/// <param name="channel">Channel to read from</param>
+		/// <param name="cancellationToken"></param>
+		/// <returns></returns>
+		public static async Task<object?> ReadCbMessageAsync(this IComputeChannel channel, CancellationToken cancellationToken)
+		{
+			ReadOnlyMemory<byte> message = await channel.ReadMessageAsync(cancellationToken);
+			if (message.Length == 0)
+			{
+				return null;
+			}
+			return ComputeMessage.Deserialize(message.ToArray());
+		}
+
+		/// <summary>
+		/// Receives a message from the remote
+		/// </summary>
+		/// <param name="channel">Channel to read from</param>
+		/// <param name="cancellationToken"></param>
+		/// <returns></returns>
+		public static async Task<T> ReadCbMessageAsync<T>(this IComputeChannel channel, CancellationToken cancellationToken)
+		{
+			object? message = await ReadCbMessageAsync(channel, cancellationToken);
+			if (message == null)
+			{
+				throw new EndOfStreamException();
+			}
+			return (T)message;
+		}
+
+		/// <summary>
+		/// Receives a message from the remote
+		/// </summary>
+		/// <param name="channel">Channel to read from</param>
+		/// <param name="message">Message to write</param>
+		/// <param name="cancellationToken">Cancellation token for the operation</param>
+		/// <returns></returns>
+		public static async Task WriteCbMessageAsync(this IComputeChannel channel, object message, CancellationToken cancellationToken)
+		{
+			CbWriter writer = new CbWriter();
+			Serialize(message, writer);
+
+			int size = writer.GetSize();
+			using (IMemoryOwner<byte> buffer = MemoryPool<byte>.Shared.Rent(size))
+			{
+				writer.CopyTo(buffer.Memory.Span);
+				await channel.WriteMessageAsync(buffer.Memory.Slice(0, size), cancellationToken);
+			}
+		}
 	}
 }
diff --git a/Engine/Source/Programs/Shared/EpicGames.Horde/Compute/ComputeMessageTypes.cs b/Engine/Source/Programs/Shared/EpicGames.Horde/Compute/ComputeMessageTypes.cs
--- a/Engine/Source/Programs/Shared/EpicGames.Horde/Compute/ComputeMessageTypes.cs
+++ b/Engine/Source/Programs/Shared/EpicGames.Horde/Compute/ComputeMessageTypes.cs
@@ -1,7 +1,11 @@
 // Copyright Epic Games, Inc. All Rights Reserved.
 
+using EpicGames.Horde.Storage;
 using EpicGames.Serialization;
 using System;
+using System.IO;
+using System.Threading;
+using System.Threading.Tasks;
 
 #pragma warning disable CS1591
 #pragma warning disable CA1819
@@ -38,4 +42,129 @@ namespace EpicGames.Horde.Compute
 		[CbField]
 		public byte[] Payload { get; set; } = Array.Empty<byte>();
 	}
+
+	/// <summary>
+	/// Requests a C++ compilation
+	/// </summary>
+	[ComputeMessage("cpp")]
+	public class CppComputeMessage
+	{
+		[CbField]
+		public NodeLocator Locator { get; set; }
+	}
+
+	/// <summary>
+	/// Output from a C++ build job
+	/// </summary>
+	[ComputeMessage("cpp-out")]
+	public class CppComputeOutputMessage
+	{
+		[CbField]
+		public NodeLocator Locator { get; set; }
+	}
+
+	/// <summary>
+	/// Indicates that a C++ compute job is complete, and the sandbox can be discarded.
+	/// </summary>
+	[ComputeMessage("cpp-finish"), CbObject]
+	public class CppComputeFinishMessage
+	{
+	}
+
+	/// <summary>
+	/// Reads a blob from the remote
+	/// </summary>
+	[ComputeMessage("blob-read")]
+	public class BlobReadMessage
+	{
+		[CbField("loc")]
+		public BlobLocator Locator { get; set; }
+
+		[CbField("ofs")]
+		public int Offset { get; set; }
+
+		[CbField("len")]
+		public int Length { get; set; }
+	}
+
+	/// <summary>
+	/// Reads a blob from the remote
+	/// </summary>
+	[ComputeMessage("blob-data")]
+	public class BlobDataMessage
+	{
+		[CbField]
+		public ReadOnlyMemory<byte> Data { get; set; }
+	}
+
+	/// <summary>
+	/// Extension methods for sending compute messages
+	/// </summary>
+	public static class ComputeMessageExtensions
+	{
+		/// <summary>
+		/// Reads data for a blob from a remote
+		/// </summary>
+		/// <param name="channel">Channel to write to</param>
+		/// <param name="locator">Locator for the blob to send</param>
+		/// <param name="offset">Starting offset of the data</param>
+		/// <param name="length">Length of the data</param>
+		/// <param name="cancellationToken">Cancellation token for the operation</param>
+		/// <returns>Data for the blob</returns>
+		public static async Task<ReadOnlyMemory<byte>> ReadBlobDataAsync(this IComputeChannel channel, BlobLocator locator, int offset, int length, CancellationToken cancellationToken)
+		{
+			BlobReadMessage message = new BlobReadMessage { Locator = locator, Offset = offset, Length = length };
+			await channel.WriteCbMessageAsync(message, cancellationToken);
+
+			BlobDataMessage response = await channel.ReadCbMessageAsync<BlobDataMessage>(cancellationToken);
+			return response.Data;
+		}
+
+		/// <summary>
+		/// Writes blob data to a compute channel
+		/// </summary>
+		/// <param name="channel">Channel to write to</param>
+		/// <param name="request">The request to read blob data</param>
+		/// <param name="storage">Storage client to retrieve the blob from</param>
+		/// <param name="cancellationToken">Cancellation token for the operation</param>
+		/// <returns></returns>
+		public static async Task WriteBlobDataAsync(this IComputeChannel channel, BlobReadMessage request, IStorageClient storage, CancellationToken cancellationToken)
+		{
+			await WriteBlobDataAsync(channel, request.Locator, request.Offset, request.Length, storage, cancellationToken);
+		}
+
+		/// <summary>
+		/// Writes blob data to a compute channel
+		/// </summary>
+		/// <param name="channel">Channel to write to</param>
+		/// <param name="locator">Locator for the blob to send</param>
+		/// <param name="offset">Starting offset of the data</param>
+		/// <param name="length">Length of the data</param>
+		/// <param name="storage">Storage client to retrieve the blob from</param>
+		/// <param name="cancellationToken">Cancellation token for the operation</param>
+		public static async Task WriteBlobDataAsync(this IComputeChannel channel, BlobLocator locator, int offset, int length, IStorageClient storage, CancellationToken cancellationToken)
+		{
+			BlobDataMessage response = new BlobDataMessage();
+			if (offset == 0 && length == 0)
+			{
+				using (Stream stream = await storage.ReadBlobAsync(locator, cancellationToken))
+				{
+					MemoryStream target = new MemoryStream();
+					await stream.CopyToAsync(target, cancellationToken);
+					response.Data = target.ToArray();
+				}
+			}
+			else
+			{
+				using (Stream stream = await storage.ReadBlobRangeAsync(locator, offset, length, cancellationToken))
+				{
+					MemoryStream target = new MemoryStream();
+					await stream.CopyToAsync(target, cancellationToken);
+					response.Data = target.ToArray();
+				}
+			}
+			await channel.WriteCbMessageAsync(response, cancellationToken);
+		}
+	}
 }
+
diff --git a/Engine/Source/Programs/Shared/EpicGames.Horde/Compute/ComputeStorageClient.cs b/Engine/Source/Programs/Shared/EpicGames.Horde/Compute/ComputeStorageClient.cs
new file mode 100644
--- /dev/null
+++ b/Engine/Source/Programs/Shared/EpicGames.Horde/Compute/ComputeStorageClient.cs
@@ -0,0 +1,103 @@
+// Copyright Epic Games, Inc. All Rights Reserved.
+
+using System;
+using System.Collections.Generic;
+using System.IO;
+using System.Threading;
+using System.Threading.Tasks;
+using EpicGames.Core;
+using EpicGames.Horde.Storage;
+
+namespace EpicGames.Horde.Compute
+{
+	/// <summary>
+	/// Storage client which can read bundles over a compute channel
+	/// </summary>
+	public sealed class ComputeStorageClient : StorageClientBase, IDisposable
+	{
+		readonly IComputeChannel _channel;
+		readonly SemaphoreSlim _semaphore;
+
+		/// <summary>
+		/// Constructor
+		/// </summary>
+		/// <param name="channel"></param>
+		public ComputeStorageClient(IComputeChannel channel)
+		{
+			_channel = channel;
+			_semaphore = new SemaphoreSlim(1);
+		}
+
+		/// <inheritdoc/>
+		public void Dispose()
+		{
+			_semaphore.Dispose();
+		}
+
+		#region Blobs
+
+		/// <inheritdoc/>
+		public override async Task<Stream> ReadBlobAsync(BlobLocator locator, CancellationToken cancellationToken = default)
+		{
+			return await ReadBlobRangeAsync(locator, 0, 0, cancellationToken);
+		}
+
+		/// <inheritdoc/>
+		public override async Task<Stream> ReadBlobRangeAsync(BlobLocator locator, int offset, int length, CancellationToken cancellationToken = default)
+		{
+			await _semaphore.WaitAsync(cancellationToken);
+			try
+			{
+				BlobReadMessage request = new BlobReadMessage { Locator = locator, Offset = offset, Length = length };
+				await _channel.WriteCbMessageAsync(request, cancellationToken);
+
+				BlobDataMessage response = await _channel.ReadCbMessageAsync<BlobDataMessage>(cancellationToken);
+				return new ReadOnlyMemoryStream(response.Data);
+			}
+			finally
+			{
+				_semaphore.Release();
+			}
+		}
+
+		/// <inheritdoc/>
+		public override Task<BlobLocator> WriteBlobAsync(Stream stream, Utf8String prefix = default, CancellationToken cancellationToken = default)
+		{
+			throw new NotSupportedException();
+		}
+
+		#endregion
+
+		#region Nodes
+
+		/// <inheritdoc/>
+		public override IAsyncEnumerable<NodeHandle> FindNodesAsync(Utf8String name, CancellationToken cancellationToken = default)
+		{
+			throw new NotSupportedException();
+		}
+
+		#endregion
+
+		#region Refs
+
+		/// <inheritdoc/>
+		public override Task DeleteRefAsync(RefName name, CancellationToken cancellationToken = default)
+		{
+			throw new NotSupportedException();
+		}
+
+		/// <inheritdoc/>
+		public override Task<NodeHandle?> TryReadRefTargetAsync(RefName name, DateTime cacheTime = default, CancellationToken cancellationToken = default)
+		{
+			throw new NotSupportedException();
+		}
+
+		/// <inheritdoc/>
+		public override Task WriteRefTargetAsync(RefName name, NodeHandle target, RefOptions? options = null, CancellationToken cancellationToken = default)
+		{
+			throw new NotSupportedException();
+		}
+
+		#endregion
+	}
+}
diff --git a/Engine/Source/Programs/Shared/EpicGames.Horde/Compute/ComputeTask.cs b/Engine/Source/Programs/Shared/EpicGames.Horde/Compute/ComputeTask.cs
deleted file mode 100644
diff --git a/Engine/Source/Programs/Shared/EpicGames.Horde/Compute/ComputeTaskResult.cs b/Engine/Source/Programs/Shared/EpicGames.Horde/Compute/ComputeTaskResult.cs
deleted file mode 100644
diff --git a/Engine/Source/Programs/Shared/EpicGames.Horde/Compute/Cpp/CppComputeNode.cs b/Engine/Source/Programs/Shared/EpicGames.Horde/Compute/Cpp/CppComputeNode.cs
new file mode 100644
--- /dev/null
+++ b/Engine/Source/Programs/Shared/EpicGames.Horde/Compute/Cpp/CppComputeNode.cs
@@ -0,0 +1,92 @@
+// Copyright Epic Games, Inc. All Rights Reserved.
+
+using EpicGames.Core;
+using EpicGames.Horde.Storage;
+using EpicGames.Horde.Storage.Nodes;
+using System.Collections.Generic;
+
+namespace EpicGames.Horde.Compute.Cpp
+{
+	/// <summary>
+	/// Describes an action to be executed in a particular workspace
+	/// </summary>
+	[TreeNode("{43D16B98-1AD5-4BB2-8BD9-EF0CA3635BE0}", 1)]
+	public class CppComputeNode : TreeNode
+	{
+		/// <summary>
+		/// The executable to run
+		/// </summary>
+		public string Executable { get; set; }
+
+		/// <summary>
+		/// List of command line arguments for the process to run.
+		/// </summary>
+		public List<string> Arguments { get; } = new List<string>();
+
+		/// <summary>
+		/// Environment variables to set for the child process
+		/// </summary>
+		public Dictionary<string, string> EnvVars { get; } = new Dictionary<string, string>();
+
+		/// <summary>
+		/// Path to the working directory within the workspace
+		/// </summary>
+		public string WorkingDirectory { get; set; }
+
+		/// <summary>
+		/// Reference to a sandbox 
+		/// </summary>
+		public TreeNodeRef<DirectoryNode> Sandbox { get; set; }
+
+		/// <summary>
+		/// List of output paths to be captured on completion of the action. These may be files or directories.
+		/// </summary>
+		public List<string> OutputPaths { get; } = new List<string>();
+
+		/// <summary>
+		/// Constructor
+		/// </summary>
+		/// <param name="executable">The executable to run</param>
+		/// <param name="arguments">Arguments for the executable to run</param>
+		/// <param name="workingDirectory">Working directory for execution</param>
+		/// <param name="sandbox">Hash of the sandbox</param>
+		public CppComputeNode(string executable, List<string> arguments, string workingDirectory, TreeNodeRef<DirectoryNode> sandbox)
+		{
+			Executable = executable;
+			Arguments = arguments;
+			WorkingDirectory = workingDirectory;
+			Sandbox = sandbox;
+		}
+
+		/// <summary>
+		/// Deserialization constructor
+		/// </summary>
+		/// <param name="reader">Reader to deserialize from</param>
+		public CppComputeNode(ITreeNodeReader reader)
+		{
+			Executable = reader.ReadString();
+			Arguments = reader.ReadList(MemoryReaderExtensions.ReadString);
+			EnvVars = reader.ReadDictionary(MemoryReaderExtensions.ReadString, MemoryReaderExtensions.ReadString);
+			WorkingDirectory = reader.ReadString();
+			Sandbox = reader.ReadRef<DirectoryNode>();
+			OutputPaths = reader.ReadList(MemoryReaderExtensions.ReadString);
+		}
+
+		/// <inheritdoc/>
+		public override void Serialize(ITreeNodeWriter writer)
+		{
+			writer.WriteUtf8String(Executable);
+			writer.WriteList(Arguments, MemoryWriterExtensions.WriteString);
+			writer.WriteDictionary(EnvVars, MemoryWriterExtensions.WriteString, MemoryWriterExtensions.WriteString);
+			writer.WriteUtf8String(WorkingDirectory);
+			writer.WriteRef(Sandbox);
+			writer.WriteList(OutputPaths, MemoryWriterExtensions.WriteString);
+		}
+
+		/// <inheritdoc/>
+		public override IEnumerable<TreeNodeRef> EnumerateRefs()
+		{
+			yield return Sandbox;
+		}
+	}
+}
diff --git a/Engine/Source/Programs/Shared/EpicGames.Horde/Compute/DirectoryTree.cs b/Engine/Source/Programs/Shared/EpicGames.Horde/Compute/DirectoryTree.cs
deleted file mode 100644
diff --git a/Engine/Source/Programs/Shared/EpicGames.Horde/Compute/IComputeChannel.cs b/Engine/Source/Programs/Shared/EpicGames.Horde/Compute/IComputeChannel.cs
--- a/Engine/Source/Programs/Shared/EpicGames.Horde/Compute/IComputeChannel.cs
+++ b/Engine/Source/Programs/Shared/EpicGames.Horde/Compute/IComputeChannel.cs
@@ -1,5 +1,6 @@
 // Copyright Epic Games, Inc. All Rights Reserved.
 
+using System;
 using System.Threading;
 using System.Threading.Tasks;
 
@@ -11,35 +12,17 @@ namespace EpicGames.Horde.Compute
 	public interface IComputeChannel
 	{
 		/// <summary>
-		/// Receives a message from the remote
+		/// Reads a message from the channel
 		/// </summary>
-		/// <param name="cancellationToken"></param>
-		/// <returns></returns>
-		Task<object> ReadAsync(CancellationToken cancellationToken);
+		/// <param name="cancellationToken">Cancellation token for the operation</param>
+		/// <returns>Data for the message that was read. This buffer is ephemeral, and may be overwritten on the next read call.</returns>
+		Task<ReadOnlyMemory<byte>> ReadMessageAsync(CancellationToken cancellationToken);
 
 		/// <summary>
-		/// Sends a message to the remote
+		/// Writes a message to the channel
 		/// </summary>
-		/// <param name="message">Message to be sent</param>
-		/// <param name="cancellationToken">Cancellation token</param>
-		/// <returns></returns>
-		Task WriteAsync(object message, CancellationToken cancellationToken);
-	}
-
-	/// <summary>
-	/// Extension methods for compute channels
-	/// </summary>
-	public static class ComputeChannelExtensions
-	{
-		/// <summary>
-		/// Receives a message from the remote
-		/// </summary>
-		/// <param name="channel">Channel to read from</param>
-		/// <param name="cancellationToken"></param>
-		/// <returns></returns>
-		public static async Task<T> ReadAsync<T>(this IComputeChannel channel, CancellationToken cancellationToken)
-		{
-			return (T)await channel.ReadAsync(cancellationToken);
-		}
+		/// <param name="buffer">Buffer containing the data to be written</param>
+		/// <param name="cancellationToken">Cancellation token for the operation</param>
+		Task WriteMessageAsync(ReadOnlyMemory<byte> buffer, CancellationToken cancellationToken);
 	}
 }
diff --git a/Engine/Source/Programs/Shared/EpicGames.Horde/Compute/IComputeClient.cs b/Engine/Source/Programs/Shared/EpicGames.Horde/Compute/IComputeClient.cs
--- a/Engine/Source/Programs/Shared/EpicGames.Horde/Compute/IComputeClient.cs
+++ b/Engine/Source/Programs/Shared/EpicGames.Horde/Compute/IComputeClient.cs
@@ -1,408 +1,22 @@
 // Copyright Epic Games, Inc. All Rights Reserved.
 
-using EpicGames.Core;
-using EpicGames.Horde.Storage;
-using EpicGames.Serialization;
 using System;
-using System.Collections.Generic;
 using System.Threading;
 using System.Threading.Tasks;
 
 namespace EpicGames.Horde.Compute
 {
 	/// <summary>
-	/// Response decribing a cluster
+	/// Interface for uploading compute work to remote machines
 	/// </summary>
-	public interface IComputeClusterInfo
+	public interface IComputeClient : IAsyncDisposable
 	{
 		/// <summary>
-		/// Id of the cluster
+		/// Adds a new remote request
 		/// </summary>
-		public ClusterId Id { get; }
-
-		/// <summary>
-		/// Namespace containing the requests and responses
-		/// </summary>
-		public NamespaceId NamespaceId { get; }
-
-		/// <summary>
-		/// The request bucket id
-		/// </summary>
-		public BucketId RequestBucketId { get; }
-
-		/// <summary>
-		/// Bucket to store response refs
-		/// </summary>
-		public BucketId ResponseBucketId { get; }
-	}
-
-	/// <summary>
-	/// State of a compute task
-	/// </summary>
-	public enum ComputeTaskState
-	{
-		/// <summary>
-		/// The task is queued for execution
-		/// </summary>
-		Queued = 0,
-
-		/// <summary>
-		/// Currently being executed by a remote agent
-		/// </summary>
-		Executing = 1,
-
-		/// <summary>
-		/// Completed
-		/// </summary>
-		Complete = 2,
-	}
-
-	/// <summary>
-	/// Outcome of a task, if it's in the complete state
-	/// </summary>
-	public enum ComputeTaskOutcome
-	{
-		/// <summary>
-		/// The task was executed successfully
-		/// </summary>
-		Success = 0,
-
-		/// <summary>
-		/// The lease failed with an error
-		/// </summary>
-		Failed = 1,
-
-		/// <summary>
-		/// The lease was cancelled
-		/// </summary>
-		Cancelled = 2,
-
-		/// <summary>
-		/// The task ran but did not return a result for unknown reasons
-		/// </summary>
-		NoResult = 3,
-
-		/// <summary>
-		/// The item was not scheduled for execution in the default time period, and was expired.
-		/// </summary>
-		Expired = 4,
-
-		/// <summary>
-		/// A blob could not be found. The detail field includes the hash of the blob.
-		/// </summary>
-		BlobNotFound = 5,
-
-		/// <summary>
-		/// An uncaught exception occurred during task execution
-		/// </summary>
-		Exception = 6,
-	}
-
-	/// <summary>
-	/// Supplies information about the current execution state of a task
-	/// </summary>
-	public class ComputeTaskStatus
-	{
-		/// <summary>
-		/// Reference to the task decriptor that was requested
-		/// </summary>
-		[CbField("h")]
-		public RefId TaskRefId { get; set; }
-
-		/// <summary>
-		/// Time that the event happened
-		/// </summary>
-		[CbField("t")]
-		public DateTime Time { get; set; }
-
-		/// <summary>
-		/// New state of the task
-		/// </summary>
-		[CbField("s")]
-		public ComputeTaskState State { get; set; }
-
-		/// <summary>
-		/// Outcome of the task execution.Note that this reflects the outcome of the execution rather than the outcome of the task.        
-		/// </summary>
-		[CbField("o")]
-		public ComputeTaskOutcome Outcome { get; set; }
-
-		/// <summary>
-		/// Additional information about the outcome of the task.Dependent on the value of Outcome.        
-		/// </summary>
-		[CbField("d")]
-		public string? Detail { get; set; }
-
-		/// <summary>
-		/// Ref containing the task result
-		/// </summary>
-		[CbField("r")]
-		public RefId? ResultRefId { get; set; }
-
-		/// <summary>
-		/// When transitioning to the executing state, includes the name of the agent performing the work
-		/// </summary>
-		[CbField("a")]
-		public string? AgentId { get; set; }
-
-		/// <summary>
-		/// When transitioning to the executing state, includes the id of the lease assigned to the agent
-		/// </summary>
-		[CbField("l")]
-		public string? LeaseId { get; set; }
-
-		/// <summary>
-		/// Stats about how the task was queued
-		/// </summary>
-		[CbField("qs")]
-		public ComputeTaskQueueStats? QueueStats { get; set; }
-
-		/// <summary>
-		/// Stats for execution of this job
-		/// </summary>
-		[CbField("es")]
-		public ComputeTaskExecutionStats? ExecutionStats { get; set; }
-	}
-
-	/// <summary>
-	/// Stats for queueing a task
-	/// </summary>
-	public class ComputeTaskQueueStats
-	{
-		/// <summary>
-		/// Names of scopes 
-		/// </summary>
-		public static readonly string[] ScopeNames =
-		{
-			"dispatched",
-			"complete"
-		};
-
-		private readonly int[] _scopes = new int[ScopeNames.Length];
-
-		/// <summary>
-		/// Time that the task was queued
-		/// </summary>
-		[CbField("t")]
-		public DateTime StartTime { get; set; }
-
-		/// <summary>
-		/// Timing values measured sequentially starting from StartTime, in milliseconds.
-		/// If the order or meaning of fields in this array are changed, the field name should be changed.
-		/// </summary>
-		[CbField("s")]
-#pragma warning disable CA1819 // Properties should not return arrays
-		public int[] Scopes
-#pragma warning restore CA1819 // Properties should not return arrays
-		{
-			get => _scopes;
-			set => value.AsSpan(0, Math.Max(_scopes.Length, value.Length)).CopyTo(_scopes);
-		}
-
-		/// <summary>
-		/// Time taken until the task was dequeued
-		/// </summary>
-		public int QueuedMs
-		{
-			get => _scopes[0];
-			set => _scopes[0] = value;
-		}
-
-		/// <summary>
-		/// Time that the task was executing
-		/// </summary>
-		public int ExecutingMs
-		{
-			get => _scopes[1];
-			set => _scopes[1] = value;
-		}
-
-		/// <summary>
-		/// Constructor
-		/// </summary>
-		public ComputeTaskQueueStats()
-		{
-		}
-
-		/// <summary>
-		/// Constructor
-		/// </summary>
-		public ComputeTaskQueueStats(DateTime startTime, int queuedMs, int executingMs)
-		{
-			StartTime = startTime;
-			QueuedMs = queuedMs;
-			ExecutingMs = executingMs;
-		}
-	}
-
-	/// <summary>
-	/// Stats for executing a task
-	/// </summary>
-	public class ComputeTaskExecutionStats
-	{
-		/// <summary>
-		/// Names of scopes 
-		/// </summary>
-		public static readonly string[] ScopeNames =
-		{
-			"download-ref",
-			"download-input",
-			"execute",
-			"upload-log",
-			"upload-output",
-			"upload-ref",
-		};
-
-		private readonly int[] _scopes = new int[ScopeNames.Length];
-
-		/// <summary>
-		/// Start time for execution
-		/// </summary>
-		[CbField("t")]
-		public DateTime StartTime { get; set; }
-
-		/// <summary>
-		/// Timing values measured sequentially starting from StartTime, in milliseconds. If the order or meaning of fields in this array are changed, the field name should be changed.
-		/// </summary>
-		[CbField("s")]
-#pragma warning disable CA1819 // Properties should not return arrays
-		public int[] Scopes
-#pragma warning restore CA1819 // Properties should not return arrays
-		{
-			get => _scopes;
-			set => value.AsSpan(0, Math.Max(_scopes.Length, value.Length)).CopyTo(_scopes);
-		}
-
-		/// <summary>
-		/// Time taken to download the job specification
-		/// </summary>
-		public int DownloadRefMs
-		{
-			get => _scopes[0];
-			set => _scopes[0] = value;
-		}
-
-		/// <summary>
-		/// Time taken to download and unpack the sandbox
-		/// </summary>
-		public int DownloadInputMs
-		{
-			get => _scopes[1];
-			set => _scopes[1] = value;
-		}
-
-		/// <summary>
-		/// Time taken to execute the task
-		/// </summary>
-		public int ExecMs
-		{
-			get => _scopes[2];
-			set => _scopes[2] = value;
-		}
-
-		/// <summary>
-		/// Time taken to upload log data
-		/// </summary>
-		public int UploadLogMs
-		{
-			get => _scopes[3];
-			set => _scopes[3] = value;
-		}
-
-		/// <summary>
-		/// Time taken to upload output data
-		/// </summary>
-		public int UploadOutputMs
-		{
-			get => _scopes[4];
-			set => _scopes[4] = value;
-		}
-
-		/// <summary>
-		/// Time taken to upload the output ref
-		/// </summary>
-		public int UploadRefMs
-		{
-			get => _scopes[5];
-			set => _scopes[5] = value;
-		}
-
-		/// <summary>
-		/// Constructor
-		/// </summary>
-		public ComputeTaskExecutionStats()
-		{
-		}
-
-		/// <summary>
-		/// Constructor
-		/// </summary>
-		public ComputeTaskExecutionStats(DateTime startTime, int downloadRefMs, int downloadInputMs, int execMs, int uploadLogMs, int uploadOutputMs, int uploadRefMs)
-		{
-			StartTime = startTime;
-
-			DownloadRefMs = downloadRefMs;
-			DownloadInputMs = downloadInputMs;
-			ExecMs = execMs;
-			UploadLogMs = uploadLogMs;
-			UploadOutputMs = uploadOutputMs;
-			UploadRefMs = uploadRefMs;
-		}
-	}
-
-	/// <summary>
-	/// Interface for communicating with a compute server
-	/// </summary>
-	public interface IComputeClient
-	{
-		/// <summary>
-		/// Gets information about a cluster
-		/// </summary>
-		/// <param name="clusterId">The cluster to retreive information on</param>
-		/// <param name="cancellationToken">Cancellation token for the request</param>
-		/// <returns>Cluster information</returns>
-		Task<IComputeClusterInfo> GetClusterInfoAsync(ClusterId clusterId, CancellationToken cancellationToken = default);
-
-		/// <summary>
-		/// Queues a set of tasks for execution
-		/// </summary>
-		/// <param name="clusterId">Cluster for executing the tasks</param>
-		/// <param name="channelId">Indicates an identifier generated by the client to query responses on</param>
-		/// <param name="requirementsHash">Requirements of the task to execute</param>
-		/// <param name="taskRefIds">Refs describing the request</param>
-		/// <param name="skipCacheLookup">Whether to skip the lookup of cached output for this task</param>
-		/// <param name="cancellationToken">Cancellation token for the request</param>
-		/// <returns>The new channel id</returns>
-		Task AddTasksAsync(ClusterId clusterId, ChannelId channelId, IEnumerable<RefId> taskRefIds, IoHash requirementsHash, bool skipCacheLookup, CancellationToken cancellationToken = default);
-
-		/// <summary>
-		/// Read updates from the given remote execution channel
-		/// </summary>
-		/// <param name="clusterId">Cluster for executing the tasks</param>
-		/// <param name="channelId">Channel to receive updates on</param>
-		/// <param name="cancellationToken">Cancellation token for the operation</param>
-		/// <returns></returns>
-		IAsyncEnumerable<ComputeTaskStatus> GetTaskUpdatesAsync(ClusterId clusterId, ChannelId channelId, CancellationToken cancellationToken = default);
-	}
-
-	/// <summary>
-	/// Extension methods for <see cref="IComputeClient"/>
-	/// </summary>
-	public static class ComputeClientExtensions
-	{
-		/// <summary>
-		/// Queues a single task for execution
-		/// </summary>
-		/// <param name="computeClient">The compute client instance</param>
-		/// <param name="clusterId">Name of the profile to use for executing the tasks</param>
-		/// <param name="channelId">Channel to receive updates on</param>
-		/// <param name="taskRefId">Ref describing the request</param>
-		/// <param name="requirementsHash">Requirements of the task to execute</param>
-		/// <param name="skipCacheLookup">Whether to skip the lookup of cached output for this task</param>
-		/// <param name="cancellationToken">Cancellation token for the operation</param>
-		public static Task AddTaskAsync(this IComputeClient computeClient, ClusterId clusterId, ChannelId channelId, RefId taskRefId, IoHash requirementsHash, bool skipCacheLookup, CancellationToken cancellationToken = default)
-		{
-			return computeClient.AddTasksAsync(clusterId, channelId, new[] { taskRefId }, requirementsHash, skipCacheLookup, cancellationToken);
-		}
+		/// <param name="clusterId">Cluster to execute the request</param>
+		/// <param name="requirements">Requirements for the agent</param>
+		/// <param name="handler">Handler for the connection</param>
+		public Task<IComputeRequest<TResult>> AddRequestAsync<TResult>(ClusterId clusterId, Requirements? requirements, Func<IComputeChannel, CancellationToken, Task<TResult>> handler);
 	}
 }
diff --git a/Engine/Source/Programs/Shared/EpicGames.Horde/Compute/IComputeRequest.cs b/Engine/Source/Programs/Shared/EpicGames.Horde/Compute/IComputeRequest.cs
new file mode 100644
--- /dev/null
+++ b/Engine/Source/Programs/Shared/EpicGames.Horde/Compute/IComputeRequest.cs
@@ -0,0 +1,22 @@
+// Copyright Epic Games, Inc. All Rights Reserved.
+
+using System.Threading.Tasks;
+
+namespace EpicGames.Horde.Compute
+{
+	/// <summary>
+	/// Interface for a compute request
+	/// </summary>
+	public interface IComputeRequest<TResult>
+	{
+		/// <summary>
+		/// Task to await completion of this request
+		/// </summary>
+		Task<TResult> Result { get; }
+
+		/// <summary>
+		/// Cancel the current request
+		/// </summary>
+		void Cancel();
+	}
+}
diff --git a/Engine/Source/Programs/Shared/EpicGames.Horde/Compute/Impl/HttpComputeClient.cs b/Engine/Source/Programs/Shared/EpicGames.Horde/Compute/Impl/HttpComputeClient.cs
deleted file mode 100644
diff --git a/Engine/Source/Programs/Shared/EpicGames.Horde/Compute/LoopbackComputeClient.cs b/Engine/Source/Programs/Shared/EpicGames.Horde/Compute/LoopbackComputeClient.cs
new file mode 100644
--- /dev/null
+++ b/Engine/Source/Programs/Shared/EpicGames.Horde/Compute/LoopbackComputeClient.cs
@@ -0,0 +1,74 @@
+// Copyright Epic Games, Inc. All Rights Reserved.
+
+using System;
+using System.Net;
+using System.Net.Sockets;
+using System.Security.Cryptography;
+using System.Threading;
+using System.Threading.Tasks;
+
+namespace EpicGames.Horde.Compute
+{
+	/// <summary>
+	/// Implementation of <see cref="IComputeChannel"/> which marshals data over a loopback connection to a method running on a background task.
+	/// </summary>
+	public sealed class LoopbackComputeClient : IComputeClient
+	{
+		class ComputeRequest<TResult> : IComputeRequest<TResult>
+		{
+			public Task<TResult> Result { get; }
+
+			public ComputeRequest(Task<TResult> result) => Result = result;
+			public void Cancel() { }
+		}
+
+		readonly Task _listenerTask;
+		readonly SocketComputeChannel _inner;
+		readonly Socket _socket;
+
+		/// <summary>
+		/// Constructor
+		/// </summary>
+		/// <param name="serverFunc">Callback to process the server side of the connection</param>
+		/// <param name="port">Port to connect on</param>
+		public LoopbackComputeClient(Func<IComputeChannel, CancellationToken, Task> serverFunc, int port = 2000)
+		{
+			using Aes aes = Aes.Create();
+			_listenerTask = RunListenerAsync(aes.Key, aes.IV, serverFunc, port, CancellationToken.None);
+
+			_socket = new Socket(SocketType.Stream, ProtocolType.IP);
+			_socket.Connect(IPAddress.Loopback, port);
+
+			_inner = new SocketComputeChannel(_socket, aes.Key, aes.IV);
+		}
+
+		/// <inheritdoc/>
+		public async ValueTask DisposeAsync()
+		{
+			_inner.Dispose();
+			_socket.Dispose();
+			await _listenerTask;
+		}
+
+		/// <summary>
+		/// Sets up the loopback listener and calls the server method
+		/// </summary>
+		static async Task RunListenerAsync(ReadOnlyMemory<byte> aesKey, ReadOnlyMemory<byte> aesIv, Func<IComputeChannel, CancellationToken, Task> serverFunc, int port, CancellationToken cancellationToken)
+		{
+			using Socket listener = new Socket(SocketType.Stream, ProtocolType.IP);
+			listener.Bind(new IPEndPoint(IPAddress.Loopback, port));
+			listener.Listen();
+
+			using Socket socket = await listener.AcceptAsync(cancellationToken);
+			using SocketComputeChannel channel = new SocketComputeChannel(socket, aesKey, aesIv);
+
+			await serverFunc(channel, cancellationToken);
+		}
+
+		/// <inheritdoc/>
+		public Task<IComputeRequest<TResult>> AddRequestAsync<TResult>(ClusterId clusterId, Requirements? requirements, Func<IComputeChannel, CancellationToken, Task<TResult>> handler)
+		{
+			return Task.FromResult<IComputeRequest<TResult>>(new ComputeRequest<TResult>(handler(_inner, CancellationToken.None)));
+		}
+	}
+}
diff --git a/Engine/Source/Programs/Shared/EpicGames.Horde/Compute/Requirements.cs b/Engine/Source/Programs/Shared/EpicGames.Horde/Compute/Requirements.cs
--- a/Engine/Source/Programs/Shared/EpicGames.Horde/Compute/Requirements.cs
+++ b/Engine/Source/Programs/Shared/EpicGames.Horde/Compute/Requirements.cs
@@ -40,12 +40,22 @@ namespace EpicGames.Horde.Compute
 		/// <summary>
 		/// Construct a requirements object with a condition
 		/// </summary>
-		/// <param name="condition"></param>
+		/// <param name="condition">Condition for matching machines to execute the work</param>
 		public Requirements(Condition? condition)
 		{
 			Condition = condition;
 		}
 
+		/// <summary>
+		/// Serialize this object to bytes
+		/// </summary>
+		public byte[] Serialize()
+		{
+			CbWriter writer = new CbWriter();
+			CbSerializer.Serialize(writer, this);
+			return writer.ToByteArray();
+		}
+
 		/// <inheritdoc/>
 		public override string ToString()
 		{
diff --git a/Engine/Source/Programs/Shared/EpicGames.Horde/Compute/ServerComputeClient.cs b/Engine/Source/Programs/Shared/EpicGames.Horde/Compute/ServerComputeClient.cs
new file mode 100644
--- /dev/null
+++ b/Engine/Source/Programs/Shared/EpicGames.Horde/Compute/ServerComputeClient.cs
@@ -0,0 +1,289 @@
+// Copyright Epic Games, Inc. All Rights Reserved.
+
+using EpicGames.Core;
+using Microsoft.Extensions.Logging;
+using System;
+using System.Collections.Generic;
+using System.Net;
+using System.Net.Http;
+using System.Net.Sockets;
+using System.Security.Cryptography;
+using System.Threading;
+using System.Threading.Tasks;
+
+namespace EpicGames.Horde.Compute
+{
+	/// <summary>
+	/// Helper class to enlist remote resources to perform compute-intensive tasks.
+	/// </summary>
+	public sealed class ServerComputeClient : IComputeClient
+	{
+		const int NonceLength = 64;
+
+		class AddComputeTaskRequest
+		{
+			public Requirements? Requirements { get; set; }
+			public int? RemotePort { get; set; }
+			public string Nonce { get; set; } = String.Empty;
+			public string AesKey { get; set; } = String.Empty;
+			public string AesIv { get; set; } = String.Empty;
+		}
+
+		class AddComputeTasksRequest
+		{
+			public Requirements? Requirements { get; set; }
+			public int RemotePort { get; set; }
+			public List<AddComputeTaskRequest> Tasks { get; set; } = new List<AddComputeTaskRequest>();
+		}
+
+		abstract class RequestInfo : IDisposable
+		{
+			public ReadOnlyMemory<byte> Nonce { get; }
+			public ReadOnlyMemory<byte> AesKey { get; }
+			public ReadOnlyMemory<byte> AesIv { get; }
+
+			protected readonly TaskCompletionSource<IComputeChannel> _computeChannelSource;
+			protected readonly CancellationTokenSource _cancellationSource;
+
+			public RequestInfo(CancellationToken cancellationToken)
+			{
+				Nonce = RandomNumberGenerator.GetBytes(NonceLength);
+
+				using (Aes aes = Aes.Create())
+				{
+					AesKey = aes.Key;
+					AesIv = aes.IV;
+				}
+
+				_computeChannelSource = new TaskCompletionSource<IComputeChannel>();
+				_cancellationSource = CancellationTokenSource.CreateLinkedTokenSource(cancellationToken);
+			}
+
+			public abstract Task RunAsync(IComputeChannel channel);
+
+			public void Cancel()
+			{
+				_computeChannelSource.TrySetCanceled();
+				_cancellationSource.Cancel();
+			}
+
+			public void Dispose()
+			{
+				_cancellationSource.Dispose();
+			}
+		}
+
+		class RequestInfo<TResult> : RequestInfo, IComputeRequest<TResult>
+		{
+			readonly Func<IComputeChannel, CancellationToken, Task<TResult>> _handler;
+
+			public Task<TResult> Result { get; }
+
+			public RequestInfo(Func<IComputeChannel, CancellationToken, Task<TResult>> handler, CancellationToken cancellationToken)
+				: base(cancellationToken)
+			{
+				_handler = handler;
+
+				Result = RunWrapperAsync();
+			}
+
+			public override Task RunAsync(IComputeChannel channel)
+			{
+				_computeChannelSource.SetResult(channel);
+				return Result;
+			}
+
+			async Task<TResult> RunWrapperAsync()
+			{
+				IComputeChannel channel = await _computeChannelSource.Task;
+				return await _handler(channel, _cancellationSource.Token);
+			}
+		}
+
+		class MemoryComparer : IEqualityComparer<ReadOnlyMemory<byte>>
+		{
+			public static MemoryComparer Instance { get; } = new MemoryComparer();
+
+			public bool Equals(ReadOnlyMemory<byte> x, ReadOnlyMemory<byte> y) => x.Span.SequenceEqual(y.Span);
+
+			public int GetHashCode(ReadOnlyMemory<byte> memory)
+			{
+				HashCode hasher = new HashCode();
+				hasher.AddBytes(memory.Span);
+				return hasher.ToHashCode();
+			}
+		}
+
+		readonly object _lockObject = new object();
+		readonly Func<HttpClient> _createHttpClient;
+		readonly CancellationTokenSource _cancellationSource = new CancellationTokenSource();
+		readonly Dictionary<ReadOnlyMemory<byte>, RequestInfo> _requests = new Dictionary<ReadOnlyMemory<byte>, RequestInfo>(MemoryComparer.Instance);
+		TcpListener? _listener;
+		Task _listenTask = Task.CompletedTask;
+		readonly ILogger _logger;
+
+		/// <summary>
+		/// Constructor
+		/// </summary>
+		/// <param name="createHttpClient">Creates an HTTP client with the correct base address for the server</param>
+		/// <param name="logger">Logger for diagnostic messages</param>
+		public ServerComputeClient(Func<HttpClient> createHttpClient, ILogger logger)
+		{
+			_createHttpClient = createHttpClient;
+			_logger = logger;
+		}
+
+		/// <inheritdoc/>
+		public async ValueTask DisposeAsync()
+		{
+			await StopAsync();
+			_cancellationSource.Dispose();
+		}
+
+		/// <inheritdoc/>
+		public void Dispose()
+		{
+			DisposeAsync().AsTask().Wait();
+		}
+
+		/// <summary>
+		/// Starts the client
+		/// </summary>
+		public void Start()
+		{
+			_listener = new TcpListener(IPAddress.IPv6Any, 0);
+			_listener.Start();
+			_listenTask = Task.Run(() => ListenAsync(_listener, _cancellationSource.Token));
+		}
+
+		/// <summary>
+		/// Cancels all running tasks and waits for them to complete
+		/// </summary>
+		public async Task StopAsync()
+		{
+			_cancellationSource.Cancel();
+			try
+			{
+				await _listenTask;
+			}
+			catch (OperationCanceledException)
+			{
+			}
+		}
+
+		/// <summary>
+		/// Adds a new remote request
+		/// </summary>
+		/// <param name="clusterId">Cluster to execute the request</param>
+		/// <param name="requirements">Requirements for the agent</param>
+		/// <param name="handler">Handler for the connection</param>
+		public async Task<IComputeRequest<TResult>> AddRequestAsync<TResult>(ClusterId clusterId, Requirements? requirements, Func<IComputeChannel, CancellationToken, Task<TResult>> handler)
+		{
+			RequestInfo<TResult> requestInfo = new RequestInfo<TResult>(handler, _cancellationSource.Token);
+			await AddRequestAsync(clusterId, requirements, requestInfo);
+			return requestInfo;
+		}
+
+		async Task AddRequestAsync(ClusterId clusterId, Requirements? requirements, RequestInfo requestInfo)
+		{
+			try
+			{
+				HttpClient client = _createHttpClient();
+
+				AddComputeTaskRequest task = new AddComputeTaskRequest();
+				task.Nonce = StringUtils.FormatHexString(requestInfo.Nonce.Span);
+				task.AesKey = StringUtils.FormatHexString(requestInfo.AesKey.Span);
+				task.AesIv = StringUtils.FormatHexString(requestInfo.AesIv.Span);
+
+				AddComputeTasksRequest request = new AddComputeTasksRequest();
+				request.Requirements = requirements;
+				request.RemotePort = ((IPEndPoint)_listener!.LocalEndpoint).Port;
+				request.Tasks.Add(task);
+
+				using (HttpResponseMessage response = await client.PostAsync($"api/v2/compute/{clusterId}", request, _cancellationSource.Token))
+				{
+					response.EnsureSuccessStatusCode();
+					_requests.Add(requestInfo.Nonce, requestInfo);
+				}
+			}
+			catch
+			{
+				requestInfo.Dispose();
+			}
+		}
+
+		async Task ListenAsync(TcpListener listener, CancellationToken cancellationToken)
+		{
+			List<Task> tasks = new List<Task>();
+			try
+			{
+				for (; ; )
+				{
+					TcpClient client = await listener.AcceptTcpClientAsync(cancellationToken);
+					_logger.LogInformation("Received connection from {Remote}", client.Client.RemoteEndPoint);
+
+					Task task = HandleConnectionGuardedAsync(client, cancellationToken);
+					tasks.Add(task);
+				}
+			}
+			finally
+			{
+				await Task.WhenAll(tasks);
+				listener.Stop();
+			}
+		}
+
+		async Task HandleConnectionGuardedAsync(TcpClient client, CancellationToken cancellationToken)
+		{
+			try
+			{
+				await HandleConnectionAsync(client, cancellationToken);
+			}
+			catch (Exception ex)
+			{
+				_logger.LogInformation(ex, "Exception while handling request: {Message}", ex.Message);
+			}
+		}
+
+		async Task HandleConnectionAsync(TcpClient client, CancellationToken cancellationToken)
+		{
+			using TcpClient disposeClient = client;
+			Socket socket = client.Client;
+
+			// Read the nonce for this connection
+			byte[] nonceBuffer = new byte[NonceLength];
+			for (int nonceLength = 0; nonceLength < nonceBuffer.Length;)
+			{
+				int read = await socket.ReceiveAsync(nonceBuffer.AsMemory(nonceLength), SocketFlags.Partial, cancellationToken);
+				if (read == 0)
+				{
+					return;
+				}
+				nonceLength += read;
+			}
+
+			// Dispatch the connection to the appropriate handler
+			RequestInfo? requestInfo = null;
+			try
+			{
+				// Find the matching request
+				lock (_lockObject)
+				{
+					if (!_requests.Remove(nonceBuffer, out requestInfo))
+					{
+						_logger.LogInformation("Invalid nonce for connection; ignoring.");
+						return;
+					}
+				}
+
+				// Create a crypto stream for communication, using the predetermined key for this request
+				using SocketComputeChannel channel = new SocketComputeChannel(socket, requestInfo.AesKey, requestInfo.AesIv);
+				await requestInfo.RunAsync(channel);
+			}
+			finally
+			{
+				requestInfo?.Dispose();
+			}
+		}
+	}
+}
diff --git a/Engine/Source/Programs/Shared/EpicGames.Horde/Compute/SocketComputeChannel.cs b/Engine/Source/Programs/Shared/EpicGames.Horde/Compute/SocketComputeChannel.cs
--- a/Engine/Source/Programs/Shared/EpicGames.Horde/Compute/SocketComputeChannel.cs
+++ b/Engine/Source/Programs/Shared/EpicGames.Horde/Compute/SocketComputeChannel.cs
@@ -1,9 +1,7 @@
 // Copyright Epic Games, Inc. All Rights Reserved.
 
 using EpicGames.Core;
-using EpicGames.Serialization;
 using System;
-using System.IO;
 using System.Net.Sockets;
 using System.Security.Cryptography;
 using System.Threading;
@@ -16,15 +14,13 @@ namespace EpicGames.Horde.Compute
 	/// <summary>
 	/// Transfers messages via a AES-encrypted socket
 	/// </summary>
-	public class SocketComputeChannel : IComputeChannel
+	public sealed class SocketComputeChannel : IComputeChannel, IDisposable
 	{
 		readonly Socket _socket;
 		readonly int _blockSize;
 		readonly ICryptoTransform _decryptor;
 		readonly ICryptoTransform _encryptor;
 
-		readonly CbWriter _writer;
-
 		byte[] _readBuffer = new byte[128];
 		int _decryptedLength;
 		int _encryptedLength;
@@ -49,8 +45,13 @@ namespace EpicGames.Horde.Compute
 			_blockSize = aes.BlockSize / 8;
 			_decryptor = aes.CreateDecryptor();
 			_encryptor = aes.CreateEncryptor();
+		}
 
-			_writer = new CbWriter();
+		/// <inheritdoc/>
+		public void Dispose()
+		{
+			_decryptor.Dispose();
+			_encryptor.Dispose();
 		}
 
 		/// <summary>
@@ -58,7 +59,7 @@ namespace EpicGames.Horde.Compute
 		/// </summary>
 		/// <param name="cancellationToken"></param>
 		/// <returns></returns>
-		public async Task<object> ReadAsync(CancellationToken cancellationToken)
+		public async Task<ReadOnlyMemory<byte>> ReadMessageAsync(CancellationToken cancellationToken)
 		{
 			// Discard any data from the previous message
 			if (_paddedMessageLength > 0)
@@ -66,25 +67,28 @@ namespace EpicGames.Horde.Compute
 				Buffer.BlockCopy(_readBuffer, _paddedMessageLength, _readBuffer, 0, _encryptedLength - _paddedMessageLength);
 				_encryptedLength -= _paddedMessageLength;
 				_decryptedLength -= _paddedMessageLength;
+				_paddedMessageLength = 0;
 			}
 
 			// Read the next message data
 			for (; ; )
 			{
 				// Check the object in the buffer
-				if (_decryptedLength >= 2)
+				if (_decryptedLength > 0)
 				{
-					int messageLength = 1 + VarInt.Measure(_readBuffer[1]);
-					if (_decryptedLength > messageLength)
+					int headerLength = VarInt.Measure(_readBuffer);
+
+					int messageLength = headerLength;
+					if (_decryptedLength >= headerLength)
 					{
-						messageLength += (int)VarInt.ReadUnsigned(_readBuffer.AsSpan(1));
+						messageLength += (int)VarInt.ReadUnsigned(_readBuffer);
 					}
 
 					_paddedMessageLength = messageLength + GetPadding(messageLength);
 
 					if (_decryptedLength >= _paddedMessageLength)
 					{
-						return ComputeMessage.Deserialize(_readBuffer.AsMemory(0, messageLength));
+						return _readBuffer.AsMemory(headerLength, messageLength - headerLength);
 					}
 
 					if (_readBuffer.Length < _paddedMessageLength)
@@ -97,7 +101,7 @@ namespace EpicGames.Horde.Compute
 				int read = await _socket.ReceiveAsync(_readBuffer.AsMemory(_encryptedLength), SocketFlags.Partial, cancellationToken);
 				if (read == 0)
 				{
-					throw new EndOfStreamException();
+					return ReadOnlyMemory<byte>.Empty;
 				}
 				_encryptedLength += read;
 
@@ -114,16 +118,15 @@ namespace EpicGames.Horde.Compute
 		/// <summary>
 		/// Sends a message to the remote
 		/// </summary>
-		/// <param name="message">Message to be sent</param>
+		/// <param name="message">Message to be written</param>
 		/// <param name="cancellationToken">Cancellation token</param>
 		/// <returns></returns>
-		public async Task WriteAsync(object message, CancellationToken cancellationToken)
+		public async Task WriteMessageAsync(ReadOnlyMemory<byte> message, CancellationToken cancellationToken)
 		{
-			// Format the output object
-			ComputeMessage.Serialize(message, _writer);
-
 			// Serialize it to the write buffer
-			int length = _writer.GetSize();
+			int headerLength = VarInt.MeasureUnsigned(message.Length);
+
+			int length = headerLength + message.Length;
 			length += GetPadding(length);
 
 			if (length > _writeBuffer.Length)
@@ -131,7 +134,9 @@ namespace EpicGames.Horde.Compute
 				_writeBuffer = new byte[length];
 			}
 
-			_writer.CopyTo(_writeBuffer);
+			// Write the data to the buffer
+			VarInt.WriteUnsigned(_writeBuffer, message.Length);
+			message.CopyTo(_writeBuffer.AsMemory(headerLength));
 
 			// Encrypt the data and send it
 			int transformed = _encryptor.TransformBlock(_writeBuffer, 0, length, _writeBuffer, 0);
diff --git a/Engine/Source/Programs/Shared/EpicGames.Horde/Storage/IStorageClient.cs b/Engine/Source/Programs/Shared/EpicGames.Horde/Storage/IStorageClient.cs
--- a/Engine/Source/Programs/Shared/EpicGames.Horde/Storage/IStorageClient.cs
+++ b/Engine/Source/Programs/Shared/EpicGames.Horde/Storage/IStorageClient.cs
@@ -52,67 +52,6 @@ namespace EpicGames.Horde.Storage
 		public bool? Extend { get; set; }
 	}
 
-	/// <summary>
-	/// Locates a node in storage
-	/// </summary>
-	public struct NodeLocator : IEquatable<NodeLocator>
-	{
-		/// <summary>
-		/// Location of the blob containing this node
-		/// </summary>
-		public BlobLocator Blob { get; }
-
-		/// <summary>
-		/// Index of the export within the blob
-		/// </summary>
-		public int ExportIdx { get; }
-
-		/// <summary>
-		/// Constructor
-		/// </summary>
-		public NodeLocator(BlobLocator blob, int exportIdx)
-		{
-			Blob = blob;
-			ExportIdx = exportIdx;
-		}
-
-		/// <summary>
-		/// Determines if this locator points to a valid entry
-		/// </summary>
-		public bool IsValid() => Blob.IsValid();
-
-		/// <summary>
-		/// Parse a string as a node locator
-		/// </summary>
-		/// <param name="text">Text to parse</param>
-		/// <returns></returns>
-		public static NodeLocator Parse(ReadOnlySpan<char> text)
-		{
-			int hashIdx = text.IndexOf('#');
-			int exportIdx = Int32.Parse(text.Slice(hashIdx + 1), NumberStyles.None, CultureInfo.InvariantCulture);
-			BlobLocator blobLocator = new BlobLocator(new Utf8String(text.Slice(0, hashIdx)));
-			return new NodeLocator(blobLocator, exportIdx);
-		}
-
-		/// <inheritdoc/>
-		public override bool Equals([NotNullWhen(true)] object? obj) => obj is NodeLocator locator && Equals(locator);
-
-		/// <inheritdoc/>
-		public bool Equals(NodeLocator other) => Blob == other.Blob && ExportIdx == other.ExportIdx;
-
-		/// <inheritdoc/>
-		public override int GetHashCode() => HashCode.Combine(Blob, ExportIdx);
-
-		/// <inheritdoc/>
-		public override string ToString() => $"{Blob}#{ExportIdx}";
-
-		/// <inheritdoc/>
-		public static bool operator ==(NodeLocator left, NodeLocator right) => left.Equals(right);
-
-		/// <inheritdoc/>
-		public static bool operator !=(NodeLocator left, NodeLocator right) => !left.Equals(right);
-	}
-
 	/// <summary>
 	/// Base interface for a low-level storage backend. Blobs added to this store are not content addressed, but referenced by <see cref="BlobLocator"/>.
 	/// </summary>
diff --git a/Engine/Source/Programs/Shared/EpicGames.Horde/Storage/NodeLocator.cs b/Engine/Source/Programs/Shared/EpicGames.Horde/Storage/NodeLocator.cs
new file mode 100644
--- /dev/null
+++ b/Engine/Source/Programs/Shared/EpicGames.Horde/Storage/NodeLocator.cs
@@ -0,0 +1,148 @@
+// Copyright Epic Games, Inc. All Rights Reserved.
+
+using System;
+using System.ComponentModel;
+using System.Diagnostics.CodeAnalysis;
+using System.Globalization;
+using System.Text.Json;
+using System.Text.Json.Serialization;
+using EpicGames.Core;
+using EpicGames.Serialization;
+
+namespace EpicGames.Horde.Storage
+{
+	/// <summary>
+	/// Locates a node in storage
+	/// </summary>
+	[JsonConverter(typeof(NodeLocatorJsonConverter))]
+	[TypeConverter(typeof(NodeLocatorTypeConverter))]
+	[CbConverter(typeof(NodeLocatorCbConverter))]
+	public struct NodeLocator : IEquatable<NodeLocator>
+	{
+		/// <summary>
+		/// Location of the blob containing this node
+		/// </summary>
+		public BlobLocator Blob { get; }
+
+		/// <summary>
+		/// Index of the export within the blob
+		/// </summary>
+		public int ExportIdx { get; }
+
+		/// <summary>
+		/// Constructor
+		/// </summary>
+		public NodeLocator(BlobLocator blob, int exportIdx)
+		{
+			Blob = blob;
+			ExportIdx = exportIdx;
+		}
+
+		/// <summary>
+		/// Determines if this locator points to a valid entry
+		/// </summary>
+		public bool IsValid() => Blob.IsValid();
+
+		/// <summary>
+		/// Parse a string as a node locator
+		/// </summary>
+		/// <param name="text">Text to parse</param>
+		/// <returns></returns>
+		public static NodeLocator Parse(ReadOnlySpan<char> text)
+		{
+			int hashIdx = text.IndexOf('#');
+			int exportIdx = Int32.Parse(text.Slice(hashIdx + 1), NumberStyles.None, CultureInfo.InvariantCulture);
+			BlobLocator blobLocator = new BlobLocator(new Utf8String(text.Slice(0, hashIdx)));
+			return new NodeLocator(blobLocator, exportIdx);
+		}
+
+		/// <inheritdoc/>
+		public override bool Equals([NotNullWhen(true)] object? obj) => obj is NodeLocator locator && Equals(locator);
+
+		/// <inheritdoc/>
+		public bool Equals(NodeLocator other) => Blob == other.Blob && ExportIdx == other.ExportIdx;
+
+		/// <inheritdoc/>
+		public override int GetHashCode() => HashCode.Combine(Blob, ExportIdx);
+
+		/// <inheritdoc/>
+		public override string ToString() => $"{Blob}#{ExportIdx}";
+
+		/// <inheritdoc/>
+		public static bool operator ==(NodeLocator left, NodeLocator right) => left.Equals(right);
+
+		/// <inheritdoc/>
+		public static bool operator !=(NodeLocator left, NodeLocator right) => !left.Equals(right);
+	}
+
+	/// <summary>
+	/// Type converter for BlobId to and from JSON
+	/// </summary>
+	sealed class NodeLocatorJsonConverter : JsonConverter<BlobLocator>
+	{
+		/// <inheritdoc/>
+		public override BlobLocator Read(ref Utf8JsonReader reader, Type typeToConvert, JsonSerializerOptions options) => new BlobLocator(new Utf8String(reader.GetUtf8String().ToArray()));
+
+		/// <inheritdoc/>
+		public override void Write(Utf8JsonWriter writer, BlobLocator value, JsonSerializerOptions options) => writer.WriteStringValue(value.Inner.Span);
+	}
+
+	/// <summary>
+	/// Type converter from strings to BlobId objects
+	/// </summary>
+	sealed class NodeLocatorTypeConverter : TypeConverter
+	{
+		/// <inheritdoc/>
+		public override bool CanConvertFrom(ITypeDescriptorContext? context, Type sourceType)
+		{
+			return sourceType == typeof(string);
+		}
+
+		/// <inheritdoc/>
+		public override object ConvertFrom(ITypeDescriptorContext? context, CultureInfo? culture, object? value)
+		{
+			return new BlobLocator((string)value!);
+		}
+	}
+
+	/// <summary>
+	/// Type converter to compact binary
+	/// </summary>
+	sealed class NodeLocatorCbConverter : CbConverterBase<NodeLocator>
+	{
+		/// <inheritdoc/>
+		public override NodeLocator Read(CbField field) => NodeLocator.Parse(field.AsString());
+
+		/// <inheritdoc/>
+		public override void Write(CbWriter writer, NodeLocator value) => writer.WriteUtf8StringValue(value.ToString());
+
+		/// <inheritdoc/>
+		public override void WriteNamed(CbWriter writer, Utf8String name, NodeLocator value) => writer.WriteUtf8String(name, value.ToString());
+	}
+
+	/// <summary>
+	/// Extension methods for node locators
+	/// </summary>
+	public static class NodeLocatorExtensions
+	{
+		/// <summary>
+		/// Deserialize a node locator
+		/// </summary>
+		/// <param name="reader">Reader to deserialize from</param>
+		/// <returns>The node id that was read</returns>
+		public static NodeLocator ReadNodeLocator(this IMemoryReader reader)
+		{
+			return NodeLocator.Parse(reader.ReadString());
+		}
+
+		/// <summary>
+		/// Serialize a node locator
+		/// </summary>
+		/// <param name="writer">Writer to serialize to</param>
+		/// <param name="value">Value to serialize</param>
+		public static void WriteNodeLocator(this IMemoryWriter writer, NodeLocator value)
+		{
+			writer.WriteString(value.ToString());
+		}
+	}
+}
diff --git a/Engine/Source/Programs/Shared/EpicGames.Horde/Storage/TreeReader.cs b/Engine/Source/Programs/Shared/EpicGames.Horde/Storage/TreeReader.cs
--- a/Engine/Source/Programs/Shared/EpicGames.Horde/Storage/TreeReader.cs
+++ b/Engine/Source/Programs/Shared/EpicGames.Horde/Storage/TreeReader.cs
@@ -64,6 +64,8 @@ namespace EpicGames.Horde.Storage
 			typeof(Logs.LogNode),
 			typeof(Logs.LogChunkNode),
 			typeof(Logs.LogIndexNode),
+			typeof(Compute.Cpp.CppComputeNode),
+			typeof(Compute.Cpp.CppComputeOutputNode)
 		};
 
 		/// <summary>
diff --git a/Engine/Source/Programs/Shared/EpicGames.Serialization/CbWriter.cs b/Engine/Source/Programs/Shared/EpicGames.Serialization/CbWriter.cs
--- a/Engine/Source/Programs/Shared/EpicGames.Serialization/CbWriter.cs
+++ b/Engine/Source/Programs/Shared/EpicGames.Serialization/CbWriter.cs
@@ -131,7 +131,7 @@ namespace EpicGames.Serialization
 			}
 		}
 
-		const int DefaultChunkSize = 1024;
+		public const int DefaultChunkSize = 1024;
 
 		readonly List<Chunk> _chunks = new List<Chunk>();
 		readonly Stack<Scope> _openScopes = new Stack<Scope>();
