Horde: Fix tests.

diff --git a/Engine/Source/Programs/Horde/Horde.Build/Compute/ComputeService.cs b/Engine/Source/Programs/Horde/Horde.Build/Compute/ComputeService.cs
--- a/Engine/Source/Programs/Horde/Horde.Build/Compute/ComputeService.cs
+++ b/Engine/Source/Programs/Horde/Horde.Build/Compute/ComputeService.cs
@@ -3,7 +3,6 @@
 using System.Collections.Generic;
 using EpicGames.Core;
 using EpicGames.Horde.Compute;
-using EpicGames.Serialization;
 using Horde.Build.Agents;
 using Horde.Build.Agents.Pools;
 using Horde.Build.Server;
@@ -42,7 +41,7 @@ namespace Horde.Build.Compute
 		/// <summary>
 		/// Constructor
 		/// </summary>
-		protected ComputeService(IOptionsMonitor<GlobalConfig> globalConfig)
+		public ComputeService(IOptionsMonitor<GlobalConfig> globalConfig)
 		{
 			_globalConfig = globalConfig;
 		}
