Horde: Fix issue severity icons taking the severity in the stream into account, which is more confusing than not
#rnx
#skipci

diff --git a/Engine/Source/Programs/Horde/HordeDashboard/src/components/StatusIcon.tsx b/Engine/Source/Programs/Horde/HordeDashboard/src/components/StatusIcon.tsx
--- a/Engine/Source/Programs/Horde/HordeDashboard/src/components/StatusIcon.tsx
+++ b/Engine/Source/Programs/Horde/HordeDashboard/src/components/StatusIcon.tsx
@@ -66,7 +66,7 @@ export const LabelStatusIcon: React.FC<{ label: JobLabel, style?: CSSProperties
    return <StatusIcon iconName={icon.icon} style={{ ...style, color: icon.color }} />;
 }
 
-export const IssueStatusIcon: React.FC<{ issue: GetIssueResponse, streamId?: string, style?: CSSProperties }> = ({ issue, streamId, style }) => {
+export const IssueStatusIcon: React.FC<{ issue: GetIssueResponse, style?: CSSProperties }> = ({ issue, style }) => {
 
    const colors = dashboard.getStatusColors();
 
@@ -78,6 +78,7 @@ export const IssueStatusIcon: React.FC<{ issue: GetIssueResponse, streamId?: str
    let icon = "Square";//issue.severity === IssueSeverity.Warning ? "Warning" : "FAWindowCloseSolid";
    style.color = issue.severity === IssueSeverity.Warning ? colors.get(StatusColor.Warnings) : colors.get(StatusColor.Failure);
 
+   /*
    if (streamId) {
       const affectedStream = issue.affectedStreams.find(s => s.streamId === streamId);
       if (affectedStream) {
@@ -96,11 +97,12 @@ export const IssueStatusIcon: React.FC<{ issue: GetIssueResponse, streamId?: str
          }
       }
    }
+   */
 
    return <StatusIcon iconName={icon} style={style} />
 }
 
-export const IssueStatusIconV2: React.FC<{ issue: FindIssueResponse, streamId?: string, style?: CSSProperties }> = ({ issue, streamId, style }) => {
+export const IssueStatusIconV2: React.FC<{ issue: FindIssueResponse, style?: CSSProperties }> = ({ issue, style }) => {
 
    const colors = dashboard.getStatusColors();
 
@@ -112,10 +114,6 @@ export const IssueStatusIconV2: React.FC<{ issue: FindIssueResponse, streamId?:
    let icon = "Square";//issue.severity === IssueSeverity.Warning ? "Warning" : "FAWindowCloseSolid";
    style.color = issue.severity === IssueSeverity.Warning ? colors.get(StatusColor.Warnings) : colors.get(StatusColor.Failure);
 
-   if (issue.streamSeverity) {      
-      style.color = issue.streamSeverity === IssueSeverity.Warning ? colors.get(StatusColor.Warnings) : colors.get(StatusColor.Failure);
-   }
-
    return <StatusIcon iconName={icon} style={style} />
 }
 
diff --git a/Engine/Source/Programs/Horde/HordeDashboard/src/components/StreamSummary.tsx b/Engine/Source/Programs/Horde/HordeDashboard/src/components/StreamSummary.tsx
--- a/Engine/Source/Programs/Horde/HordeDashboard/src/components/StreamSummary.tsx
+++ b/Engine/Source/Programs/Horde/HordeDashboard/src/components/StreamSummary.tsx
@@ -481,7 +481,7 @@ const HealthPanelIssues: React.FC<{ desktopAlerts?: boolean }> = observer(({ des
             summary = summary.slice(0, 100) + "...";
          }
 
-         return <Stack horizontal disableShrink={true}>{<IssueStatusIconV2 issue={issue} streamId={handler.streamId} />}<Text variant={textSize}>{`Issue ${issue.id} - ${summary}`}</Text></Stack>;
+         return <Stack horizontal disableShrink={true}>{<IssueStatusIconV2 issue={issue} />}<Text variant={textSize}>{`Issue ${issue.id} - ${summary}`}</Text></Stack>;
       }
 
       if (column.name === "Quarantine") {
diff --git a/Engine/Source/Programs/Horde/HordeDashboard/src/components/jobDetailsV2/JobDetailHealthV2.tsx b/Engine/Source/Programs/Horde/HordeDashboard/src/components/jobDetailsV2/JobDetailHealthV2.tsx
--- a/Engine/Source/Programs/Horde/HordeDashboard/src/components/jobDetailsV2/JobDetailHealthV2.tsx
+++ b/Engine/Source/Programs/Horde/HordeDashboard/src/components/jobDetailsV2/JobDetailHealthV2.tsx
@@ -248,7 +248,7 @@ export const HealthPanel: React.FC<{ jobDetails: JobDetailsV2 }> = observer(({ j
          query.set("issue", item.issue.id.toString());
          const href = `${location.pathname}?${query.toString()}` + window.location.hash;
 
-         return <Link to={href} className="job-item" onClick={() => setIssueHistory(true)}><Stack horizontal disableShrink={true}>{<IssueStatusIcon issue={issue} streamId={jobDetails.jobData!.streamId} />}<Text variant={textSize} style={{ textDecoration: !!issue.resolvedAt ? "line-through" : undefined }}>{`Issue ${issue.id} - ${summary}`}</Text></Stack></Link>;
+         return <Link to={href} className="job-item" onClick={() => setIssueHistory(true)}><Stack horizontal disableShrink={true}>{<IssueStatusIcon issue={issue} />}<Text variant={textSize} style={{ textDecoration: !!issue.resolvedAt ? "line-through" : undefined }}>{`Issue ${issue.id} - ${summary}`}</Text></Stack></Link>;
       }
 
       if (column.name === "Quarantine") {
