Horde: Disable autoscaling log output for local builds.

diff --git a/Engine/Source/Programs/Horde/Horde.Build/appsettings.Local.json b/Engine/Source/Programs/Horde/Horde.Build/appsettings.Local.json
--- a/Engine/Source/Programs/Horde/Horde.Build/appsettings.Local.json
+++ b/Engine/Source/Programs/Horde/Horde.Build/appsettings.Local.json
@@ -9,5 +9,13 @@
 			"ReplicateMetadata": false,
 			"ReplicateContent": false
 		}
+	},
+	"Serilog": {
+		"MinimumLevel": {
+			"Override": {
+				"Horde.Build.Agents.Fleet": "Warning",
+				"Horde.Build.Agents.Fleet.Providers.NoOpFleetManager": "Warning"
+			}
+		}
 	}
 }
