Horde: Downgrade some log output from information to debug.

diff --git a/Engine/Source/Programs/Horde/Horde.Build/Configuration/ConfigService.cs b/Engine/Source/Programs/Horde/Horde.Build/Configuration/ConfigService.cs
--- a/Engine/Source/Programs/Horde/Horde.Build/Configuration/ConfigService.cs
+++ b/Engine/Source/Programs/Horde/Horde.Build/Configuration/ConfigService.cs
@@ -250,7 +250,7 @@ namespace Horde.Build.Configuration
 			if (initialData.Length > 0)
 			{
 				snapshot = DeserializeSnapshot(initialData);
-				_logger.LogInformation("Initial snapshot for update: {@Info}", GetSnapshotInfo(initialData.Span, snapshot));
+				_logger.LogDebug("Initial snapshot for update: {@Info}", GetSnapshotInfo(initialData.Span, snapshot));
 			}
 
 			// Update the snapshot until we're asked to stop
@@ -427,9 +427,7 @@ namespace Horde.Build.Configuration
 		async Task WriteSnapshotAsync(ConfigSnapshot snapshot)
 		{
 			ReadOnlyMemory<byte> data = SerializeSnapshot(snapshot);
-
 			IoHash hash = IoHash.Compute(data.Span);
-			_logger.LogInformation("Publishing new config snapshot (hash: {Hash}, server: {Server}, size: {Size})", hash.ToString(), Program.Version.ToString(), data.Length);
 
 			await _redisService.GetDatabase().StringSetAsync(_snapshotKey, data);
 			await _redisService.PublishAsync(_updateChannel, RedisValue.EmptyString);
diff --git a/Engine/Source/Programs/Horde/Horde.Build/Devices/DeviceService.cs b/Engine/Source/Programs/Horde/Horde.Build/Devices/DeviceService.cs
--- a/Engine/Source/Programs/Horde/Horde.Build/Devices/DeviceService.cs
+++ b/Engine/Source/Programs/Horde/Horde.Build/Devices/DeviceService.cs
@@ -216,7 +216,7 @@ namespace Horde.Build.Devices
 
 				try
 				{
-					_logger.LogInformation("Expiring reservations");
+					_logger.LogDebug("Expiring reservations");
 					await _devices.ExpireReservationsAsync();
 				}
 				catch (Exception ex)
@@ -228,7 +228,7 @@ namespace Horde.Build.Devices
 				try
 				{
 
-					_logger.LogInformation("Sending shared device notifications");
+					_logger.LogDebug("Sending shared device notifications");
 					List<(UserId, IDevice)>? expireNotifications = await _devices.ExpireNotificatonsAsync(_settings.CurrentValue.SharedDeviceCheckoutDays);
 					if (expireNotifications != null && expireNotifications.Count > 0)
 					{
@@ -238,7 +238,7 @@ namespace Horde.Build.Devices
 						}
 					}
 
-					_logger.LogInformation("Expiring shared device checkouts");
+					_logger.LogDebug("Expiring shared device checkouts");
 					List<(UserId, IDevice)>? expireCheckouts = await _devices.ExpireCheckedOutAsync(_settings.CurrentValue.SharedDeviceCheckoutDays);
 					if (expireCheckouts != null && expireCheckouts.Count > 0)
 					{
