Horde: Fix multilist sorting on new build dialog
#rnx
#skipci

diff --git a/Engine/Source/Programs/Horde/HordeDashboard/src/components/NewBuild.tsx b/Engine/Source/Programs/Horde/HordeDashboard/src/components/NewBuild.tsx
--- a/Engine/Source/Programs/Horde/HordeDashboard/src/components/NewBuild.tsx
+++ b/Engine/Source/Programs/Horde/HordeDashboard/src/components/NewBuild.tsx
@@ -1034,7 +1034,7 @@ export const NewBuild: React.FC<{ streamId: string; show: boolean; onClose: (new
          }
       });
 
-      const groups = Array.from(gset).sort((a, b) => a < b ? - 1 : 1);
+      const groups = Array.from(gset).sort((a, b) => a.localeCompare(b));
 
       const options: IDropdownOption[] = [];
 
@@ -1057,7 +1057,7 @@ export const NewBuild: React.FC<{ streamId: string; show: boolean; onClose: (new
             dupes.set(item.text, !v ? 1 : v + 1);
          });
 
-         param.items.sort((a, b) => a.text < b.text ? -1 : 1).forEach(item => {
+         param.items.forEach(item => {
             if (item.group === group) {
                const key = buildParams.paramKey(param, item);
                const selected = buildParams!.values[buildParams.paramKey(param, item)] ? true : false;
