Horde: Fix anonymous login compatibility with dashboard.

diff --git a/Engine/Source/Programs/Horde/Horde.Build/Authentication/AnonymousHandler.cs b/Engine/Source/Programs/Horde/Horde.Build/Authentication/AnonymousHandler.cs
--- a/Engine/Source/Programs/Horde/Horde.Build/Authentication/AnonymousHandler.cs
+++ b/Engine/Source/Programs/Horde/Horde.Build/Authentication/AnonymousHandler.cs
@@ -35,16 +35,14 @@ namespace Horde.Build.Authentication
 		{
 			List<Claim> claims = new List<Claim>();
 			claims.Add(new Claim(ClaimTypes.Name, AuthenticationScheme));
+			claims.Add(new Claim(HordeClaimTypes.UserId, UserId.Anonymous.ToString()));
 
 			if (Options.AdminClaimType != null && Options.AdminClaimValue != null)
 			{
 				claims.Add(new Claim(Options.AdminClaimType, Options.AdminClaimValue));
 			}
 
-			ServerSettings currentSettings = _settings.CurrentValue;
-
 			ClaimsIdentity identity = new ClaimsIdentity(claims, Scheme.Name);
-			identity.AddClaim(new Claim(currentSettings.AdminClaimType, currentSettings.AdminClaimValue));
 
 			ClaimsPrincipal principal = new ClaimsPrincipal(identity);
 			AuthenticationTicket ticket = new AuthenticationTicket(principal, Scheme.Name);
diff --git a/Engine/Source/Programs/Horde/Horde.Build/Startup.cs b/Engine/Source/Programs/Horde/Horde.Build/Startup.cs
--- a/Engine/Source/Programs/Horde/Horde.Build/Startup.cs
+++ b/Engine/Source/Programs/Horde/Horde.Build/Startup.cs
@@ -447,6 +447,10 @@ namespace Horde.Build
 				services.AddSingleton<IAvatarService, SlackNotificationSink>(sp => sp.GetRequiredService<SlackNotificationSink>());
 				services.AddSingleton<INotificationSink, SlackNotificationSink>(sp => sp.GetRequiredService<SlackNotificationSink>());
 			}
+			else
+			{
+				services.AddSingleton<IAvatarService, NullAvatarService>();
+			}
 
 			services.AddSingleton<DeviceService>();						
 			services.AddSingleton<NoticeService>();
diff --git a/Engine/Source/Programs/Horde/Horde.Build/Users/IAvatarService.cs b/Engine/Source/Programs/Horde/Horde.Build/Users/IAvatarService.cs
--- a/Engine/Source/Programs/Horde/Horde.Build/Users/IAvatarService.cs
+++ b/Engine/Source/Programs/Horde/Horde.Build/Users/IAvatarService.cs
@@ -16,5 +16,14 @@ namespace Horde.Build.Users
 		/// <returns></returns>
 		Task<IAvatar?> GetAvatarAsync(IUser user);
 	}
+
+	/// <summary>
+	/// Default implementation of <see cref="IAvatarService"/>
+	/// </summary>
+	public class NullAvatarService : IAvatarService
+	{
+		/// <inheritdoc/>
+		public Task<IAvatar?> GetAvatarAsync(IUser user) => Task.FromResult<IAvatar?>(null);
+	}
 }
 
diff --git a/Engine/Source/Programs/Horde/Horde.Build/Users/UserId.cs b/Engine/Source/Programs/Horde/Horde.Build/Users/UserId.cs
--- a/Engine/Source/Programs/Horde/Horde.Build/Users/UserId.cs
+++ b/Engine/Source/Programs/Horde/Horde.Build/Users/UserId.cs
@@ -21,6 +21,11 @@ namespace Horde.Build.Users
 		/// </summary>
 		public static UserId Empty { get; } = new UserId(ObjectId.Empty);
 
+		/// <summary>
+		/// Special user id for an anonymous administrator
+		/// </summary>
+		public static UserId Anonymous { get; } = UserId.Parse("63f7d3525119b9aa4c0f035a");
+
 		/// <summary>
 		/// Creates a new <see cref="UserId"/>
 		/// </summary>
