Horde: Fix issue with jira components not being processsed
#rnx

diff --git a/Engine/Source/Programs/Horde/Horde.Build/Issues/External/JiraService.cs b/Engine/Source/Programs/Horde/Horde.Build/Issues/External/JiraService.cs
--- a/Engine/Source/Programs/Horde/Horde.Build/Issues/External/JiraService.cs
+++ b/Engine/Source/Programs/Horde/Horde.Build/Issues/External/JiraService.cs
@@ -279,21 +279,18 @@ namespace Horde.Build.Issues.External
 					JsonElement component = jiraComponents[j];
 
 					// skip archived components
-					if (component.TryGetProperty("archived", out JsonElement archived))
+					if (component.TryGetProperty("archived", out JsonElement archived) && archived.GetBoolean())
 					{
-						if (archived.GetBoolean())
-						{
-							continue;
-						}
+						continue;
+					}
 
-						if (component.TryGetStringProperty("id", out string? id))
+					if (component.TryGetStringProperty("id", out string? id))
+					{
+						if (component.TryGetStringProperty("name", out string? name))
 						{
-							if (component.TryGetStringProperty("name", out string? name))
-							{
-								jiraProject.Components[id] = name;
-							}
+							jiraProject.Components[id] = name;
 						}
-					}
+					}					
 				}
 
 				_cachedJiraProjects[projectKey] = jiraProject;
