Horde: Fix anonymous dashboard login, and copy contents of the DashboardApp folder to the output directory.

diff --git a/Engine/Source/Programs/Horde/Horde.Build/Horde.Build.csproj b/Engine/Source/Programs/Horde/Horde.Build/Horde.Build.csproj
--- a/Engine/Source/Programs/Horde/Horde.Build/Horde.Build.csproj
+++ b/Engine/Source/Programs/Horde/Horde.Build/Horde.Build.csproj
@@ -124,6 +124,12 @@
     </Content>
   </ItemGroup>
 
+  <ItemGroup>
+    <None Include="DashboardApp\**">
+      <CopyToOutputDirectory>Always</CopyToOutputDirectory>
+    </None>
+  </ItemGroup>
+
   <ItemGroup>
     <Content Include="Defaults\ue5.project.png">
       <CopyToOutputDirectory>Always</CopyToOutputDirectory>
diff --git a/Engine/Source/Programs/Horde/Horde.Build/Users/UserCollectionV2.cs b/Engine/Source/Programs/Horde/Horde.Build/Users/UserCollectionV2.cs
--- a/Engine/Source/Programs/Horde/Horde.Build/Users/UserCollectionV2.cs
+++ b/Engine/Source/Programs/Horde/Horde.Build/Users/UserCollectionV2.cs
@@ -187,6 +187,12 @@ namespace Horde.Build.Users
 		/// <inheritdoc/>
 		public async Task<IUser?> GetUserAsync(UserId id)
 		{
+			if (id == UserId.Anonymous)
+			{
+				// Anonymous user is handled as a special case so we can debug with local anonymous login against a production DB in read-only mode.
+				return new UserDocument(id, "Anonymous", "anonymous", null);
+			}
+
 			IUser? user = await _users.Find(x => x.Id == id).FirstOrDefaultAsync();
 			using (ICacheEntry entry = _userCache.CreateEntry(id))
 			{
