Horde: Add a service to allow tunneling connections through the Horde server. Accepts incoming connections on two ports and matches them up via a cryptographic unique id. Designed for use by compute service, where client machines may not accept incoming connections.

diff --git a/Engine/Source/Programs/Horde/Horde.Agent/Leases/LeaseManager.cs b/Engine/Source/Programs/Horde/Horde.Agent/Leases/LeaseManager.cs
--- a/Engine/Source/Programs/Horde/Horde.Agent/Leases/LeaseManager.cs
+++ b/Engine/Source/Programs/Horde/Horde.Agent/Leases/LeaseManager.cs
@@ -7,7 +7,6 @@ using System.Linq;
 using System.Threading;
 using System.Threading.Tasks;
 using EpicGames.Core;
-using Google.Protobuf;
 using Google.Protobuf.WellKnownTypes;
 using Grpc.Core;
 using Horde.Agent.Services;
@@ -23,6 +22,8 @@ using OpenTracing.Util;
 
 namespace Horde.Agent.Leases
 {
+	using ByteString = Google.Protobuf.ByteString;
+
 	/// <summary>
 	/// Outcome from a session
 	/// </summary>
diff --git a/Engine/Source/Programs/Horde/Horde.Agent/Utility/JsonRpcLogger.cs b/Engine/Source/Programs/Horde/Horde.Agent/Utility/JsonRpcLogger.cs
--- a/Engine/Source/Programs/Horde/Horde.Agent/Utility/JsonRpcLogger.cs
+++ b/Engine/Source/Programs/Horde/Horde.Agent/Utility/JsonRpcLogger.cs
@@ -20,6 +20,7 @@ using Microsoft.Extensions.Logging;
 
 namespace Horde.Agent.Utility
 {
+	using ByteString = Google.Protobuf.ByteString;
 	using JsonObject = System.Text.Json.Nodes.JsonObject;
 
 	/// <summary>
diff --git a/Engine/Source/Programs/Horde/Horde.Build.Tests/TunnelServiceTests.cs b/Engine/Source/Programs/Horde/Horde.Build.Tests/TunnelServiceTests.cs
new file mode 100644
--- /dev/null
+++ b/Engine/Source/Programs/Horde/Horde.Build.Tests/TunnelServiceTests.cs
@@ -0,0 +1,117 @@
+// Copyright Epic Games, Inc. All Rights Reserved.
+
+using System;
+using System.Buffers.Binary;
+using System.Linq;
+using System.Net;
+using System.Net.Sockets;
+using System.Threading;
+using System.Threading.Tasks;
+using EpicGames.Horde.Compute;
+using Horde.Build.Compute;
+using Microsoft.Extensions.DependencyInjection;
+using Microsoft.Extensions.Options;
+using Microsoft.VisualStudio.TestTools.UnitTesting;
+
+namespace Horde.Build.Tests
+{
+	[TestClass]
+	public sealed class TunnelServiceTests : IDisposable
+	{
+		readonly ServiceProvider _serviceProvider;
+
+		public TunnelServiceTests()
+		{
+			ServiceCollection services = new ServiceCollection();
+			services.AddLogging();
+			services.Configure<ServerSettings>(x =>
+			{
+				x.ComputeInitiatorPort = 7000;
+				x.ComputeRemotePort = 7002;
+			});
+			services.AddSingleton<TunnelService>();
+
+			_serviceProvider = services.BuildServiceProvider();
+		}
+
+		/// <inheritdoc/>
+		public void Dispose()
+		{
+			_serviceProvider.Dispose();
+		}
+
+		[TestMethod]
+		public async Task SendRecvTest()
+		{
+			CancellationToken cancellationToken = CancellationToken.None;
+			byte[] nonce = new byte[ServerComputeClient.NonceLength];
+
+			TunnelService tunnelService = _serviceProvider.GetRequiredService<TunnelService>();
+			tunnelService.Start(IPAddress.IPv6Loopback);
+
+			IOptions<ServerSettings> settings = _serviceProvider.GetRequiredService<IOptions<ServerSettings>>();
+
+			using (Socket initiatorSocket = await ConnectAsync(settings.Value.ComputeInitiatorPort, nonce, cancellationToken))
+			{
+				byte[] remotePortBuffer = new byte[2];
+				await initiatorSocket.ReceiveAsync(remotePortBuffer, SocketFlags.None, cancellationToken);
+				ushort remotePort = BinaryPrimitives.ReadUInt16LittleEndian(remotePortBuffer);
+
+				using (Socket remoteSocket = await ConnectAsync(remotePort, nonce, cancellationToken))
+				{
+					byte[] sendInitiatorData = new byte[] { 1, 2, 3, 4, 5 };
+					byte[] sendRemoteData = new byte[] { 6, 7, 8, 9 };
+
+					Task<byte[]> initiatorTask = SendRecvAsync(initiatorSocket, sendInitiatorData, sendRemoteData.Length, cancellationToken);
+					Task<byte[]> remoteTask = SendRecvAsync(remoteSocket, sendRemoteData, sendInitiatorData.Length, cancellationToken);
+
+					byte[] recvInitiatorData = await initiatorTask.ConfigureAwait(false);
+					byte[] recvRemoteData = await remoteTask.ConfigureAwait(false);
+
+					Assert.IsTrue(recvRemoteData.SequenceEqual(sendInitiatorData));
+					Assert.IsTrue(recvInitiatorData.SequenceEqual(sendRemoteData));
+				}
+			}
+
+			await tunnelService.StopAsync(CancellationToken.None);
+		}
+
+		static async Task<Socket> ConnectAsync(int port, byte[] nonce, CancellationToken cancellationToken)
+		{
+			Socket socket = new Socket(SocketType.Stream, ProtocolType.Tcp);
+			await socket.ConnectAsync(IPAddress.IPv6Loopback, port, cancellationToken);
+			await socket.SendAsync(nonce, SocketFlags.None);
+			return socket;
+		}
+
+		static async Task<byte[]> SendRecvAsync(Socket socket, byte[] dataToSend, int lengthToRecv, CancellationToken cancellationToken)
+		{
+			Task<byte[]> recv = RecvAsync(socket, lengthToRecv, cancellationToken);
+			for (int offset = 0; offset < dataToSend.Length;)
+			{
+				int sent = await socket.SendAsync(dataToSend.AsMemory(offset), SocketFlags.None, cancellationToken);
+				if (sent == 0)
+				{
+					throw new NotImplementedException();
+				}
+				offset += sent;
+			}
+			return await recv;
+		}
+
+		static async Task<byte[]> RecvAsync(Socket socket, int length, CancellationToken cancellationToken)
+		{
+			byte[] buffer = new byte[length];
+			for (int offset = 0; offset < length;)
+			{
+				int read = await socket.ReceiveAsync(buffer.AsMemory(offset), SocketFlags.None, cancellationToken);
+				if (read == 0)
+				{
+					throw new NotImplementedException();
+				}
+				offset += read;
+			}
+			return buffer;
+		}
+	}
+}
diff --git a/Engine/Source/Programs/Horde/Horde.Build/Compute/ComputeController.cs b/Engine/Source/Programs/Horde/Horde.Build/Compute/ComputeController.cs
--- a/Engine/Source/Programs/Horde/Horde.Build/Compute/ComputeController.cs
+++ b/Engine/Source/Programs/Horde/Horde.Build/Compute/ComputeController.cs
@@ -5,7 +5,6 @@ using System.Collections.Generic;
 using System.Net;
 using EpicGames.Core;
 using EpicGames.Horde.Compute;
-using Google.Protobuf;
 using Horde.Build.Acls;
 using Horde.Build.Server;
 using Horde.Build.Utilities;
@@ -16,6 +15,8 @@ using Microsoft.Extensions.Options;
 
 namespace Horde.Build.Compute
 {
+	using ByteString = Google.Protobuf.ByteString;
+
 	/// <summary>
 	/// Request a machine to execute compute requests
 	/// </summary>
diff --git a/Engine/Source/Programs/Horde/Horde.Build/Compute/TunnelService.cs b/Engine/Source/Programs/Horde/Horde.Build/Compute/TunnelService.cs
new file mode 100644
--- /dev/null
+++ b/Engine/Source/Programs/Horde/Horde.Build/Compute/TunnelService.cs
@@ -0,0 +1,248 @@
+// Copyright Epic Games, Inc. All Rights Reserved.
+
+using System;
+using System.Collections.Generic;
+using System.Net;
+using System.Net.Sockets;
+using System.Threading;
+using System.Threading.Tasks;
+using Microsoft.Extensions.Hosting;
+using Microsoft.Extensions.Logging;
+using Microsoft.Extensions.Options;
+using EpicGames.Core;
+using EpicGames.Horde.Compute;
+using System.IO;
+using System.Buffers;
+using System.Buffers.Binary;
+
+namespace Horde.Build.Compute
+{
+	/// <summary>
+	/// Tunnels connections between Horde agents and compute initiators
+	/// </summary>
+	public sealed class TunnelService : IHostedService, IDisposable
+	{
+		const int BufferSize = 4096;
+
+		class Tunnel
+		{
+			public Socket InitiatorSocket { get; }
+			public TaskCompletionSource<Socket> RemoteSocket { get; } = new TaskCompletionSource<Socket>();
+			public AsyncEvent Complete { get; } = new AsyncEvent();
+
+			public Tunnel(Socket initiatorSocket) => InitiatorSocket = initiatorSocket;
+		}
+
+		readonly CancellationTokenSource _cancellationSource;
+		readonly ServerSettings _settings;
+		readonly ILogger _logger;
+
+		Task? _initiatorServerTask;
+		Task? _remoteServerTask;
+
+		readonly Dictionary<ByteString, Tunnel> _waiters = new Dictionary<ByteString, Tunnel>();
+
+		/// <summary>
+		/// Constructor
+		/// </summary>
+		public TunnelService(IOptions<ServerSettings> settings, ILogger<TunnelService> logger)
+		{
+			_cancellationSource = new CancellationTokenSource();
+			_settings = settings.Value;
+			_logger = logger;
+		}
+
+		/// <inheritdoc/>
+		public void Dispose()
+		{
+			_cancellationSource.Dispose();
+		}
+
+		/// <inheritdoc/>
+		public Task StartAsync(CancellationToken cancellationToken)
+		{
+			Start(IPAddress.IPv6Any);
+			return Task.CompletedTask;
+		}
+
+		/// <inheritdoc/>
+		public void Start(IPAddress address)
+		{
+			_initiatorServerTask = RunTcpListenerAsync(address, _settings.ComputeInitiatorPort, HandleInitiatorAsync, _cancellationSource.Token);
+			_remoteServerTask = RunTcpListenerAsync(address, _settings.ComputeRemotePort, HandleRemoteAsync, _cancellationSource.Token);
+		}
+
+		/// <inheritdoc/>
+		public async Task StopAsync(CancellationToken cancellationToken)
+		{
+			_cancellationSource.Cancel();
+
+			await _initiatorServerTask!.IgnoreCanceledExceptionsAsync().ConfigureAwait(false);
+			_initiatorServerTask = null;
+
+			await _remoteServerTask!.IgnoreCanceledExceptionsAsync().ConfigureAwait(false);
+			_remoteServerTask = null;
+		}
+
+		async Task HandleRemoteAsync(Socket remoteSocket, CancellationToken cancellationToken)
+		{
+			// Read the nonce for the connection
+			byte[] nonceBuffer = new byte[ServerComputeClient.NonceLength];
+			await ReadFullBufferAsync(remoteSocket, nonceBuffer, cancellationToken);
+			ByteString nonce = new ByteString(nonceBuffer);
+
+			// Get the completion source for this nonce
+			Tunnel? tunnel;
+			lock (_waiters)
+			{
+				_waiters.TryGetValue(nonce, out tunnel);
+			}
+
+			// Pump data from the remote socket through to the initiator socket
+			if (tunnel != null && tunnel.RemoteSocket.TrySetResult(remoteSocket))
+			{
+				await tunnel.Complete.Task;
+			}
+		}
+
+		async Task HandleInitiatorAsync(Socket initiatorSocket, CancellationToken cancellationToken)
+		{
+			// Read the nonce for the connection
+			byte[] nonceBuffer = new byte[ServerComputeClient.NonceLength];
+			await ReadFullBufferAsync(initiatorSocket, nonceBuffer, cancellationToken);
+			ByteString nonce = new ByteString(nonceBuffer);
+
+			// Create the tunnel and add it to the waiting list
+			Tunnel tunnel = new Tunnel(initiatorSocket);
+			lock (_waiters)
+			{
+				_waiters.Add(nonce, tunnel);
+			}
+
+			// Reply with the port number that the remote should connect on
+			byte[] ackBuffer = new byte[2];
+			BinaryPrimitives.WriteUInt16LittleEndian(ackBuffer, (ushort)_settings.ComputeRemotePort);
+			await SendFullBufferAsync(initiatorSocket, ackBuffer, cancellationToken);
+
+			// Try/finally to mark the tunnel as complete when finished
+			try
+			{
+				// Start a read from the initiating socket, so we can react to disconnection events
+				using IMemoryOwner<byte> initiatorToRemoteBuffer = MemoryPool<byte>.Shared.Rent(BufferSize);
+				Task<int> readTask = initiatorSocket.ReceiveAsync(initiatorToRemoteBuffer.Memory, SocketFlags.None, cancellationToken).AsTask();
+
+				// Wait for the other end of the socket to be available
+				Task initTask = Task.WhenAny(readTask, tunnel.RemoteSocket.Task);
+				if (initTask == readTask)
+				{
+					return;
+				}
+
+				// Start a task to copy data from the remote back to the initiator
+				Socket remoteSocket = await tunnel.RemoteSocket.Task;
+
+				// Start copying data from the remote to the initiator
+				using IMemoryOwner<byte> remoteToInitiatorBuffer = MemoryPool<byte>.Shared.Rent(BufferSize);
+				Task remoteToInitiatorTask = CopySocketDataAsync(remoteSocket, tunnel.InitiatorSocket, remoteToInitiatorBuffer.Memory, cancellationToken);
+
+				// Copy data from the initiator to the remote
+				int readSize = await readTask;
+				await SendFullBufferAsync(remoteSocket, initiatorToRemoteBuffer.Memory.Slice(0, readSize), cancellationToken);
+				await CopySocketDataAsync(initiatorSocket, remoteSocket, initiatorToRemoteBuffer.Memory, cancellationToken);
+
+				// Wait for the socket to close
+				await remoteToInitiatorTask;
+			}
+			finally
+			{
+				// Remove the nonce to the waiting list
+				lock (_waiters)
+				{
+					_waiters.Remove(nonce);
+				}
+
+				// Signal to the other end of the socket that we're complete
+				tunnel.RemoteSocket.TrySetResult(null!);
+				tunnel.Complete.Latch();
+			}
+		}
+
+		static async Task CopySocketDataAsync(Socket sourceSocket, Socket targetSocket, Memory<byte> buffer, CancellationToken cancellationToken)
+		{
+			for (; ; )
+			{
+				int read = await sourceSocket.ReceiveAsync(buffer, SocketFlags.None, cancellationToken);
+				if (read == 0)
+				{
+					break;
+				}
+				await SendFullBufferAsync(targetSocket, buffer.Slice(0, read), cancellationToken);
+			}
+		}
+
+		static async Task ReadFullBufferAsync(Socket socket, Memory<byte> buffer, CancellationToken cancellationToken)
+		{
+			for (int offset = 0; offset < buffer.Length; )
+			{
+				int read = await socket.ReceiveAsync(buffer.Slice(offset), SocketFlags.None, cancellationToken);
+				if (read == 0)
+				{
+					throw new EndOfStreamException();
+				}
+				offset += read;
+			}
+		}
+
+		static async Task SendFullBufferAsync(Socket socket, ReadOnlyMemory<byte> buffer, CancellationToken cancellationToken)
+		{
+			while (buffer.Length > 0)
+			{
+				int written = await socket.SendAsync(buffer, SocketFlags.None, cancellationToken);
+				buffer = buffer.Slice(written);
+			}
+		}
+
+		async Task RunTcpListenerAsync(IPAddress address, int port, Func<Socket, CancellationToken, Task> handleConnectionAsync, CancellationToken cancellationToken)
+		{
+			if (port != 0)
+			{
+				TcpListener listener = new TcpListener(address, port);
+				listener.Start();
+
+				List<Task> tasks = new List<Task>();
+				try
+				{
+					for (; ; )
+					{
+						TcpClient client = await listener.AcceptTcpClientAsync(cancellationToken);
+						_logger.LogInformation("Received connection from {Remote}", client.Client.RemoteEndPoint);
+
+						Task task = HandleConnectionGuardedAsync(client, handleConnectionAsync, cancellationToken);
+						tasks.Add(task);
+					}
+				}
+				finally
+				{
+					await Task.WhenAll(tasks);
+					listener.Stop();
+				}
+			}
+		}
+
+		async Task HandleConnectionGuardedAsync(TcpClient tcpClient, Func<Socket, CancellationToken, Task> handleConnectionAsync, CancellationToken cancellationToken)
+		{
+			try
+			{
+				await handleConnectionAsync(tcpClient.Client, cancellationToken);
+			}
+			catch (Exception ex)
+			{
+				_logger.LogInformation(ex, "Exception while handling request: {Message}", ex.Message);
+			}
+			finally
+			{
+				tcpClient.Dispose();
+			}
+		}
+	}
+}
diff --git a/Engine/Source/Programs/Horde/Horde.Build/ServerSettings.cs b/Engine/Source/Programs/Horde/Horde.Build/ServerSettings.cs
--- a/Engine/Source/Programs/Horde/Horde.Build/ServerSettings.cs
+++ b/Engine/Source/Programs/Horde/Horde.Build/ServerSettings.cs
@@ -281,6 +281,16 @@ namespace Horde.Build
 		/// </summary>
 		public int Http2Port { get; set; }
 
+		/// <summary>
+		/// Port for listening to compute tunnel initiator requests
+		/// </summary>
+		public int ComputeInitiatorPort { get; set; }
+
+		/// <summary>
+		/// Port for compute remotes to connect to
+		/// </summary>
+		public int ComputeRemotePort { get; set; }
+
 		/// <summary>
 		/// MongoDB connection string
 		/// </summary>
diff --git a/Engine/Source/Programs/Horde/Horde.Build/Startup.cs b/Engine/Source/Programs/Horde/Horde.Build/Startup.cs
--- a/Engine/Source/Programs/Horde/Horde.Build/Startup.cs
+++ b/Engine/Source/Programs/Horde/Horde.Build/Startup.cs
@@ -390,6 +390,7 @@ namespace Horde.Build
 			services.AddSingleton<AclService>();
 			services.AddSingleton<AgentService>();			
 			services.AddSingleton<AgentSoftwareService>();
+			services.AddSingleton<TunnelService>();
 			services.AddSingleton<ConsistencyService>();
 			services.AddSingleton<RequestTrackerService>();
 			services.AddSingleton<CredentialService>();
diff --git a/Engine/Source/Programs/Shared/EpicGames.Core/AsyncUtils.cs b/Engine/Source/Programs/Shared/EpicGames.Core/AsyncUtils.cs
--- a/Engine/Source/Programs/Shared/EpicGames.Core/AsyncUtils.cs
+++ b/Engine/Source/Programs/Shared/EpicGames.Core/AsyncUtils.cs
@@ -33,6 +33,21 @@ namespace EpicGames.Core
 			return Task.Delay(-1, token).ContinueWith(_ => Task.FromCanceled<T>(token), TaskScheduler.Default).Unwrap();
 		}
 
+		/// <summary>
+		/// Waits for a task to complete, ignoring any cancellation exceptions
+		/// </summary>
+		/// <param name="task">Task to wait for</param>
+		public static async Task IgnoreCanceledExceptionsAsync(this Task task)
+		{
+			try
+			{
+				await task.ConfigureAwait(false);
+			}
+			catch (OperationCanceledException)
+			{
+			}
+		}
+
 		/// <summary>
 		/// Returns a task that will be abandoned if a cancellation token is activated. This differs from the normal cancellation pattern in that the task will run to completion, but waiting for it can be cancelled.
 		/// </summary>
diff --git a/Engine/Source/Programs/Shared/EpicGames.Core/ByteString.cs b/Engine/Source/Programs/Shared/EpicGames.Core/ByteString.cs
--- a/Engine/Source/Programs/Shared/EpicGames.Core/ByteString.cs
+++ b/Engine/Source/Programs/Shared/EpicGames.Core/ByteString.cs
@@ -16,7 +16,7 @@ namespace EpicGames.Core
 	[JsonSchemaString]
 	[JsonConverter(typeof(ByteStringJsonConverter))]
 	[TypeConverter(typeof(ByteStringTypeConverter))]
-	struct ByteString : IEquatable<ByteString>, IComparable<ByteString>
+	public struct ByteString : IEquatable<ByteString>, IComparable<ByteString>
 	{
 		/// <summary>
 		/// Underlying data for this string
diff --git a/Engine/Source/Programs/Shared/EpicGames.Horde/Compute/ServerComputeClient.cs b/Engine/Source/Programs/Shared/EpicGames.Horde/Compute/ServerComputeClient.cs
--- a/Engine/Source/Programs/Shared/EpicGames.Horde/Compute/ServerComputeClient.cs
+++ b/Engine/Source/Programs/Shared/EpicGames.Horde/Compute/ServerComputeClient.cs
@@ -18,7 +18,10 @@ namespace EpicGames.Horde.Compute
 	/// </summary>
 	public sealed class ServerComputeClient : IComputeClient
 	{
-		const int NonceLength = 64;
+		/// <summary>
+		/// Length of the nonce sent as part of handshaking between initiator and remote
+		/// </summary>
+		public const int NonceLength = 64;
 
 		class AddComputeTaskRequest
 		{
