Horde: Disable SendRecvTest; not working through Docker.

diff --git a/Engine/Source/Programs/Horde/Horde.Build.Tests/TunnelServiceTests.cs b/Engine/Source/Programs/Horde/Horde.Build.Tests/TunnelServiceTests.cs
--- a/Engine/Source/Programs/Horde/Horde.Build.Tests/TunnelServiceTests.cs
+++ b/Engine/Source/Programs/Horde/Horde.Build.Tests/TunnelServiceTests.cs
@@ -40,6 +40,7 @@ namespace Horde.Build.Tests
 			_serviceProvider.Dispose();
 		}
 
+		[Ignore]
 		[TestMethod]
 		public async Task SendRecvTest()
 		{
