Horde: Try using IPv4 loopback adapter for tunnel service tests.

diff --git a/Engine/Source/Programs/Horde/Horde.Build.Tests/TunnelServiceTests.cs b/Engine/Source/Programs/Horde/Horde.Build.Tests/TunnelServiceTests.cs
--- a/Engine/Source/Programs/Horde/Horde.Build.Tests/TunnelServiceTests.cs
+++ b/Engine/Source/Programs/Horde/Horde.Build.Tests/TunnelServiceTests.cs
@@ -40,7 +40,6 @@ namespace Horde.Build.Tests
 			_serviceProvider.Dispose();
 		}
 
-		[Ignore]
 		[TestMethod]
 		public async Task SendRecvTest()
 		{
@@ -48,7 +47,7 @@ namespace Horde.Build.Tests
 			byte[] nonce = new byte[ServerComputeClient.NonceLength];
 
 			TunnelService tunnelService = _serviceProvider.GetRequiredService<TunnelService>();
-			tunnelService.Start(IPAddress.IPv6Loopback);
+			tunnelService.Start(IPAddress.Loopback);
 
 			IOptions<ServerSettings> settings = _serviceProvider.GetRequiredService<IOptions<ServerSettings>>();
 
@@ -80,7 +79,7 @@ namespace Horde.Build.Tests
 		static async Task<Socket> ConnectAsync(int port, byte[] nonce, CancellationToken cancellationToken)
 		{
 			Socket socket = new Socket(SocketType.Stream, ProtocolType.Tcp);
-			await socket.ConnectAsync(IPAddress.IPv6Loopback, port, cancellationToken);
+			await socket.ConnectAsync(IPAddress.Loopback, port, cancellationToken);
 			await socket.SendAsync(nonce, SocketFlags.None);
 			return socket;
 		}
