Horde: Add an endpoint to return the configured tunnel information.

diff --git a/Engine/Source/Programs/Horde/Horde.Build/Compute/ComputeController.cs b/Engine/Source/Programs/Horde/Horde.Build/Compute/ComputeController.cs
--- a/Engine/Source/Programs/Horde/Horde.Build/Compute/ComputeController.cs
+++ b/Engine/Source/Programs/Horde/Horde.Build/Compute/ComputeController.cs
@@ -69,6 +69,27 @@ namespace Horde.Build.Compute
 		public List<AddComputeTaskRequest> Tasks { get; set; } = new List<AddComputeTaskRequest>();
 	}
 
+	/// <summary>
+	/// Gets information about the configured compute tunnel
+	/// </summary>
+	public class GetComputeTunnelResponse
+	{
+		/// <summary>
+		/// IP address of the server
+		/// </summary>
+		public string? Ip { get; set; }
+
+		/// <summary>
+		/// Listening port for the initiator
+		/// </summary>
+		public int InitiatorPort { get; set; }
+
+		/// <summary>
+		/// Listening port for the remote
+		/// </summary>
+		public int RemotePort { get; set; }
+	}
+
 	/// <summary>
 	/// Controller for the /api/v2/compute endpoint
 	/// </summary>
@@ -131,5 +152,34 @@ namespace Horde.Build.Compute
 			}
 			return Ok();
 		}
+
+		/// <summary>
+		/// Add tasks to be executed remotely
+		/// </summary>
+		/// <param name="clusterId">Id of the compute cluster</param>
+		/// <returns></returns>
+		[HttpPost]
+		[Authorize]
+		[Route("/api/v2/compute/{clusterId}/tunnel")]
+		public ActionResult<GetComputeTunnelResponse> GetTunnelInfoAsync(ClusterId clusterId)
+		{
+			ComputeClusterConfig? clusterConfig;
+			if (!_globalConfig.Value.TryGetComputeCluster(clusterId, out clusterConfig))
+			{
+				return NotFound(clusterId);
+			}
+
+			ServerSettings settings = _globalConfig.Value.ServerSettings;
+			if (settings.ComputeInitiatorPort == 0 || settings.ComputeRemotePort == 0)
+			{
+				return NotFound("Tunnelling is not configured on the server");
+			}
+
+			GetComputeTunnelResponse response = new GetComputeTunnelResponse();
+			response.Ip = HttpContext.Connection.LocalIpAddress?.ToString();
+			response.InitiatorPort = _globalConfig.Value.ServerSettings.ComputeInitiatorPort;
+			response.RemotePort = _globalConfig.Value.ServerSettings.ComputeRemotePort;
+			return response;
+		}
 	}
 }
