Horde: Fix method for querying tunnel configuration.

diff --git a/Engine/Source/Programs/Horde/Horde.Build/Compute/ComputeController.cs b/Engine/Source/Programs/Horde/Horde.Build/Compute/ComputeController.cs
--- a/Engine/Source/Programs/Horde/Horde.Build/Compute/ComputeController.cs
+++ b/Engine/Source/Programs/Horde/Horde.Build/Compute/ComputeController.cs
@@ -158,13 +158,12 @@ namespace Horde.Build.Compute
 		/// </summary>
 		/// <param name="clusterId">Id of the compute cluster</param>
 		/// <returns></returns>
-		[HttpPost]
+		[HttpGet]
 		[Authorize]
 		[Route("/api/v2/compute/{clusterId}/tunnel")]
 		public ActionResult<GetComputeTunnelResponse> GetTunnelInfoAsync(ClusterId clusterId)
 		{
-			ComputeClusterConfig? clusterConfig;
-			if (!_globalConfig.Value.TryGetComputeCluster(clusterId, out clusterConfig))
+			if (!_globalConfig.Value.TryGetComputeCluster(clusterId, out _))
 			{
 				return NotFound(clusterId);
 			}
