Horde: Adding dashboard preview workflow
#rnx

diff --git a/Engine/Source/Programs/Horde/Horde.Build.Tests/DashboardTests.cs b/Engine/Source/Programs/Horde/Horde.Build.Tests/DashboardTests.cs
new file mode 100644
--- /dev/null
+++ b/Engine/Source/Programs/Horde/Horde.Build.Tests/DashboardTests.cs
@@ -0,0 +1,73 @@
+// Copyright Epic Games, Inc. All Rights Reserved.
+
+using System.Collections.Generic;
+using System.Threading.Tasks;
+using Microsoft.VisualStudio.TestTools.UnitTesting;
+using Horde.Build.Dashboard;
+using System.Linq;
+
+namespace Horde.Build.Tests
+{
+	/// <summary>
+	/// Tests for the dashboard
+	/// </summary>
+	[TestClass]
+	public class DashboardTest : TestSetup
+	{
+		[TestMethod]
+		public async Task TestPreviews()
+		{
+			CreateDashboardPreviewRequest request = new CreateDashboardPreviewRequest();
+			request.Summary = "This is a test preview item";
+
+			// create a reservation
+			GetDashboardPreviewResponse result =  (await DashboardController!.CreateDashbordPreview(request)).Value!;
+
+			Assert.AreEqual(1, result.Id);
+			Assert.AreEqual("This is a test preview item", result.Summary);
+
+			request.Summary = "This is another test preview item";
+			result = (await DashboardController!.CreateDashbordPreview(request)).Value!;
+
+			Assert.AreEqual(2, result.Id);
+			Assert.AreEqual("This is another test preview item", result.Summary);
+
+			UpdateDashboardPreviewRequest updateRequest = new UpdateDashboardPreviewRequest();
+			updateRequest.Id = 2;
+			updateRequest.Summary = "Updated test summary";
+			updateRequest.ExampleLink = "http://testexample";
+			updateRequest.TrackingLink = "http://testtracking";
+			updateRequest.DiscussionLink = "http://testdiscussion";
+			updateRequest.DeployedCL = 123;
+
+			result = (await DashboardController!.UpdateDashbordPreview(updateRequest)).Value!;
+
+			Assert.AreEqual(2, result.Id);
+			Assert.AreEqual("http://testexample", result.ExampleLink);
+			Assert.AreEqual("http://testtracking", result.TrackingLink);
+			Assert.AreEqual("http://testdiscussion", result.DiscussionLink);
+			Assert.AreEqual(123, result.DeployedCL);
+
+			List<GetDashboardPreviewResponse> results = (await DashboardController!.GetDashbordPreviews(true)).Value!;
+			results = results.OrderBy(r => r.Id).ToList();
+			Assert.AreEqual(2, results.Count);
+			Assert.AreEqual(1, results[0].Id);
+			Assert.AreEqual(2, results[1].Id);
+
+			UpdateDashboardPreviewRequest updateRequest2 = new UpdateDashboardPreviewRequest();
+			updateRequest2.Id = 1;
+			updateRequest2.Open = false;
+			result = (await DashboardController!.UpdateDashbordPreview(updateRequest2)).Value!;
+			Assert.AreEqual(false, result.Open);
+
+			results = (await DashboardController!.GetDashbordPreviews(true)).Value!;
+			Assert.AreEqual(1, results.Count);
+			Assert.AreEqual(2, results[0].Id);
+			Assert.AreEqual("Updated test summary", results[0].Summary);
+
+			results = (await DashboardController!.GetDashbordPreviews(false)).Value!;
+			Assert.AreEqual(1, results.Count);
+			Assert.AreEqual(1, results[0].Id);
+		}
+	}
+}
diff --git a/Engine/Source/Programs/Horde/Horde.Build.Tests/TestSetup.cs b/Engine/Source/Programs/Horde/Horde.Build.Tests/TestSetup.cs
--- a/Engine/Source/Programs/Horde/Horde.Build.Tests/TestSetup.cs
+++ b/Engine/Source/Programs/Horde/Horde.Build.Tests/TestSetup.cs
@@ -64,6 +64,7 @@ using Horde.Build.Devices;
 using Moq;
 using Horde.Build.Telemetry;
 using Horde.Build.Artifacts;
+using Horde.Build.Dashboard;
 
 namespace Horde.Build.Tests
 {
@@ -93,6 +94,7 @@ namespace Horde.Build.Tests
 		public ITestDataCollection TestDataCollection => ServiceProvider.GetRequiredService<ITestDataCollection>();
 		public IUserCollection UserCollection => ServiceProvider.GetRequiredService<IUserCollection>();
 		public IDeviceCollection DeviceCollection => ServiceProvider.GetRequiredService<IDeviceCollection>();
+		public IDashboardPreviewCollection DashboardPreviewCollection => ServiceProvider.GetRequiredService<IDashboardPreviewCollection>();
 
 		public AclService AclService => ServiceProvider.GetRequiredService<AclService>();
 		public FleetService FleetService => ServiceProvider.GetRequiredService<FleetService>();
@@ -129,6 +131,7 @@ namespace Horde.Build.Tests
 		public PoolsController PoolsController => GetPoolsController();
 		public LeasesController LeasesController => GetLeasesController();
 		public DevicesController DevicesController => GetDevicesController();
+		public DashboardController DashboardController => GetDashboardController();
 		public TestDataController TestDataController => GetTestDataController();
 
 		public ConfigService ConfigService => ServiceProvider.GetRequiredService<ConfigService>();
@@ -222,6 +225,7 @@ namespace Horde.Build.Tests
 			services.AddSingleton<IUgsMetadataCollection, UgsMetadataCollection>();
 			services.AddSingleton<IUserCollection, UserCollectionV1>();
 			services.AddSingleton<IDeviceCollection, DeviceCollection>();
+			services.AddSingleton<IDashboardPreviewCollection, DashboardPreviewCollection>();
 
 			services.AddSingleton<ToolCollection>();
 
@@ -301,6 +305,13 @@ namespace Horde.Build.Tests
 			return devicesCtrl;
 		}
 
+		private DashboardController GetDashboardController()
+		{
+			DashboardController dashboardCtrl = new DashboardController(DashboardPreviewCollection, ServerSettingsMon);
+			dashboardCtrl.ControllerContext = GetControllerContext();
+			return dashboardCtrl;
+		}
+
 		private TestDataController GetTestDataController()
 		{
 			TestDataController dataCtrl = new TestDataController(TestDataService, StreamCollection, JobService, TestDataCollection, GlobalConfigSnapshot);
diff --git a/Engine/Source/Programs/Horde/Horde.Build/Dashboard/DashboardController.cs b/Engine/Source/Programs/Horde/Horde.Build/Dashboard/DashboardController.cs
new file mode 100644
--- /dev/null
+++ b/Engine/Source/Programs/Horde/Horde.Build/Dashboard/DashboardController.cs
@@ -0,0 +1,198 @@
+// Copyright Epic Games, Inc. All Rights Reserved.
+
+using System;
+using System.Text;
+using System.Threading.Tasks;
+using Microsoft.AspNetCore.Authentication;
+using Microsoft.AspNetCore.Authentication.Cookies;
+using Microsoft.AspNetCore.Authorization;
+using Microsoft.AspNetCore.Mvc;
+using Microsoft.Extensions.Options;
+using Horde.Build.Server;
+using System.Collections.Generic;
+using System.Linq;
+using Horde.Build.Agents;
+
+namespace Horde.Build.Dashboard
+{
+	/// <summary>	
+	/// Dashboard authorization challenge controller	
+	/// </summary>	
+	[ApiController]
+	[Route("[controller]")]
+	public class DashboardController : Controller
+	{
+		/// <summary>
+		/// Authentication scheme in use
+		/// </summary>
+		readonly string _authenticationScheme;
+
+		/// <summary>
+		/// Server settings
+		/// </summary>
+		private readonly ServerSettings _settings;
+
+		private readonly IDashboardPreviewCollection _previewCollection;
+
+		/// <summary>
+		/// Constructor
+		/// </summary>
+		/// <param name="previewCollection" />
+		/// <param name="serverSettings">Server settings</param>
+		public DashboardController(IDashboardPreviewCollection previewCollection, IOptionsMonitor<ServerSettings> serverSettings)
+		{
+			_authenticationScheme = AccountController.GetAuthScheme(serverSettings.CurrentValue.AuthMethod);
+			_previewCollection = previewCollection;
+			_settings = serverSettings.CurrentValue;
+		}
+
+		/// <summary>	
+		/// Challenge endpoint for the dashboard, using cookie authentication scheme	
+		/// </summary>	
+		/// <returns>Ok on authorized, otherwise will 401</returns>	
+		[HttpGet]
+		[Authorize]
+		[Route("/api/v1/dashboard/challenge")]
+		public StatusCodeResult GetChallenge()
+		{
+			return Ok();
+		}
+
+		/// <summary>
+		/// Login to server, redirecting to the specified URL on success
+		/// </summary>
+		/// <param name="redirect"></param>
+		/// <returns></returns>
+		[HttpGet]
+		[Route("/api/v1/dashboard/login")]
+		public IActionResult Login([FromQuery] string? redirect)
+		{
+			return new ChallengeResult(_authenticationScheme, new AuthenticationProperties { RedirectUri = redirect ?? "/" });
+		}
+
+		/// <summary>
+		/// Login to server, redirecting to the specified Base64 encoded URL, which fixes some escaping issues on some auth providers, on success
+		/// </summary>
+		/// <param name="redirect"></param>
+		/// <returns></returns>
+		[HttpGet]
+		[Route("/api/v2/dashboard/login")]
+		public IActionResult LoginV2([FromQuery] string? redirect)
+		{
+			string? redirectUri = null;
+
+			if (redirect != null)
+			{
+				byte[] data = Convert.FromBase64String(redirect);
+				redirectUri = Encoding.UTF8.GetString(data);
+			}
+
+			return new ChallengeResult(_authenticationScheme, new AuthenticationProperties { RedirectUri = redirectUri ?? "/index" });
+		}
+
+		/// <summary>
+		/// Logout of the current account
+		/// </summary>
+		/// <returns></returns>
+		[HttpGet]
+		[Route("/api/v1/dashboard/logout")]
+		public async Task<StatusCodeResult> Logout()
+		{
+			await HttpContext.SignOutAsync(CookieAuthenticationDefaults.AuthenticationScheme);
+			try
+			{
+				await HttpContext.SignOutAsync(_authenticationScheme);
+			}
+			catch
+			{
+			}
+
+			return Ok();
+		}
+
+		/// <summary>
+		/// Query all the projects
+		/// </summary>
+		/// <returns>Config information needed by the dashboard</returns>
+		[HttpGet]
+		[Authorize]
+		[Route("/api/v1/dashboard/config")]
+		public ActionResult<GetDashboardConfigResponse> GetDashbordConfig()
+		{
+			GetDashboardConfigResponse dashboardConfigResponse = new GetDashboardConfigResponse();
+
+			if (_settings.JiraUrl != null)
+			{
+				dashboardConfigResponse.ExternalIssueServiceName = "Jira";
+				dashboardConfigResponse.ExternalIssueServiceUrl = _settings.JiraUrl.ToString().TrimEnd('/');
+			}
+
+			if (_settings.P4SwarmUrl != null)
+			{
+				dashboardConfigResponse.PerforceSwarmUrl = _settings.P4SwarmUrl.ToString().TrimEnd('/');
+			}
+
+			dashboardConfigResponse.HelpEmailAddress = _settings.HelpEmailAddress;
+			dashboardConfigResponse.HelpSlackChannel = _settings.HelpSlackChannel;
+
+			return dashboardConfigResponse;
+		}
+
+		/// <summary>
+		/// Create a new dashboard preview items
+		/// </summary>
+		/// <returns>Config information needed by the dashboard</returns>
+		[HttpPost]
+		[Authorize]
+		[Route("/api/v1/dashboard/preview")]
+		public async Task<ActionResult<GetDashboardPreviewResponse>> CreateDashbordPreview([FromBody] CreateDashboardPreviewRequest request)
+		{
+			IDashboardPreview preview = await _previewCollection.AddPreviewAsync(request.Summary);
+
+			if (!String.IsNullOrEmpty(request.ExampleLink) || !String.IsNullOrEmpty(request.DiscussionLink) || !String.IsNullOrEmpty(request.TrackingLink))
+			{
+				IDashboardPreview? updated = await _previewCollection.UpdatePreviewAsync(preview.Id, null, null, null, request.ExampleLink, request.DiscussionLink, request.TrackingLink);
+				if (updated == null) 
+				{
+					return NotFound(preview.Id);
+				}
+
+				return new GetDashboardPreviewResponse(updated);
+			}
+
+			return new GetDashboardPreviewResponse(preview);
+		}
+
+		/// <summary>
+		/// Create a new dashboard preview items
+		/// </summary>
+		/// <returns>Config information needed by the dashboard</returns>
+		[HttpPut]
+		[Authorize]
+		[Route("/api/v1/dashboard/preview")]
+		public async Task<ActionResult<GetDashboardPreviewResponse>> UpdateDashbordPreview([FromBody] UpdateDashboardPreviewRequest request)
+		{
+			IDashboardPreview? preview = await _previewCollection.UpdatePreviewAsync(request.Id, request.Summary, request.DeployedCL, request.Open, request.ExampleLink, request.DiscussionLink, request.TrackingLink);
+			
+			if (preview == null)
+			{
+				return NotFound(request.Id);
+			}			
+
+			return new GetDashboardPreviewResponse(preview);
+		}
+
+		/// <summary>
+		/// Query all the dashboard preview items
+		/// </summary>
+		/// <returns>Config information needed by the dashboard</returns>
+		[HttpGet]
+		[Authorize]
+		[Route("/api/v1/dashboard/previews")]
+		public async Task<ActionResult<List<GetDashboardPreviewResponse>>> GetDashbordPreviews([FromQuery] bool open = true)
+		{			
+			List <IDashboardPreview> previews = await _previewCollection.FindPreviewsAsync(open);			
+			return previews.Select(p => new GetDashboardPreviewResponse(p)).ToList();
+		}
+	}
+}
diff --git a/Engine/Source/Programs/Horde/Horde.Build/Dashboard/DashboardMessages.cs b/Engine/Source/Programs/Horde/Horde.Build/Dashboard/DashboardMessages.cs
new file mode 100644
--- /dev/null
+++ b/Engine/Source/Programs/Horde/Horde.Build/Dashboard/DashboardMessages.cs
@@ -0,0 +1,189 @@
+// Copyright Epic Games, Inc. All Rights Reserved.
+
+using System;
+using Amazon.EC2.Model;
+
+namespace Horde.Build.Dashboard
+{
+	/// <summary>
+	/// Setting information required by dashboard
+	/// </summary>
+	public class GetDashboardConfigResponse
+	{
+		/// <summary>
+		/// The name of the external issue service
+		/// </summary>
+		public string? ExternalIssueServiceName { get; set; }
+
+		/// <summary>
+		/// The url of the external issue service
+		/// </summary>
+		public string? ExternalIssueServiceUrl { get; set; }
+
+		/// <summary>
+		/// The url of the perforce swarm installation
+		/// </summary>
+		public string? PerforceSwarmUrl { get; set; }
+
+		/// <summary>
+		/// Help email address that users can contact with issues
+		/// </summary>
+		public string? HelpEmailAddress { get; set; }
+
+		/// <summary>
+		/// Help slack channel that users can use for issues
+		/// </summary>
+		public string? HelpSlackChannel { get; set; }
+	}
+
+	/// <summary>
+	/// 
+	/// </summary>
+	public class CreateDashboardPreviewRequest
+	{
+		/// <summary>
+		/// A summary of what the preview item changes
+		/// </summary>
+		public string Summary { get; set; }
+
+		/// <summary>
+		/// The CL the preview was deployed in
+		/// </summary>
+		public int? DeployedCL { get; set; }
+
+		/// <summary>
+		/// An example of the preview site users can view the changes
+		/// </summary>
+		public string? ExampleLink { get; set; }
+
+		/// <summary>
+		/// Optional Link for discussion the preview item
+		/// </summary>
+		public string? DiscussionLink { get; set; }
+
+		/// <summary>
+		/// Optional Link for discussing the preview item
+		/// </summary>
+		public string? TrackingLink { get; set; }
+
+		/// <summary>
+		/// Constructor
+		/// </summary>
+		public CreateDashboardPreviewRequest()
+		{
+			Summary = String.Empty;
+		}
+	}
+
+	/// <summary>
+	/// 
+	/// </summary>
+	public class UpdateDashboardPreviewRequest
+	{
+		/// <summary>
+		/// The preview item to update
+		/// </summary>
+		public int Id { get; set; }
+
+		/// <summary>
+		/// A summary of what the preview item changes
+		/// </summary>
+		public string? Summary { get; set; }
+
+		/// <summary>
+		/// The CL the preview was deployed in
+		/// </summary>
+		public int? DeployedCL { get; set; }
+
+		/// <summary>
+		/// Whather the preview is under consideration, if false the preview item didn't pass muster
+		/// </summary>
+		public bool? Open { get; set; }
+
+		/// <summary>
+		/// An example of the preview site users can view the changes
+		/// </summary>
+		public string? ExampleLink { get; set; }
+
+		/// <summary>
+		/// Optional Link for discussion the preview item
+		/// </summary>
+		public string? DiscussionLink { get; set; }
+
+		/// <summary>
+		/// Optional Link for discussing the preview item
+		/// </summary>
+		public string? TrackingLink { get; set; }
+
+		/// <summary>
+		/// Constructor
+		/// </summary>
+		public UpdateDashboardPreviewRequest()
+		{
+			Id = 0;
+		}
+	}
+
+
+	/// <summary>
+	/// 
+	/// </summary>
+	public class GetDashboardPreviewResponse
+	{
+		/// <summary>
+		/// The unique ID of the preview item
+		/// </summary>
+		public int Id { get; set; }
+
+		/// <summary>
+		/// When the preview item was created
+		/// </summary>
+		public DateTime CreatedAt { get; set; }
+
+		/// <summary>
+		/// A summary of what the preview item changes
+		/// </summary>
+		public string Summary { get; set; }
+
+		/// <summary>
+		/// The CL the preview was deployed in
+		/// </summary>
+		public int? DeployedCL { get; set; }
+
+		/// <summary>
+		/// Whather the preview is under consideration, if false the preview item didn't pass muster
+		/// </summary>
+		public bool Open { get; set; }
+
+		/// <summary>
+		/// An example of the preview site users can view the changes
+		/// </summary>
+		public string? ExampleLink { get; set; }
+
+		/// <summary>
+		/// Optional Link for discussion the preview item
+		/// </summary>
+		public string? DiscussionLink { get; set; }
+
+		/// <summary>
+		/// Optional Link for discussing the preview item
+		/// </summary>
+		public string? TrackingLink { get; set; }
+
+		/// <summary>
+		/// Constructor
+		/// </summary>
+		/// <param name="preview"></param>
+		public GetDashboardPreviewResponse(IDashboardPreview preview)
+		{
+			Id = preview.Id;
+			CreatedAt = preview.CreatedAt;
+			Summary = preview.Summary;	
+			DeployedCL= preview.DeployedCL;
+			Open = preview.Open;
+			ExampleLink = preview.ExampleLink;
+			DiscussionLink = preview.DiscussionLink;
+			TrackingLink = preview.TrackingLink;
+		}
+	}
+}
diff --git a/Engine/Source/Programs/Horde/Horde.Build/Dashboard/DashboardPreviewCollection.cs b/Engine/Source/Programs/Horde/Horde.Build/Dashboard/DashboardPreviewCollection.cs
new file mode 100644
--- /dev/null
+++ b/Engine/Source/Programs/Horde/Horde.Build/Dashboard/DashboardPreviewCollection.cs
@@ -0,0 +1,238 @@
+// Copyright Epic Games, Inc. All Rights Reserved.
+
+using System.Collections.Generic;
+using System;
+using System.Threading.Tasks;
+using MongoDB.Driver;
+using Horde.Build.Server;
+using Horde.Build.Utilities;
+using MongoDB.Bson.Serialization.Attributes;
+using System.Linq;
+
+namespace Horde.Build.Dashboard
+{
+
+	/// <summary>
+	/// 
+	/// </summary>
+	public interface IDashboardPreview
+	{
+		/// <summary>
+		/// The unique ID of the preview item
+		/// </summary>
+		public int Id { get; }
+
+		/// <summary>
+		/// When the preview item was created
+		/// </summary>
+		public DateTime CreatedAt { get; }
+
+		/// <summary>
+		/// A summary of what the preview item changes
+		/// </summary>
+		public string Summary { get; }
+
+		/// <summary>
+		/// The CL the preview was deployed in
+		/// </summary>
+		public int? DeployedCL { get; }
+
+		/// <summary>
+		/// Whather the preview is under consideration, if false the preview item didn't pass muster
+		/// </summary>
+		public bool Open { get; }
+
+		/// <summary>
+		/// An example of the preview site users can view the changes
+		/// </summary>
+		public string? ExampleLink { get; }
+
+		/// <summary>
+		/// Optional Link for discussion the preview item
+		/// </summary>
+		public string? DiscussionLink { get; }
+
+		/// <summary>
+		/// Optional Link for discussing the preview item
+		/// </summary>
+		public string? TrackingLink { get; }
+
+	}
+
+	/// <summary>
+	/// Collection of preview documents
+	/// </summary>
+	public interface IDashboardPreviewCollection
+	{
+
+		/// <summary>
+		/// Add a preview item
+		/// </summary>
+		/// <param name="summary"></param>
+		/// <returns></returns>
+		Task<IDashboardPreview> AddPreviewAsync(string summary);
+
+		/// <summary>
+		/// Update a dashboard preview item
+		/// </summary>
+		/// <param name="previewId"></param>
+		/// <param name="summary"></param>
+		/// <param name="deployedCL"></param>
+		/// <param name="open"></param>
+		/// <param name="exampleLink"></param>
+		/// <param name="discussionLink"></param>
+		/// <param name="trackingLink"></param>
+		/// <returns></returns>
+		Task<IDashboardPreview?> UpdatePreviewAsync(int previewId, string? summary = null, int? deployedCL = null, bool? open = null, string? exampleLink = null, string? discussionLink = null, string? trackingLink = null);
+
+		/// <summary>
+		/// Finds preview items with the given criterial
+		/// </summary>
+		/// <param name="open"></param>
+		/// <returns></returns>
+		Task<List<IDashboardPreview>> FindPreviewsAsync(bool? open = null);
+	}
+
+	internal class DashboardPreviewCollection : IDashboardPreviewCollection
+	{
+		[SingletonDocument("dashboard-preview-ledger", "63f8e2154a2e859581b5bb24")]
+		internal class PreviewLedger : SingletonBase
+		{
+			public int NextId { get; set; }
+		}
+
+
+		internal class PreviewDocument : IDashboardPreview
+		{
+			[BsonId]
+			public int Id { get; set; }
+
+			public DateTime CreatedAt { get; set; }
+
+			public string Summary { get; set; }
+
+			public bool Open { get; set; }
+
+			[BsonIgnoreIfNull]
+			public int? DeployedCL { get; set; }
+
+			[BsonIgnoreIfNull]
+			public string? ExampleLink { get; set; }
+
+			[BsonIgnoreIfNull]
+			public string? DiscussionLink { get; set; }
+
+			[BsonIgnoreIfNull]
+			public string? TrackingLink { get; set; }
+
+			[BsonConstructor]
+			private PreviewDocument()
+			{
+				Summary = String.Empty;
+			}
+
+			public PreviewDocument(int id, string summary)
+			{
+				Id = id;
+				Summary = summary;
+				CreatedAt = DateTime.UtcNow;
+				Open = true;
+			}
+		}
+
+		readonly IMongoCollection<PreviewDocument> _previews;
+		readonly ISingletonDocument<PreviewLedger> _ledgerSingleton;
+
+		public DashboardPreviewCollection(MongoService mongoService)
+		{
+			_ledgerSingleton = new SingletonDocument<PreviewLedger>(mongoService);
+			_previews = mongoService.GetCollection<PreviewDocument>("DashboardPreviews");
+		}
+
+		public async Task<IDashboardPreview> AddPreviewAsync(string summary)
+		{
+			PreviewLedger ledger = await _ledgerSingleton.UpdateAsync(x => x.NextId++);
+
+			PreviewDocument newPreview = new PreviewDocument(ledger.NextId, summary);
+			await _previews.InsertOneAsync(newPreview);
+			
+
+			return newPreview;
+		}
+
+		public async Task<IDashboardPreview> UpdatePreviewAsync(string summary)
+		{
+			PreviewLedger ledger = await _ledgerSingleton.UpdateAsync(x => x.NextId++);
+
+			PreviewDocument newPreview = new PreviewDocument(ledger.NextId, summary);
+			await _previews.InsertOneAsync(newPreview);
+
+
+			return newPreview;
+		}
+
+		public async Task<IDashboardPreview?> UpdatePreviewAsync(int previewId, string? summary = null, int? deployedCL = null, bool? open = null, string? exampleLink = null, string? discussionLink = null, string? trackingLink = null )
+		{
+			UpdateDefinitionBuilder<PreviewDocument> updateBuilder = Builders<PreviewDocument>.Update;
+			List<UpdateDefinition<PreviewDocument>> updates = new List<UpdateDefinition<PreviewDocument>>();
+
+			if (summary != null)
+			{
+				updates.Add(updateBuilder.Set(x => x.Summary, summary));
+			}
+
+			if (deployedCL != null)
+			{
+				updates.Add(updateBuilder.Set(x => x.DeployedCL, deployedCL));
+			}
+
+			if (open != null)
+			{
+				updates.Add(updateBuilder.Set(x => x.Open, open));
+			}
+
+			if (exampleLink != null)
+			{
+				updates.Add(updateBuilder.Set(x => x.ExampleLink, exampleLink));
+			}
+
+			if (discussionLink != null)
+			{
+				updates.Add(updateBuilder.Set(x => x.DiscussionLink, discussionLink));
+			}
+
+			if (trackingLink != null)
+			{
+				updates.Add(updateBuilder.Set(x => x.TrackingLink, trackingLink));
+			}
+
+			if (updates.Count == 0)
+			{
+				await _previews.Find(x => x.Id == previewId).FirstOrDefaultAsync();
+			}
+
+			if (await _previews.FindOneAndUpdateAsync(x => x.Id == previewId, updateBuilder.Combine(updates)) == null)
+			{
+				return null;
+			}
+
+			return await _previews.Find(x => x.Id == previewId).FirstOrDefaultAsync();
+
+		}
+
+		public async Task<List<IDashboardPreview>> FindPreviewsAsync(bool? open = null)
+		{
+			FilterDefinition<PreviewDocument> filter = Builders<PreviewDocument>.Filter.Empty;
+			if (open != null)
+			{
+				filter &= Builders<PreviewDocument>.Filter.Eq(x => x.Open, open);
+			}
+
+			List<PreviewDocument> results = await _previews.Find(filter).ToListAsync();
+			return results.Select<PreviewDocument, IDashboardPreview>(x => x).ToList();
+		}
+
+	}
+
+}
+
diff --git a/Engine/Source/Programs/Horde/Horde.Build/Server/DashboardController.Messages.cs b/Engine/Source/Programs/Horde/Horde.Build/Server/DashboardController.Messages.cs
deleted file mode 100644
diff --git a/Engine/Source/Programs/Horde/Horde.Build/Server/DashboardController.cs b/Engine/Source/Programs/Horde/Horde.Build/Server/DashboardController.cs
deleted file mode 100644
diff --git a/Engine/Source/Programs/Horde/Horde.Build/Startup.cs b/Engine/Source/Programs/Horde/Horde.Build/Startup.cs
--- a/Engine/Source/Programs/Horde/Horde.Build/Startup.cs
+++ b/Engine/Source/Programs/Horde/Horde.Build/Startup.cs
@@ -32,6 +32,7 @@ using Horde.Build.Agents.Sessions;
 using Horde.Build.Agents.Software;
 using Horde.Build.Agents.Telemetry;
 using Horde.Build.Configuration;
+using Horde.Build.Dashboard;
 using Horde.Build.Devices;
 using Horde.Build.Issues;
 using Horde.Build.Issues.External;
@@ -360,6 +361,7 @@ namespace Horde.Build
 			services.AddSingleton<IUserCollection, UserCollectionV2>();
 			services.AddSingleton<IDeviceCollection, DeviceCollection>();
 			services.AddSingleton<INoticeCollection, NoticeCollection>();
+			services.AddSingleton<IDashboardPreviewCollection, DashboardPreviewCollection>();
 
 			services.AddSingleton<ToolCollection>();
 
