Horde - Add missing general/framework Gauntlet known log event (mostly used in warnings)

diff --git a/Engine/Source/Programs/Horde/Horde.Build/Issues/Handlers/GauntletIssueHandler.cs b/Engine/Source/Programs/Horde/Horde.Build/Issues/Handlers/GauntletIssueHandler.cs
--- a/Engine/Source/Programs/Horde/Horde.Build/Issues/Handlers/GauntletIssueHandler.cs
+++ b/Engine/Source/Programs/Horde/Horde.Build/Issues/Handlers/GauntletIssueHandler.cs
@@ -19,6 +19,11 @@ namespace Horde.Build.Issues.Handlers
 	/// </summary>
 	class GauntletIssueHandler : IssueHandler
 	{
+		/// <summary>
+		/// Prefix for framework keys
+		/// </summary>
+		const string FrameworkPrefix = "framework";
+
 		/// <summary>
 		/// Prefix for test keys
 		/// </summary>
@@ -54,6 +59,7 @@ namespace Horde.Build.Issues.Handlers
 		///  Known Gauntlet events
 		/// </summary>
 		static readonly Dictionary<EventId, string> s_knownGauntletEvents = new Dictionary<EventId, string> {
+			{ KnownLogEvents.Gauntlet, FrameworkPrefix},
 			{ KnownLogEvents.Gauntlet_TestEvent, TestPrefix},
 			{ KnownLogEvents.Gauntlet_DeviceEvent, DevicePrefix},
 			{ KnownLogEvents.Gauntlet_UnrealEngineTestEvent, TestPrefix},
