Horde: Add TunnelService as a hosted service.

diff --git a/Engine/Source/Programs/Horde/Horde.Build/Startup.cs b/Engine/Source/Programs/Horde/Horde.Build/Startup.cs
--- a/Engine/Source/Programs/Horde/Horde.Build/Startup.cs
+++ b/Engine/Source/Programs/Horde/Horde.Build/Startup.cs
@@ -389,10 +389,13 @@ namespace Horde.Build
 			// Though most fleet managers are created on a per-pool basis
 			services.AddSingleton<IFleetManager>(ctx => ctx.GetRequiredService<IFleetManagerFactory>().CreateFleetManager(FleetManagerType.Default));
 
+			// Run the tunnel service for all runmodes
+			services.AddSingleton<TunnelService>();
+			services.AddHostedService<TunnelService>(sp => sp.GetRequiredService<TunnelService>());
+
 			services.AddSingleton<AclService>();
 			services.AddSingleton<AgentService>();			
 			services.AddSingleton<AgentSoftwareService>();
-			services.AddSingleton<TunnelService>();
 			services.AddSingleton<ConsistencyService>();
 			services.AddSingleton<RequestTrackerService>();
 			services.AddSingleton<CredentialService>();
