Horde: Register default serializers when writing configuration.

diff --git a/Engine/Source/Programs/Horde/Horde.Build/Server/DebugController.cs b/Engine/Source/Programs/Horde/Horde.Build/Server/DebugController.cs
--- a/Engine/Source/Programs/Horde/Horde.Build/Server/DebugController.cs
+++ b/Engine/Source/Programs/Horde/Horde.Build/Server/DebugController.cs
@@ -328,10 +328,10 @@ namespace Horde.Build.Server
 				return Forbid();
 			}
 
-			JsonSerializerOptions options = new JsonSerializerOptions
-			{
-				WriteIndented = true
-			};
+			JsonSerializerOptions options = new JsonSerializerOptions();
+			options.WriteIndented = true;
+
+			Startup.ConfigureJsonSerializer(options);
 
 			return Ok(JsonSerializer.Serialize(_settings.Value, options));
 		}
