Horde: Listen to tunnel requests on IPv4 interfaces.

diff --git a/Engine/Source/Programs/Horde/Horde.Build/Compute/TunnelService.cs b/Engine/Source/Programs/Horde/Horde.Build/Compute/TunnelService.cs
--- a/Engine/Source/Programs/Horde/Horde.Build/Compute/TunnelService.cs
+++ b/Engine/Source/Programs/Horde/Horde.Build/Compute/TunnelService.cs
@@ -61,7 +61,7 @@ namespace Horde.Build.Compute
 		/// <inheritdoc/>
 		public Task StartAsync(CancellationToken cancellationToken)
 		{
-			Start(IPAddress.IPv6Any);
+			Start(IPAddress.Any);
 			return Task.CompletedTask;
 		}
 
