Horde: Fix default arguments not being specified when starting a chained job.

diff --git a/Engine/Source/Programs/Horde/Horde.Build/Jobs/JobService.cs b/Engine/Source/Programs/Horde/Horde.Build/Jobs/JobService.cs
--- a/Engine/Source/Programs/Horde/Horde.Build/Jobs/JobService.cs
+++ b/Engine/Source/Programs/Horde/Horde.Build/Jobs/JobService.cs
@@ -1278,6 +1278,7 @@ namespace Horde.Build.Jobs
 					CreateJobOptions options = new CreateJobOptions(templateRefConfig);
 					options.PreflightChange = newJob.PreflightChange;
 					options.PreflightDescription = newJob.PreflightDescription;
+					options.Arguments.AddRange(template.GetDefaultArguments());
 
 					await CreateJobAsync(chainedJobId, streamConfig, jobTrigger.TemplateRefId, template.Hash, triggerGraph, templateRefConfig.Name, newJob.Change, newJob.CodeChange, options);
 					return newJob;
