Horde: Add job option for collecting .ib_mon files

diff --git a/Engine/Source/Programs/Horde/Horde.Agent/Execution/JobExecutor.cs b/Engine/Source/Programs/Horde/Horde.Agent/Execution/JobExecutor.cs
--- a/Engine/Source/Programs/Horde/Horde.Agent/Execution/JobExecutor.cs
+++ b/Engine/Source/Programs/Horde/Horde.Agent/Execution/JobExecutor.cs
@@ -655,7 +655,7 @@ namespace Horde.Agent.Execution
 		
 		private async Task UploadXgeMonitorFilesAsync(BeginStepResponse step, ILogger logger, CancellationToken cancellationToken)
 		{
-			if (_xgeMetadataExtractor != null)
+			if (_xgeMetadataExtractor != null && _jobOptions.CollectIbMonFilesAsArtifacts is true)
 			{
 				using IScope scope = GlobalTracer.Instance.BuildSpan("XgeMonitorFilesUpload").StartActive();
 				List<FileReference> ibMonFiles = _xgeMetadataExtractor.GetLocalIbMonFilePaths();
@@ -718,6 +718,7 @@ namespace Horde.Agent.Execution
 			if (_jobOptions.UseNewTempStorage ?? false)
 			{
 				bool result = await ExecuteWithTempStorageAsync(step, workspaceDir, arguments.ToString(), useP4, logger, cancellationToken);
+				await UploadXgeMonitorFilesAsync(step, logger, cancellationToken);
 				return result;
 			}
 			else
@@ -728,6 +729,7 @@ namespace Horde.Agent.Execution
 				}
 				
 				bool result = await ExecuteAutomationToolAsync(step, workspaceDir, arguments.ToString(), useP4, logger, cancellationToken) == 0;
+				await UploadXgeMonitorFilesAsync(step, logger, cancellationToken);
 				return result;
 			}
 		}
