Horde: Fix test for pdbutil matcher.

diff --git a/Engine/Source/Programs/AutomationTool/AutomationUtils/Matchers/SystemicEventMatcher.cs b/Engine/Source/Programs/AutomationTool/AutomationUtils/Matchers/SystemicEventMatcher.cs
--- a/Engine/Source/Programs/AutomationTool/AutomationUtils/Matchers/SystemicEventMatcher.cs
+++ b/Engine/Source/Programs/AutomationTool/AutomationUtils/Matchers/SystemicEventMatcher.cs
@@ -34,7 +34,7 @@ namespace AutomationUtils.Matchers
 				{
 					builder.MoveNext();
 				}
-				return new LogEventBuilder(cursor).ToMatch(LogEventPriority.High, LogLevel.Information, KnownLogEvents.Systemic_PdbUtil);
+				return builder.ToMatch(LogEventPriority.High, LogLevel.Information, KnownLogEvents.Systemic_PdbUtil);
 			}
 			if (cursor.IsMatch(s_roboMerge))
 			{
diff --git a/Engine/Source/Programs/Horde/Horde.Agent.Tests/LogParserTests.cs b/Engine/Source/Programs/Horde/Horde.Agent.Tests/LogParserTests.cs
--- a/Engine/Source/Programs/Horde/Horde.Agent.Tests/LogParserTests.cs
+++ b/Engine/Source/Programs/Horde/Horde.Agent.Tests/LogParserTests.cs
@@ -187,7 +187,6 @@ namespace Horde.Agent.Tests
 		}
 
 		[TestMethod]
-		[Ignore]
 		public void SymbolStripSpuriousEventMatcher()
 		{
 			// Symbol stripping error
@@ -201,7 +200,7 @@ namespace Horde.Agent.Tests
 				};
 
 				List<LogEvent> logEvents = Parse(String.Join("\n", lines));
-				CheckEventGroup(logEvents, 1, 1, LogLevel.Information, KnownLogEvents.Systemic_PdbUtil);
+				CheckEventGroup(logEvents, 1, 2, LogLevel.Information, KnownLogEvents.Systemic_PdbUtil);
 			}
 		}
 
