Horde: Fix stream auth when creating Jira issues
#rnx

diff --git a/Engine/Source/Programs/Horde/Horde.Build/Issues/IssuesController.cs b/Engine/Source/Programs/Horde/Horde.Build/Issues/IssuesController.cs
--- a/Engine/Source/Programs/Horde/Horde.Build/Issues/IssuesController.cs
+++ b/Engine/Source/Programs/Horde/Horde.Build/Issues/IssuesController.cs
@@ -654,11 +654,17 @@ namespace Horde.Build.Issues
 			}
 
 			StreamId streamIdValue = new StreamId(issueRequest.StreamId);
-			if (!_globalConfig.Value.Authorize(AclAction.ViewStream, User))
+
+			if (!_globalConfig.Value.TryGetStream(streamIdValue, out StreamConfig? streamConfig))
+			{
+				return NotFound(streamIdValue);
+			}
+			if (!streamConfig.Authorize(AclAction.ViewStream, User))
 			{
 				return Forbid(AclAction.ViewStream, streamIdValue);
 			}
 
+
 			IUser? user = await _userCollection.GetUserAsync(User);
 
 			if (user == null)
