Horde: Updates to job timeline
#rnx
#skipci

diff --git a/Engine/Source/Programs/Horde/HordeDashboard/src/components/jobDetailsV2/JobDetailTimeline.tsx b/Engine/Source/Programs/Horde/HordeDashboard/src/components/jobDetailsV2/JobDetailTimeline.tsx
--- a/Engine/Source/Programs/Horde/HordeDashboard/src/components/jobDetailsV2/JobDetailTimeline.tsx
+++ b/Engine/Source/Programs/Horde/HordeDashboard/src/components/jobDetailsV2/JobDetailTimeline.tsx
@@ -1,19 +1,20 @@
 // Copyright Epic Games, Inc. All Rights Reserved.
 
-import { Slider, Spinner, SpinnerSize, Stack, Text } from "@fluentui/react";
+import { FontIcon, Slider, Spinner, SpinnerSize, Stack, Text } from "@fluentui/react";
 import * as d3 from "d3";
+import { action, makeObservable, observable } from "mobx";
 import { observer } from "mobx-react-lite";
+import moment from "moment";
 import { useEffect, useState } from "react";
+import { Link, NavigateFunction, useNavigate } from "react-router-dom";
 import { GetBatchResponse, GetStepResponse, JobStepOutcome, StepData } from "../../backend/Api";
+import dashboard, { StatusColor } from "../../backend/Dashboard";
 import { ISideRailLink } from "../../base/components/SideRail";
+import { msecToElapsed } from "../../base/utilities/timeUtils";
 import { hordeClasses, modeColors } from "../../styles/Styles";
+import { HistoryModal } from "../HistoryModal";
 import { JobDataView, JobDetailsV2 } from "./JobDetailsViewCommon";
 
-import moment from "moment";
-import { NavigateFunction, useNavigate } from "react-router-dom";
-import dashboard, { StatusColor } from "../../backend/Dashboard";
-import { msecToElapsed } from "../../base/utilities/timeUtils";
-
 type SelectionType = d3.Selection<SVGGElement, unknown, null, undefined>;
 
 type Batch = GetBatchResponse & {
@@ -55,6 +56,36 @@ type Span = {
 
 const sideRail: ISideRailLink = { text: "Timeline", url: "rail_timeline" };
 
+class Tooltip {
+
+   constructor() {
+      makeObservable(this);
+   }
+
+   @observable
+   updated = 0
+
+   @action
+   update(span?: Span, x?: number, y?: number) {
+      this.x = x ?? 0;
+      this.y = y ?? 0;
+      this.span = span;
+      this.updated++;
+   }
+
+   @action
+   freeze(frozen: boolean) {
+      this.frozen = frozen;
+      this.updated++;
+   }
+
+   span?: Span;
+   x: number = 0;
+   y: number = 0;
+   frozen: boolean = false;
+}
+
+
 class TimelineDataView extends JobDataView {
 
    filterUpdated() {
@@ -148,7 +179,7 @@ class TimelineDataView extends JobDataView {
          maxTime = new Date(filterStep.finishTime).getTime();
       }
 
-      let batches: Batch[] = [];      
+      let batches: Batch[] = [];
 
       bresponses.forEach(b => {
 
@@ -344,6 +375,8 @@ class TimelineDataView extends JobDataView {
    filterCost: number = 0;
    filterStepId?: string;
 
+   tooltip = new Tooltip();
+
    dirty = false;
 
 }
@@ -474,6 +507,17 @@ class TimelineRenderer {
             return colors[span.step?.outcome ?? "Unspecified"];
          })
 
+      g.append("g").selectAll("line")
+         .data(([, I]: any) => I)
+         .join("line")
+         .attr("class", "cline")
+         .attr("x1", 0)
+         .attr("x2", 0)
+         .attr("stroke-width", 0)
+         .attr("y1", i => 0)
+         .attr("y2", i => 0)
+         .attr("stroke-linecap", 0)
+         .attr("stroke", "#FF00FF")
 
 
       const xAxis = (g: SelectionType) => {
@@ -499,7 +543,7 @@ class TimelineRenderer {
       function zoomed(event: any) {
          xScale.range([margin.left, width - margin.right].map(d => event.transform.applyX(d)));
          svg!.selectAll(".separator")
-            .attr("x1", i => { console.log("hi"); if (i as number >= X.length) return 0; return xScale(X[i as any].utcStart.getTime() / 1000) })
+            .attr("x1", i => { if (i as number >= X.length) return 0; return xScale(X[i as any].utcStart.getTime() / 1000) })
             .attr("x2", i => { if (i as number >= X.length) return 0; return xScale(X[i as any].utcStart.getTime() / 1000) })
          svg!.selectAll(".spanline")
             .attr("x1", i => { if (i as number >= X.length) return 0; return xScale(X[i as any].utcStart.getTime() / 1000) })
@@ -519,21 +563,6 @@ class TimelineRenderer {
 
       svg.call(zoom as any);
 
-      const tooltip = d3.select(container)
-         .append("div")
-         .attr("id", "tooltip")
-         .style("display", "none")
-         .style("background-color", modeColors.background)
-         .style("border", "solid")
-         .style("border-width", "1px")
-         .style("border-radius", "3px")
-         .style("border-color", dashboard.darktheme ? "#413F3D" : "#2D3F5F")
-         .style("padding", "8px")
-         .style("position", "absolute")
-         .style("pointer-events", "none");
-
-
-
       const closestData = (x: number, y: number): Span | undefined => {
 
          y -= 16;
@@ -579,106 +608,68 @@ class TimelineRenderer {
 
       }
 
-      let mousedown = false;
-
       const handleMouseMove = (event: any) => {
 
+         if (dataView.tooltip.frozen) {
+            return;
+         }
+
          let mouseX = d3.pointer(event)[0];
          let mouseY = d3.pointer(event)[1];
 
          const span = closestData(mouseX, mouseY);
 
-         if (!span) {
-            tooltip.style("display", "none");
-            return;
-         }
-
          if (span) {
-
-            const details = dataView.details!;
-
-            const group = details.groups[span.batch.groupIdx];
-            const pool = details.stream?.agentTypes[group?.agentType!];
-
-            let html = "";
-
-            let agentHtml = "";
-            agentHtml += `Type: ${group?.agentType.toUpperCase() ?? "Unknown Agent Type"}<br/>`;
-            agentHtml += `Pool: ${[pool?.pool?.toUpperCase() ?? "Unknown Pool"]}<br/>`;
-            agentHtml += `Agent: ${[span.batch.agentId ?? "Unknown Pool"]}<br/>`;
-
-            if (span.cost) {
-               agentHtml += `Estimated cost: $${span.cost.toFixed(2)}`;
-            }
-
-            const elapsed = msecToElapsed(span.utcFinish.getTime() - span.utcStart.getTime(), true, false);
-
+            const svg = this.svg!;
+            let color = "";
             if (span.type === 0) {
-
-               html += `Wait: ${elapsed}<br/>`;
-               html += agentHtml;
-            }
-            if (span.type === 1) {
-               html += `Init: <b>${elapsed}</b><br/>`;
-               html += agentHtml;
+               color = waitColor;
             }
-            if (span.type === 2) {
-               html = `${span.stepName ?? "Unknown Step"}<br/>`;
-               html += `Time: ${elapsed}<br/>`;
-               html += agentHtml;
+            else if (span.type === 1) {
+               color = initColor;
+            } else {
+               color = colors[span.step?.outcome ?? "Unspecified"];
             }
 
-            let topY = mouseY - 78;
-            if (topY < 0) {
-               topY = 0;
+            if (!dataView.tooltip.frozen) {
+               svg.select(`.cline`)
+                  .attr("x1", xScale(span.utcStart.getTime() / 1000) - 1)
+                  .attr("x2", xScale(span.utcFinish.getTime() / 1000) - 1)
+                  .attr("stroke-width", 10)
+                  .attr("y1", i => yScale(span.lane) as number + 16)
+                  .attr("y2", i => yScale(span.lane) as number + 16)
+                  .attr("stroke", color)
             }
+         }
 
-            let tipX = mouseX;
-            let translateX = "0%";
-
-            if (tipX > 1500) {
-               translateX = "-180%";
-            }
+         dataView.tooltip.update(span, d3.pointer(event, container)[0], d3.pointer(event, container)[1]);
 
-            tooltip
-               .style("display", mousedown ? "none" : "block")
-               .html(html ?? "")
-               .style("position", `absolute`)
-               .style("width", `max-content`)
-               .style("top", (topY) + "px")
-               .style("left", `${tipX}px`)
-               .style("transform", `translateX(${translateX})`)
-               .style("font-family", "Horde Open Sans Semibold")
-               .style("font-size", "10px")
-               .style("line-height", "16px")
-               .style("shapeRendering", "crispEdges")
-               .style("stroke", "none")
-         }
       }
 
       const handleMouseLeave = (event: any) => {
-         mousedown = false;
-         tooltip.style("display", "none");
+         if (!dataView.tooltip.frozen) {
+            dataView.tooltip.update(undefined);
+            this.svg!.select(`.cline`)
+               .attr("stroke-width", 0)
+         }
       }
 
       const handleMouseClick = (event: any) => {
 
-         let mouseX = d3.pointer(event)[0];
-         let mouseY = d3.pointer(event)[1];
-
-         const span = closestData(mouseX, mouseY);
+         if (dataView.tooltip.frozen) {
+            dataView.tooltip.freeze(false);
+            return;
+         }
 
-         if (span && span.step) {
-            navigate(`/job/${dataView.details!.jobId}?step=${span.step.id}`);
+         if (dataView.tooltip.span) {
+            dataView.tooltip.freeze(true);
          }
+
       }
 
 
       svg.on("click", (event) => handleMouseClick(event))
 
-      svg.on("pointerdown", (event) => { mousedown = true; tooltip.style("display", "none"); });
-      svg.on("pointerup", (event) => { mousedown = false; tooltip.style("display", "block"); });
-
       svg.on("mousemove", (event) => handleMouseMove(event));
       svg.on("mouseleave", (event) => { handleMouseLeave(event); })
       svg.on("wheel", (event) => { event.preventDefault(); })
@@ -695,6 +686,154 @@ class TimelineRenderer {
    forceRender = false;
 }
 
+const GraphTooltip: React.FC<{ dataView: TimelineDataView }> = observer(({ dataView }) => {
+
+   const [viewAgent, setViewAgent] = useState("");
+
+   // subscribe
+   if (dataView.tooltip.updated) { }
+
+   const tooltip = dataView.tooltip;
+   const span = tooltip.span;
+
+   if (!span) {
+      return null;
+   }
+
+   const textSize = 12;
+
+   let tipX = tooltip.x;
+   let offsetX = 32;
+   let translateX = "0%";
+
+   if (tipX > 1000) {
+      offsetX = -32;
+      translateX = "-100%";
+   }
+
+   const translateY = "-50%";
+
+   const details = dataView.details!;
+
+   const group = details.groups[span.batch.groupIdx];
+   const pool = details.stream?.agentTypes[group?.agentType!];
+
+   const agentPool = pool?.pool?.toUpperCase() ? pool.pool.toUpperCase() : "Unknown Pool";
+   const agent = span.batch.agentId ? span.batch.agentId : "Unknown Agent";
+
+   const titleWidth = 48;
+
+   const dataElement = (title: string, value: string, link?: string, agent?: string, linkTarget?: string) => {
+
+      let valueElement = <Stack>
+         <Stack style={{ fontSize: textSize }}>
+            {value}
+         </Stack>
+      </Stack>
+
+
+      if (link) {
+         valueElement = <Link className="horde-link" to={link} target={linkTarget}>
+            {valueElement}
+         </Link>
+      }
+
+      if (agent) {
+         valueElement = <button className="horde-link" style={{borderWidth: 0, padding: 0, backgroundColor: modeColors.background}}>
+            {valueElement}
+         </button>
+      }
+
+      let component = <Stack style={{ cursor: (link || agent) ? "pointer" : undefined }} onClick={() => {
+         if (agent) {
+            setViewAgent(agent);
+         }
+
+      }}>
+         <Stack horizontal>
+            <Stack style={{ width: titleWidth }}>
+               {!!title && <Text style={{ fontFamily: "Horde Open Sans SemiBold", fontSize: textSize }} >{title}</Text>}
+            </Stack>
+            {valueElement}
+         </Stack>
+      </Stack>
+
+
+      return component;
+   }
+
+   const elements: JSX.Element[] = [];
+
+   if (span.step) {
+      elements.push(dataElement("Step:", span.stepName ?? "Unknown Step", `/job/${details!.jobId}?step=${span.step.id}`));
+      if (span.step.logId) {
+         elements.push(dataElement("", "View Log", `/log/${span.step.logId}`, undefined, "_blank"));
+      }
+   }
+
+   elements.push(dataElement("Agent:", agent, undefined, span.batch.agentId ? span.batch.agentId : undefined));
+   //elements.push(dataElement("Type:", agentType));
+
+   let poolLink: string | undefined;
+   if (pool?.pool) {
+      poolLink = `/pools?pool=${pool.pool}`
+   }
+
+   elements.push(dataElement("Pool:", agentPool, poolLink, undefined, "_blank"));
+
+   if (span.batch) {
+      elements.push(dataElement("Batch:", `View Batch`, `/job/${details!.jobId}?batch=${span.batch.id}`));
+   }
+
+   if (span.batch.logId) {
+      elements.push(dataElement("", "View Log", `/log/${span.batch.logId}`, undefined, "_blank"));
+   }
+
+
+   if (span.cost) {
+      elements.push(dataElement("Cost:", `$${span.cost.toFixed(2)}`));
+   }
+
+   const elapsed = msecToElapsed(span.utcFinish.getTime() - span.utcStart.getTime(), true, false);
+
+   let spanType = "Wait:";
+   if (span.type === 1) {
+      spanType = "Init:";
+   }
+   else if (span.type === 2) {
+      spanType = "Time:";
+   }
+
+   elements.push(dataElement(spanType, elapsed));
+
+   return <div style={{
+      position: "absolute",
+      display: "block",
+      top: `${tooltip.y}px`,
+      left: `${tooltip.x + offsetX}px`,
+      backgroundColor: modeColors.background,
+      zIndex: 1,
+      border: "solid",
+      borderWidth: "1px",
+      borderRadius: "3px",
+      width: "max-content",
+      borderColor: dashboard.darktheme ? "#413F3D" : "#2D3F5F",
+      pointerEvents: tooltip.frozen ? undefined : "none",
+      transform: `translate(${translateX}, ${translateY})`
+   }}>
+      {!!viewAgent && <HistoryModal agentId={viewAgent} onDismiss={() => setViewAgent("")} />}
+      <Stack horizontal>
+         <Stack tokens={{ childrenGap: 6, padding: 16 }}>
+            {elements}
+         </Stack>
+         {!!tooltip.frozen && <Stack style={{ paddingLeft: 32, paddingRight: 12, paddingTop: 12, cursor: "pointer" }} onClick={() => { tooltip.freeze(false); tooltip.update() }}>
+            <FontIcon iconName="Cancel" />
+         </Stack>}
+         {!tooltip.frozen && <Stack style={{ paddingLeft: 32, paddingRight: 12 + 16, paddingTop: 12 }} />}
+
+      </Stack>
+   </div>
+})
 
 const TimelineGraph: React.FC<{ dataView: TimelineDataView, filterTime: number, filterCost: number, stepId?: string }> = ({ dataView, filterTime, filterCost, stepId }) => {
 
@@ -730,6 +869,10 @@ const TimelineGraph: React.FC<{ dataView: TimelineDataView, filterTime: number,
 
    return <Stack className={hordeClasses.horde}>
       <Stack style={{ paddingLeft: 8, paddingTop: 8 }}>
+         <div style={{ position: "relative" }}>
+            <GraphTooltip dataView={dataView} />
+         </div>
+
          <div id={graph_container_id} className="horde-no-darktheme" style={{ shapeRendering: "crispEdges", userSelect: "none", position: "relative" }} ref={(ref: HTMLDivElement) => setContainer(ref)} onMouseEnter={() => { }} onMouseLeave={() => { }} />
       </Stack>
    </Stack>;
diff --git a/Engine/Source/Programs/Horde/HordeDashboard/src/index.css b/Engine/Source/Programs/Horde/HordeDashboard/src/index.css
--- a/Engine/Source/Programs/Horde/HordeDashboard/src/index.css
+++ b/Engine/Source/Programs/Horde/HordeDashboard/src/index.css
@@ -32,6 +32,16 @@ code {
   color:#000000;
 }
 
+.horde-link, .horde-link:visited {
+  text-decoration: none;
+  color: rgb(0, 120, 212);
+}
+
+.horde-link:hover {
+  text-decoration: underline;
+  color: rgb(0, 120, 212)
+}
+
 .job-item:hover {
   background:rgb(243, 242, 241);
 }
