Horde: Add some font family settings to root site css
#rnx
#skipci

diff --git a/Engine/Source/Programs/Horde/HordeDashboard/src/index.css b/Engine/Source/Programs/Horde/HordeDashboard/src/index.css
--- a/Engine/Source/Programs/Horde/HordeDashboard/src/index.css
+++ b/Engine/Source/Programs/Horde/HordeDashboard/src/index.css
@@ -1,8 +1,6 @@
 body {
   margin: 0;
-  font-family: -apple-system, BlinkMacSystemFont, 'Segoe UI', 'Roboto', 'Oxygen',
-    'Ubuntu', 'Cantarell', 'Fira Sans', 'Droid Sans', 'Helvetica Neue',
-    sans-serif;
+  font-family: "Horde Open Sans Regular";
   -webkit-font-smoothing: antialiased;
   -moz-osx-font-smoothing: grayscale; 
   overflow: hidden;
@@ -28,16 +26,19 @@ code {
 }
 
 .log-link, .log-link:visited, .log-link:hover {
+  font-family: "Horde Open Sans Regular";
   text-decoration: none;
   color:#000000;
 }
 
 .horde-link, .horde-link:visited {
+  font-family: "Horde Open Sans Regular";
   text-decoration: none;
   color: rgb(0, 120, 212);
 }
 
 .horde-link:hover {
+  font-family: "Horde Open Sans Regular";
   text-decoration: underline;
   color: rgb(0, 120, 212)
 }
@@ -52,12 +53,14 @@ code {
 
 
 .view-log-link, .view-log-link:visited, .view-log-link:hover {
+  font-family: "Horde Open Sans Regular";
   text-decoration: none;
   user-select: none;
   cursor: pointer;
 }
 
 .view-log-link-disabled, .view-log-link-disabled:visited, .view-log-link-disabled:hover {
+  font-family: "Horde Open Sans Regular";
   text-decoration: none;
   color: #8f8f8f;
   user-select: none;
