Horde: Connect directly to agents when running remote execution leases, rather than having agents connect back to the initiator.

- Agents now accept incoming connections on a defined port, and match them up with nonces specified through leases assigned from the server.
- Requests for new compute resources now attempt to allocate a lease and return immediately with nonce and AES keys.

diff --git a/Engine/Source/Programs/Horde/Horde.Agent/AgentSettings.cs b/Engine/Source/Programs/Horde/Horde.Agent/AgentSettings.cs
--- a/Engine/Source/Programs/Horde/Horde.Agent/AgentSettings.cs
+++ b/Engine/Source/Programs/Horde/Horde.Agent/AgentSettings.cs
@@ -191,6 +191,11 @@ namespace Horde.Agent
 		/// </summary>
 		public bool UseLocalStorageClient { get; set; }
 
+		/// <summary>
+		/// Incoming port for listening for compute work. Needs to be tied with a lease.
+		/// </summary>
+		public int ComputePort { get; set; } = 7000;
+
 		/// <summary>
 		/// Key/value properties in addition to those set internally by the agent
 		/// </summary>
diff --git a/Engine/Source/Programs/Horde/Horde.Agent/Leases/Handlers/ComputeHandler.cs b/Engine/Source/Programs/Horde/Horde.Agent/Leases/Handlers/ComputeHandler.cs
--- a/Engine/Source/Programs/Horde/Horde.Agent/Leases/Handlers/ComputeHandler.cs
+++ b/Engine/Source/Programs/Horde/Horde.Agent/Leases/Handlers/ComputeHandler.cs
@@ -26,14 +26,16 @@ namespace Horde.Agent.Leases.Handlers
 	/// </summary>
 	class ComputeHandler : LeaseHandler<ComputeTask>
 	{
+		readonly ComputeListenerService _listenerService;
 		readonly IMemoryCache _memoryCache;
 		readonly ILogger _logger;
 
 		/// <summary>
 		/// Constructor
 		/// </summary>
-		public ComputeHandler(IMemoryCache memoryCache, ILogger<ComputeHandler> logger)
+		public ComputeHandler(ComputeListenerService listenerService, IMemoryCache memoryCache, ILogger<ComputeHandler> logger)
 		{
+			_listenerService = listenerService;
 			_memoryCache = memoryCache;
 			_logger = logger;
 		}
@@ -41,30 +43,35 @@ namespace Horde.Agent.Leases.Handlers
 		/// <inheritdoc/>
 		public override async Task<LeaseResult> ExecuteAsync(ISession session, string leaseId, ComputeTask computeTask, CancellationToken cancellationToken)
 		{
-			_logger.LogInformation("Starting compute task (lease {LeaseId})", leaseId);
+			_logger.LogInformation("Starting compute task (lease {LeaseId}). Waiting for connection...", leaseId);
 
+			TcpClient? tcpClient = null;
 			try
 			{
-				await ConnectAsync(computeTask, cancellationToken);
+				const int TimeoutSeconds = 30;
+
+				tcpClient = await _listenerService.WaitForClientAsync(new ByteString(computeTask.Nonce.Memory), TimeSpan.FromSeconds(TimeoutSeconds), cancellationToken);
+				if (tcpClient == null)
+				{
+					_logger.LogInformation("Timed out waiting for connection after {Time}s.", TimeoutSeconds); 
+					return LeaseResult.Success;
+				}
+
+				using (SocketComputeChannel channel = new SocketComputeChannel(tcpClient.Client, computeTask.AesKey.Memory, computeTask.AesIv.Memory))
+				{
+					await RunAsync(channel, cancellationToken);
+					return LeaseResult.Success;
+				}
 			}
 			catch (Exception ex)
 			{
-				_logger.LogError(ex, "Exception while executing compute task");
+				_logger.LogError(ex, "Exception while executing compute task: {Message}", ex.Message);
+				return LeaseResult.Failed;
+			}
+			finally
+			{
+				tcpClient?.Dispose();
 			}
-
-			return LeaseResult.Success;
-		}
-
-		async Task ConnectAsync(ComputeTask computeTask, CancellationToken cancellationToken)
-		{
-			using TcpClient tcpClient = new TcpClient();
-			await tcpClient.ConnectAsync(computeTask.RemoteIp, computeTask.RemotePort);
-
-			Socket socket = tcpClient.Client;
-			await socket.SendAsync(computeTask.Nonce.Memory, SocketFlags.None, cancellationToken);
-
-			using SocketComputeChannel channel = new SocketComputeChannel(socket, computeTask.AesKey.Memory, computeTask.AesIv.Memory);
-			await RunAsync(channel, cancellationToken);
 		}
 
 		public async Task RunAsync(IComputeChannel channel, CancellationToken cancellationToken)
diff --git a/Engine/Source/Programs/Horde/Horde.Agent/Program.cs b/Engine/Source/Programs/Horde/Horde.Agent/Program.cs
--- a/Engine/Source/Programs/Horde/Horde.Agent/Program.cs
+++ b/Engine/Source/Programs/Horde/Horde.Agent/Program.cs
@@ -4,9 +4,12 @@ using System;
 using System.Collections.Generic;
 using System.Diagnostics;
 using System.IO;
+using System.Net;
 using System.Net.Http;
+using System.Net.Sockets;
 using System.Reflection;
 using System.Runtime.InteropServices;
+using System.Threading;
 using System.Threading.Tasks;
 using Datadog.Trace;
 using Datadog.Trace.Configuration;
@@ -203,6 +206,9 @@ namespace Horde.Agent
 
 			services.AddSingleton<IStorageClientFactory, StorageClientFactory>();
 
+			services.AddSingleton<ComputeListenerService>();
+			services.AddHostedService(sp => sp.GetRequiredService<ComputeListenerService>());
+
 			services.AddMemoryCache();
 
 			// Allow commands to augment the service collection for their own DI service providers
diff --git a/Engine/Source/Programs/Horde/Horde.Agent/Services/CapabilitiesService.cs b/Engine/Source/Programs/Horde/Horde.Agent/Services/CapabilitiesService.cs
--- a/Engine/Source/Programs/Horde/Horde.Agent/Services/CapabilitiesService.cs
+++ b/Engine/Source/Programs/Horde/Horde.Agent/Services/CapabilitiesService.cs
@@ -338,6 +338,7 @@ namespace Horde.Agent.Services
 			}
 
 			// Get the IP addresses
+			IPAddress? ip = null;
 			try
 			{
 				using CancellationTokenSource dnsCts = new(3000);
@@ -347,6 +348,7 @@ namespace Horde.Agent.Services
 					if (address.AddressFamily == System.Net.Sockets.AddressFamily.InterNetwork)
 					{
 						primaryDevice.Properties.Add($"Ipv4={address}");
+						ip = address;
 					}
 					else if (address.AddressFamily == System.Net.Sockets.AddressFamily.InterNetworkV6)
 					{
@@ -359,6 +361,13 @@ namespace Horde.Agent.Services
 				logger.LogDebug(ex, "Unable to get local IP address");
 			}
 
+			// Add the compute configuration
+			if (ip != null && _settings.ComputePort != 0)
+			{
+				primaryDevice.Properties.Add($"ComputeIp={ip}");
+				primaryDevice.Properties.Add($"ComputePort={_settings.ComputePort}");
+			}
+
 			// Get the time that the machine booted
 			primaryDevice.Properties.Add($"BootTime={s_bootTime}");
 			primaryDevice.Properties.Add($"StartTime={s_startTime}");
diff --git a/Engine/Source/Programs/Horde/Horde.Agent/Services/ComputeListenerService.cs b/Engine/Source/Programs/Horde/Horde.Agent/Services/ComputeListenerService.cs
new file mode 100644
--- /dev/null
+++ b/Engine/Source/Programs/Horde/Horde.Agent/Services/ComputeListenerService.cs
@@ -0,0 +1,226 @@
+// Copyright Epic Games, Inc. All Rights Reserved.
+
+using EpicGames.Core;
+using System;
+using System.Collections.Generic;
+using System.Net;
+using System.Net.Sockets;
+using System.Threading;
+using System.Threading.Tasks;
+using EpicGames.Horde.Compute;
+using Microsoft.Extensions.Hosting;
+using Microsoft.Extensions.Logging;
+using Microsoft.Extensions.Options;
+
+namespace Horde.Agent.Services
+{
+	/// <summary>
+	/// Service which listens for incoming connections that can be assigned for remote execution leases
+	/// </summary>
+	class ComputeListenerService : IHostedService, IDisposable
+	{
+		class WaitingClient : IDisposable
+		{
+			TcpClient? _tcpClient;
+			readonly TaskCompletionSource _takenTaskSource = new TaskCompletionSource();
+
+			public WaitingClient(TcpClient tcpClient)
+			{
+				_tcpClient = tcpClient;
+			}
+
+			public TcpClient? TakeClient()
+			{
+				TcpClient? tcpClient = _tcpClient;
+				if (Interlocked.CompareExchange(ref _tcpClient, null, tcpClient) == tcpClient)
+				{
+					Task.Run(() => _takenTaskSource.SetResult()); // Run continuations on another thread since we may be in a lock
+					return tcpClient;
+				}
+				return null;
+			}
+
+			public Task WaitForClientToBeTaken() => _takenTaskSource.Task;
+
+			public void Dispose()
+			{
+				_tcpClient?.Dispose();
+			}
+		}
+
+		readonly AgentSettings _settings;
+		readonly ILogger _logger;
+
+		Task _serverTask = Task.CompletedTask;
+		readonly CancellationTokenSource _cancellationSource = new CancellationTokenSource();
+
+		readonly object _lockObject = new object();
+		readonly Dictionary<ByteString, WaitingClient> _waitingClients = new Dictionary<ByteString, WaitingClient>();
+		readonly Dictionary<ByteString, TaskCompletionSource<TcpClient?>> _waitingLeases = new Dictionary<ByteString, TaskCompletionSource<TcpClient?>>();
+
+		/// <summary>
+		/// Constructor
+		/// </summary>
+		public ComputeListenerService(IOptions<AgentSettings> settings, ILogger<ComputeListenerService> logger)
+		{
+			_settings = settings.Value;
+			_logger = logger;
+		}
+
+		/// <inheritdoc/>
+		public void Dispose()
+		{
+			_cancellationSource.Dispose();
+		}
+
+		/// <summary>
+		/// Waits for a client to connect with the given connection id
+		/// </summary>
+		/// <param name="nonce">Unique id for the connection</param>
+		/// <param name="timeout">Timeout for the wait</param>
+		/// <param name="cancellationToken">Cancellation token</param>
+		/// <returns>The connected client</returns>
+		public async Task<TcpClient?> WaitForClientAsync(ByteString nonce, TimeSpan timeout, CancellationToken cancellationToken)
+		{
+			// Check if there's already a client ready, and register that we're waiting if not
+			TaskCompletionSource<TcpClient?> waitingLease;
+			lock (_lockObject)
+			{
+				if (_waitingClients.TryGetValue(nonce, out WaitingClient? waitingClientInfo))
+				{
+					return waitingClientInfo.TakeClient();
+				}
+
+				waitingLease = new TaskCompletionSource<TcpClient?>();
+				_waitingLeases.Add(nonce, waitingLease);
+			}
+
+			// Wait until we get a connection or we reach the timeout
+			try
+			{
+				Task waitTask = Task.Delay(timeout, cancellationToken);
+				await Task.WhenAny(waitTask, waitingLease!.Task);
+			}
+			finally
+			{
+				lock (_lockObject)
+				{
+					_waitingLeases.Remove(nonce);
+				}
+			}
+
+			// Try to return the client
+			waitingLease.TrySetResult(null);
+			return await waitingLease.Task;
+		}
+
+		/// <inheritdoc/>
+		public Task StartAsync(CancellationToken cancellationToken)
+		{
+			if (_settings.ComputePort != 0)
+			{
+				Start(IPAddress.Any, _settings.ComputePort);
+			}
+			return Task.CompletedTask;
+		}
+
+		/// <summary>
+		/// Starts listening for connections on the given address and port. Primarily for testing; the hosted service will use settings for this from the configuration file.
+		/// </summary>
+		/// <param name="address">Address to listen on</param>
+		/// <param name="port">Port to listen on</param>
+		public void Start(IPAddress address, int port)
+		{
+			_serverTask = RunTcpListenerAsync(address, port, _cancellationSource.Token);
+		}
+
+		/// <inheritdoc/>
+		public async Task StopAsync(CancellationToken cancellationToken)
+		{
+			_cancellationSource.Cancel();
+			await _serverTask;
+		}
+
+		async Task RunTcpListenerAsync(IPAddress address, int port, CancellationToken cancellationToken)
+		{
+			TcpListener listener = new TcpListener(address, port);
+			listener.Start();
+
+			List<Task> tasks = new List<Task>();
+			try
+			{
+				for (; ; )
+				{
+					TcpClient tcpClient = await listener.AcceptTcpClientAsync(cancellationToken);
+					_logger.LogInformation("Received connection from {Remote}", tcpClient.Client.RemoteEndPoint);
+
+					Task task = WaitForLeaseAsync(tcpClient, cancellationToken);
+					tasks.Add(task);
+				}
+			}
+			finally
+			{
+				await Task.WhenAll(tasks);
+				listener.Stop();
+			}
+		}
+
+		async Task WaitForLeaseAsync(TcpClient tcpClient, CancellationToken cancellationToken)
+		{
+			WaitingClient? waitingClient = null;
+			try
+			{
+				Socket socket = tcpClient.Client;
+
+				// Read the nonce for the connection
+				byte[] nonceBuffer = new byte[ServerComputeClient.NonceLength];
+				await socket.ReceiveFullAsync(nonceBuffer, SocketFlags.None, cancellationToken);
+
+				// Register the socket
+				ByteString nonce = new ByteString(nonceBuffer);
+				try
+				{
+					// Check if there's already a waiting connection, and add one if not
+					TaskCompletionSource<TcpClient?>? waitingLease = null;
+					lock (_lockObject)
+					{
+						if (!_waitingLeases.TryGetValue(nonce, out waitingLease))
+						{
+							waitingClient = new WaitingClient(tcpClient);
+							_waitingClients.Add(nonce, waitingClient);
+						}
+					}
+
+					// Offer ownership of the socket to any interested compute lease, or dispose of the connection if it's already closed.
+					if (waitingLease != null)
+					{
+						if (!waitingLease.TrySetResult(tcpClient))
+						{
+							tcpClient.Dispose();
+						}
+						return;
+					}
+
+					// Wait 30s for it to be claimed
+					Task delayTask = Task.Delay(TimeSpan.FromSeconds(30.0), cancellationToken);
+					await Task.WhenAny(waitingClient!.WaitForClientToBeTaken(), delayTask);
+				}
+				finally
+				{
+					lock (_lockObject)
+					{
+						_waitingClients.Remove(nonce);
+					}
+				}
+			}
+			catch (Exception ex)
+			{
+				_logger.LogError(ex, "Exception while processing compute task: {Message}", ex.Message);
+			}
+			finally
+			{
+				waitingClient?.Dispose();
+			}
+		}
+	}
+}
diff --git a/Engine/Source/Programs/Horde/Horde.Build/Compute/ComputeController.cs b/Engine/Source/Programs/Horde/Horde.Build/Compute/ComputeController.cs
--- a/Engine/Source/Programs/Horde/Horde.Build/Compute/ComputeController.cs
+++ b/Engine/Source/Programs/Horde/Horde.Build/Compute/ComputeController.cs
@@ -1,93 +1,58 @@
 // Copyright Epic Games, Inc. All Rights Reserved.
 
 using System;
-using System.Collections.Generic;
 using System.Net;
 using EpicGames.Core;
 using EpicGames.Horde.Compute;
 using Horde.Build.Acls;
 using Horde.Build.Server;
 using Horde.Build.Utilities;
-using HordeCommon.Rpc.Tasks;
 using Microsoft.AspNetCore.Authorization;
 using Microsoft.AspNetCore.Mvc;
 using Microsoft.Extensions.Options;
 
 namespace Horde.Build.Compute
 {
-	using ByteString = Google.Protobuf.ByteString;
-
 	/// <summary>
 	/// Request a machine to execute compute requests
 	/// </summary>
-	public class AddComputeTaskRequest
+	public class AssignComputeRequest
 	{
 		/// <summary>
 		/// Condition to identify machines that can execute the request
 		/// </summary>
 		public Requirements? Requirements { get; set; }
-
-		/// <summary>
-		/// Port to connect on
-		/// </summary>
-		public int? RemotePort { get; set; }
-
-		/// <summary>
-		/// Cryptographic nonce to identify the request, as a hex string
-		/// </summary>
-		public string Nonce { get; set; } = String.Empty;
-
-		/// <summary>
-		/// AES key for the channel, as a hex string
-		/// </summary>
-		public string AesKey { get; set; } = String.Empty;
-
-		/// <summary>
-		/// AES IV for the channel, as a hex string
-		/// </summary>
-		public string AesIv { get; set; } = String.Empty;
 	}
 
 	/// <summary>
 	/// Request a machine to execute compute requests
 	/// </summary>
-	public class AddComputeTasksRequest
+	public class AssignComputeResponse
 	{
 		/// <summary>
-		/// Condition to identify machines that can execute the request
+		/// IP address of the remote machine
 		/// </summary>
-		public Requirements? Requirements { get; set; }
-
-		/// <summary>
-		/// Port to connect on
-		/// </summary>
-		public int RemotePort { get; set; }
+		public string Ip { get; set; } = String.Empty;
 
 		/// <summary>
-		/// List of tasks to add
+		/// Port number on the remote machine
 		/// </summary>
-		public List<AddComputeTaskRequest> Tasks { get; set; } = new List<AddComputeTaskRequest>();
-	}
+		public int Port { get; set; }
 
-	/// <summary>
-	/// Gets information about the configured compute tunnel
-	/// </summary>
-	public class GetComputeTunnelResponse
-	{
 		/// <summary>
-		/// IP address of the server
+		/// Cryptographic nonce to identify the request, as a hex string
 		/// </summary>
-		public string? Ip { get; set; }
+		public string Nonce { get; set; } = String.Empty;
 
 		/// <summary>
-		/// Listening port for the initiator
+		/// AES key for the channel, as a hex string
 		/// </summary>
-		public int InitiatorPort { get; set; }
+		public string AesKey { get; set; } = String.Empty;
 
 		/// <summary>
-		/// Listening port for the remote
+		/// AES IV for the channel, as a hex string
 		/// </summary>
-		public int RemotePort { get; set; }
+		public string AesIv { get; set; } = String.Empty;
 	}
 
 	/// <summary>
@@ -98,15 +63,15 @@ namespace Horde.Build.Compute
 	[Route("[controller]")]
 	public class ComputeControllerV2 : HordeControllerBase
 	{
-		readonly ComputeService _computeService;
+		readonly ComputeTaskSource _computeTaskSource;
 		readonly IOptionsSnapshot<GlobalConfig> _globalConfig;
 
 		/// <summary>
 		/// Constructor
 		/// </summary>
-		public ComputeControllerV2(ComputeService computeService, IOptionsSnapshot<GlobalConfig> globalConfig)
+		public ComputeControllerV2(ComputeTaskSource computeTaskSource, IOptionsSnapshot<GlobalConfig> globalConfig)
 		{
-			_computeService = computeService;
+			_computeTaskSource = computeTaskSource;
 			_globalConfig = globalConfig;
 		}
 
@@ -119,7 +84,7 @@ namespace Horde.Build.Compute
 		[HttpPost]
 		[Authorize]
 		[Route("/api/v2/compute/{clusterId}")]
-		public ActionResult AddTasksAsync(ClusterId clusterId, [FromBody] AddComputeTasksRequest request)
+		public ActionResult<AssignComputeResponse> AssignComputeResourceAsync(ClusterId clusterId, [FromBody] AssignComputeRequest request)
 		{
 			ComputeClusterConfig? clusterConfig;
 			if (!_globalConfig.Value.TryGetComputeCluster(clusterId, out clusterConfig))
@@ -131,53 +96,21 @@ namespace Horde.Build.Compute
 				return Forbid(AclAction.AddComputeTasks, clusterId);
 			}
 
-			IPAddress? remoteIp = HttpContext.Connection.RemoteIpAddress;
-			if (remoteIp == null)
-			{
-				return BadRequest("Missing remote IP address");
-			}
-
-			foreach (AddComputeTaskRequest taskRequest in request.Tasks)
-			{
-				Requirements requirements = taskRequest.Requirements ?? request.Requirements ?? new Requirements();
-
-				ComputeTask computeTask = new ComputeTask();
-				computeTask.RemoteIp = remoteIp.ToString();
-				computeTask.RemotePort = taskRequest.RemotePort ?? request.RemotePort;
-				computeTask.Nonce = ByteString.CopyFrom(StringUtils.ParseHexString(taskRequest.Nonce));
-				computeTask.AesKey = ByteString.CopyFrom(StringUtils.ParseHexString(taskRequest.AesKey));
-				computeTask.AesIv = ByteString.CopyFrom(StringUtils.ParseHexString(taskRequest.AesIv));
+			Requirements requirements = request.Requirements ?? new Requirements();
 
-				_computeService.AddRequest(clusterId, requirements, computeTask);
-			}
-			return Ok();
-		}
-
-		/// <summary>
-		/// Add tasks to be executed remotely
-		/// </summary>
-		/// <param name="clusterId">Id of the compute cluster</param>
-		/// <returns></returns>
-		[HttpGet]
-		[Authorize]
-		[Route("/api/v2/compute/{clusterId}/tunnel")]
-		public ActionResult<GetComputeTunnelResponse> GetTunnelInfoAsync(ClusterId clusterId)
-		{
-			if (!_globalConfig.Value.TryGetComputeCluster(clusterId, out _))
+			ComputeResource? computeResource = _computeTaskSource.TryAllocateResource(clusterId, requirements);
+			if (computeResource == null)
 			{
-				return NotFound(clusterId);
+				return StatusCode((int)HttpStatusCode.ServiceUnavailable);
 			}
 
-			ServerSettings settings = _globalConfig.Value.ServerSettings;
-			if (settings.ComputeInitiatorPort == 0 || settings.ComputeRemotePort == 0)
-			{
-				return NotFound("Tunnelling is not configured on the server");
-			}
+			AssignComputeResponse response = new AssignComputeResponse();
+			response.Ip = computeResource.Ip.ToString();
+			response.Port = computeResource.Port;
+			response.Nonce = StringUtils.FormatHexString(computeResource.Task.Nonce.Span);
+			response.AesKey = StringUtils.FormatHexString(computeResource.Task.AesKey.Span);
+			response.AesIv = StringUtils.FormatHexString(computeResource.Task.AesIv.Span);
 
-			GetComputeTunnelResponse response = new GetComputeTunnelResponse();
-			response.Ip = HttpContext.Connection.LocalIpAddress?.ToString();
-			response.InitiatorPort = _globalConfig.Value.ServerSettings.ComputeInitiatorPort;
-			response.RemotePort = _globalConfig.Value.ServerSettings.ComputeRemotePort;
 			return response;
 		}
 	}
diff --git a/Engine/Source/Programs/Horde/Horde.Build/Compute/ComputeTaskSource.cs b/Engine/Source/Programs/Horde/Horde.Build/Compute/ComputeTaskSource.cs
--- a/Engine/Source/Programs/Horde/Horde.Build/Compute/ComputeTaskSource.cs
+++ b/Engine/Source/Programs/Horde/Horde.Build/Compute/ComputeTaskSource.cs
@@ -1,6 +1,10 @@
 // Copyright Epic Games, Inc. All Rights Reserved.
 
 using System;
+using System.Collections.Generic;
+using System.Linq;
+using System.Net;
+using System.Security.Cryptography;
 using System.Threading;
 using System.Threading.Tasks;
 using EpicGames.Core;
@@ -9,55 +13,207 @@ using Google.Protobuf;
 using Google.Protobuf.WellKnownTypes;
 using Horde.Build.Agents;
 using Horde.Build.Agents.Leases;
+using Horde.Build.Server;
 using Horde.Build.Tasks;
 using HordeCommon;
 using HordeCommon.Rpc.Tasks;
 using Microsoft.Extensions.Logging;
+using Microsoft.Extensions.Options;
 
 namespace Horde.Build.Compute
 {
+	/// <summary>
+	/// Information about a compute 
+	/// </summary>
+	public class ComputeResource
+	{
+		/// <summary>
+		/// IP address of the agent
+		/// </summary>
+		public IPAddress Ip { get; }
+
+		/// <summary>
+		/// Port to connect on
+		/// </summary>
+		public int Port { get; }
+
+		/// <summary>
+		/// Information about the compute task
+		/// </summary>
+		public ComputeTask Task { get; }
+
+		/// <summary>
+		/// Constructor
+		/// </summary>
+		public ComputeResource(IPAddress ip, int port, ComputeTask task)
+		{
+			Ip = ip;
+			Port = port;
+			Task = task;
+		}
+	}
+
 	/// <summary>
 	/// Dispatches requests for compute resources
 	/// </summary>
 	public class ComputeTaskSource : TaskSourceBase<ComputeTask>
 	{
+		class Waiter
+		{
+			public IAgent Agent { get; }
+			public IPAddress Ip { get; }
+			public int Port { get; }
+			public TaskCompletionSource<(ComputeTask, Requirements)?> Lease { get; } = new TaskCompletionSource<(ComputeTask, Requirements)?>(TaskCreationOptions.RunContinuationsAsynchronously);
+
+			public Waiter(IAgent agent, IPAddress ip, int port)
+			{
+				Agent = agent;
+				Ip = ip;
+				Port = port;
+			}
+		}
+
 		/// <inheritdoc/>
 		public override string Type => "Compute";
 
 		/// <inheritdoc/>
 		public override TaskSourceFlags Flags => TaskSourceFlags.None;
 
-		readonly ComputeService _computeService;
+		readonly IOptionsMonitor<GlobalConfig> _globalConfig;
 		readonly ILogger _logger;
 
+		readonly object _lockObject = new object();
+		readonly Dictionary<ClusterId, LinkedList<Waiter>> _waiters = new Dictionary<ClusterId, LinkedList<Waiter>>();
+
 		/// <summary>
 		/// Constructor
 		/// </summary>
-		public ComputeTaskSource(ComputeService computeService, ILogger<ComputeTaskSource> logger)
+		public ComputeTaskSource(IOptionsMonitor<GlobalConfig> globalConfig, ILogger<ComputeTaskSource> logger)
 		{
-			_computeService = computeService;
+			_globalConfig = globalConfig;
 			_logger = logger;
 		}
 
+		/// <summary>
+		/// Allocates a compute resource
+		/// </summary>
+		/// <returns></returns>
+		public ComputeResource? TryAllocateResource(ClusterId clusterId, Requirements requirements)
+		{
+			lock (_lockObject)
+			{
+				LinkedList<Waiter>? waiters;
+				if (_waiters.TryGetValue(clusterId, out waiters))
+				{
+					ComputeTask? task = null;
+					for (LinkedListNode<Waiter>? node = waiters.First; node != null; node = node.Next)
+					{
+						if (node.Value.Agent.MeetsRequirements(requirements))
+						{
+							task ??= CreateComputeTask();
+							if (node.Value.Lease.TrySetResult((task, requirements)))
+							{
+								Waiter waiter = node.Value;
+								return new ComputeResource(waiter.Ip, waiter.Port, task);
+							}
+						}
+					}
+				}
+			}
+			return null;
+		}
+
+		static ComputeTask CreateComputeTask()
+		{
+			ComputeTask computeTask = new ComputeTask();
+			computeTask.Nonce = UnsafeByteOperations.UnsafeWrap(RandomNumberGenerator.GetBytes(ServerComputeClient.NonceLength));
+
+			using (Aes aes = Aes.Create())
+			{
+				computeTask.AesKey = UnsafeByteOperations.UnsafeWrap(aes.Key);
+				computeTask.AesIv = UnsafeByteOperations.UnsafeWrap(aes.IV);
+			}
+
+			return computeTask;
+		}
+
 		/// <inheritdoc/>
-		public override async Task<Task<AgentLease?>> AssignLeaseAsync(IAgent agent, CancellationToken cancellationToken)
+		public override Task<Task<AgentLease?>> AssignLeaseAsync(IAgent agent, CancellationToken cancellationToken)
 		{
-			while (!cancellationToken.IsCancellationRequested)
+			return Task.FromResult(WaitInternalAsync(agent, cancellationToken));
+		}
+
+		async Task<AgentLease?> WaitInternalAsync(IAgent agent, CancellationToken cancellationToken)
+		{
+			string? ipStr = agent.GetPropertyValues("ComputeIp").FirstOrDefault();
+			if (ipStr == null || !IPAddress.TryParse(ipStr, out IPAddress? ip))
+			{
+				return null;
+			}
+
+			string? portStr = agent.GetPropertyValues("ComputePort").FirstOrDefault();
+			if (portStr == null || !Int32.TryParse(portStr, out int port))
+			{
+				return null;
+			}
+
+			// Add it to the wait queue
+			List<(LinkedList<Waiter>, LinkedListNode<Waiter>)> nodes = new();
+			try
 			{
-				ComputeRequest? request = _computeService.PopRequest(agent);
-				if (request != null)
+				GlobalConfig globalConfig = _globalConfig.CurrentValue;
+
+				Waiter? waiter = null;
+				lock (_lockObject)
 				{
-					string leaseName = $"Compute for {request.Task.RemoteIp}";
-					byte[] payload = Any.Pack(request.Task).ToByteArray();
+					foreach (ComputeClusterConfig clusterConfig in globalConfig.Compute)
+					{
+						if (clusterConfig.Condition == null || agent.SatisfiesCondition(clusterConfig.Condition))
+						{
+							LinkedList<Waiter>? list;
+							if (!_waiters.TryGetValue(clusterConfig.Id, out list))
+							{
+								list = new LinkedList<Waiter>();
+								_waiters.Add(clusterConfig.Id, list);
+							}
 
-					AgentLease lease = new AgentLease(LeaseId.GenerateNewId(), leaseName, null, null, null, LeaseState.Pending, request.Requirements.Resources, request.Requirements.Exclusive, payload);
-					_logger.LogInformation("Created compute lease for agent {AgentId} and remote {RemoteIp}:{RemotePort}", agent.Id, request.Task.RemoteIp, request.Task.RemotePort);
-					return Task.FromResult<AgentLease?>(lease);
+							waiter ??= new Waiter(agent, ip, port);
+							list.AddFirst(waiter);
+						}
+					}
 				}
 
-				await AsyncUtils.DelayNoThrow(TimeSpan.FromSeconds(5.0), cancellationToken);
+				if (waiter != null)
+				{
+					using (IDisposable disposable = cancellationToken.Register(() => waiter.Lease.TrySetResult(null)))
+					{
+						(ComputeTask, Requirements)? result = await waiter.Lease.Task;
+						if (result != null)
+						{
+							(ComputeTask task, Requirements requirements) = result.Value;
+
+							string leaseName = "Compute task";
+							byte[] payload = Any.Pack(task).ToByteArray();
+
+							AgentLease lease = new AgentLease(LeaseId.GenerateNewId(), leaseName, null, null, null, LeaseState.Pending, requirements.Resources, requirements.Exclusive, payload);
+							_logger.LogInformation("Created compute lease for agent {AgentId}", agent.Id);
+							return lease;
+						}
+					}
+				}
 			}
-			return Task.FromResult<AgentLease?>(null);
+			finally
+			{
+				lock (_lockObject)
+				{
+					foreach ((LinkedList<Waiter> list, LinkedListNode<Waiter> node) in nodes)
+					{
+						list.Remove(node);
+					}
+				}
+			}
+
+			return null;
 		}
 	}
 }
diff --git a/Engine/Source/Programs/Horde/Horde.Build/Startup.cs b/Engine/Source/Programs/Horde/Horde.Build/Startup.cs
--- a/Engine/Source/Programs/Horde/Horde.Build/Startup.cs
+++ b/Engine/Source/Programs/Horde/Horde.Build/Startup.cs
@@ -696,13 +696,14 @@ namespace Horde.Build
 				services.AddSingleton<ConformTaskSource>();
 				services.AddHostedService<ConformTaskSource>(provider => provider.GetRequiredService<ConformTaskSource>());
 				services.AddSingleton<ComputeService>();
+				services.AddSingleton<ComputeTaskSource>();
 
 				services.AddSingleton<ITaskSource, UpgradeTaskSource>();
 				services.AddSingleton<ITaskSource, ShutdownTaskSource>();
 				services.AddSingleton<ITaskSource, RestartTaskSource>();
 				services.AddSingleton<ITaskSource, ConformTaskSource>(provider => provider.GetRequiredService<ConformTaskSource>());
 				services.AddSingleton<ITaskSource, JobTaskSource>(provider => provider.GetRequiredService<JobTaskSource>());
-				services.AddSingleton<ITaskSource, ComputeTaskSource>();
+				services.AddSingleton<ITaskSource, ComputeTaskSource>(provider => provider.GetRequiredService<ComputeTaskSource>());
 			}
 
 			// Allow longer to shutdown so we can debug missing cancellation tokens
diff --git a/Engine/Source/Programs/Shared/EpicGames.Core/SocketExtensions.cs b/Engine/Source/Programs/Shared/EpicGames.Core/SocketExtensions.cs
new file mode 100644
--- /dev/null
+++ b/Engine/Source/Programs/Shared/EpicGames.Core/SocketExtensions.cs
@@ -0,0 +1,52 @@
+// Copyright Epic Games, Inc. All Rights Reserved.
+
+using System;
+using System.IO;
+using System.Net.Sockets;
+using System.Threading;
+using System.Threading.Tasks;
+
+namespace EpicGames.Core
+{
+	/// <summary>
+	/// Extension methods for sockets
+	/// </summary>
+	public static class SocketExtensions
+	{
+		/// <summary>
+		/// Reads a complete buffer from the given socket, retrying reads until the buffer is full.
+		/// </summary>
+		/// <param name="socket">Socket to read from</param>
+		/// <param name="buffer">Buffer to store the data</param>
+		/// <param name="socketFlags">Flags for the <see cref="Socket.ReceiveAsync"/> call</param>
+		/// <param name="cancellationToken">Cancellation token for the operation</param>
+		public static async Task ReceiveFullAsync(this Socket socket, Memory<byte> buffer, SocketFlags socketFlags, CancellationToken cancellationToken)
+		{
+			for (int offset = 0; offset < buffer.Length;)
+			{
+				int read = await socket.ReceiveAsync(buffer.Slice(offset), socketFlags, cancellationToken);
+				if (read == 0)
+				{
+					throw new EndOfStreamException();
+				}
+				offset += read;
+			}
+		}
+
+		/// <summary>
+		/// Sends a complete buffer over the given socket.
+		/// </summary>
+		/// <param name="socket">Socket to write to</param>
+		/// <param name="buffer">Buffer to write</param>
+		/// <param name="cancellationToken">Flags for the <see cref="Socket.SendAsync"/> call</param>
+		/// <returns></returns>
+		public static async Task SendFullAsync(this Socket socket, ReadOnlyMemory<byte> buffer, SocketFlags socketFlags, CancellationToken cancellationToken)
+		{
+			while (buffer.Length > 0)
+			{
+				int written = await socket.SendAsync(buffer, socketFlags, cancellationToken);
+				buffer = buffer.Slice(written);
+			}
+		}
+	}
+}
