Horde: Fix issue with code styling font selection
#rnx
#skipci

diff --git a/Engine/Source/Programs/Horde/HordeDashboard/src/index.css b/Engine/Source/Programs/Horde/HordeDashboard/src/index.css
--- a/Engine/Source/Programs/Horde/HordeDashboard/src/index.css
+++ b/Engine/Source/Programs/Horde/HordeDashboard/src/index.css
@@ -25,8 +25,7 @@ code {
     monospace;
 }
 
-.log-link, .log-link:visited, .log-link:hover {
-  font-family: "Horde Open Sans Regular";
+.log-link, .log-link:visited, .log-link:hover {  
   text-decoration: none;
   color:#000000;
 }
@@ -53,14 +52,12 @@ code {
 
 
 .view-log-link, .view-log-link:visited, .view-log-link:hover {
-  font-family: "Horde Open Sans Regular";
   text-decoration: none;
   user-select: none;
   cursor: pointer;
 }
 
 .view-log-link-disabled, .view-log-link-disabled:visited, .view-log-link-disabled:hover {
-  font-family: "Horde Open Sans Regular";
   text-decoration: none;
   color: #8f8f8f;
   user-select: none;
