Horde: Include timestamp in selection
#rnx
#skipci

diff --git a/Engine/Source/Programs/Horde/HordeDashboard/src/components/LogView.tsx b/Engine/Source/Programs/Horde/HordeDashboard/src/components/LogView.tsx
--- a/Engine/Source/Programs/Horde/HordeDashboard/src/components/LogView.tsx
+++ b/Engine/Source/Programs/Horde/HordeDashboard/src/components/LogView.tsx
@@ -656,7 +656,7 @@ export const LogList: React.FC<{ logId: string }> = observer(({ logId }) => {
                      <Stack styles={{ root: { color: "#c0c0c0", width: 80, textAlign: "right", userSelect: "none", fontSize: handler.fontSize } }}>{prefix + item.lineNumber}</Stack>
                      <Stack className={style} horizontal disableShrink={true} >
                         <Stack className={gutterStyle}></Stack>
-                        {(!item.issueId || !ev) && <Stack styles={{ root: { color: "#8a8a8a", width: tsWidth, whiteSpace: "nowrap", fontSize: handler.fontSize, userSelect: "none" } }}> {timestamp}</Stack>}
+                        {(!item.issueId || !ev) && <Stack styles={{ root: { color: "#8a8a8a", width: tsWidth, whiteSpace: "nowrap", fontSize: handler.fontSize, userSelect: "text" } }}> {timestamp}</Stack>}
                         {!!item.issueId && !!ev && <IssueButton item={item} event={ev!} />}
                         <div className={styles.logLineOuter}> <Stack styles={{ root: { paddingLeft: 8, paddingRight: 8 } }}> {renderLine(item.line, item.lineNumber, handler.lineRenderStyle, searchState.search)}</Stack></div>
                      </Stack>
