Horde: Add missing admin auth to preview controller
#rnx

diff --git a/Engine/Source/Programs/Horde/Horde.Build.Tests/TestSetup.cs b/Engine/Source/Programs/Horde/Horde.Build.Tests/TestSetup.cs
--- a/Engine/Source/Programs/Horde/Horde.Build.Tests/TestSetup.cs
+++ b/Engine/Source/Programs/Horde/Horde.Build.Tests/TestSetup.cs
@@ -307,7 +307,7 @@ namespace Horde.Build.Tests
 
 		private DashboardController GetDashboardController()
 		{
-			DashboardController dashboardCtrl = new DashboardController(DashboardPreviewCollection, ServerSettingsMon);
+			DashboardController dashboardCtrl = new DashboardController(DashboardPreviewCollection, ServerSettingsMon, GlobalConfigSnapshot);
 			dashboardCtrl.ControllerContext = GetControllerContext();
 			return dashboardCtrl;
 		}
diff --git a/Engine/Source/Programs/Horde/Horde.Build/Dashboard/DashboardController.cs b/Engine/Source/Programs/Horde/Horde.Build/Dashboard/DashboardController.cs
--- a/Engine/Source/Programs/Horde/Horde.Build/Dashboard/DashboardController.cs
+++ b/Engine/Source/Programs/Horde/Horde.Build/Dashboard/DashboardController.cs
@@ -12,6 +12,7 @@ using Horde.Build.Server;
 using System.Collections.Generic;
 using System.Linq;
 using Horde.Build.Agents;
+using Horde.Build.Acls;
 
 namespace Horde.Build.Dashboard
 {
@@ -34,16 +35,20 @@ namespace Horde.Build.Dashboard
 
 		private readonly IDashboardPreviewCollection _previewCollection;
 
+		private readonly IOptionsSnapshot<GlobalConfig> _globalConfig;
+
 		/// <summary>
 		/// Constructor
 		/// </summary>
 		/// <param name="previewCollection" />
 		/// <param name="serverSettings">Server settings</param>
-		public DashboardController(IDashboardPreviewCollection previewCollection, IOptionsMonitor<ServerSettings> serverSettings)
+		/// <param name="globalConfig" />
+		public DashboardController(IDashboardPreviewCollection previewCollection, IOptionsMonitor<ServerSettings> serverSettings, IOptionsSnapshot<GlobalConfig> globalConfig)
 		{
 			_authenticationScheme = AccountController.GetAuthScheme(serverSettings.CurrentValue.AuthMethod);
 			_previewCollection = previewCollection;
 			_settings = serverSettings.CurrentValue;
+			_globalConfig = globalConfig;
 		}
 
 		/// <summary>	
@@ -147,6 +152,11 @@ namespace Horde.Build.Dashboard
 		[Route("/api/v1/dashboard/preview")]
 		public async Task<ActionResult<GetDashboardPreviewResponse>> CreateDashbordPreview([FromBody] CreateDashboardPreviewRequest request)
 		{
+			if (!_globalConfig.Value.Authorize(AclAction.AdminWrite, User))
+			{
+				return Forbid();
+			}
+
 			IDashboardPreview preview = await _previewCollection.AddPreviewAsync(request.Summary);
 
 			if (!String.IsNullOrEmpty(request.ExampleLink) || !String.IsNullOrEmpty(request.DiscussionLink) || !String.IsNullOrEmpty(request.TrackingLink))
@@ -172,6 +182,11 @@ namespace Horde.Build.Dashboard
 		[Route("/api/v1/dashboard/preview")]
 		public async Task<ActionResult<GetDashboardPreviewResponse>> UpdateDashbordPreview([FromBody] UpdateDashboardPreviewRequest request)
 		{
+			if (!_globalConfig.Value.Authorize(AclAction.AdminWrite, User))
+			{
+				return Forbid();
+			}
+
 			IDashboardPreview? preview = await _previewCollection.UpdatePreviewAsync(request.Id, request.Summary, request.DeployedCL, request.Open, request.ExampleLink, request.DiscussionLink, request.TrackingLink);
 			
 			if (preview == null)
