Horde: Allow running chained jobs at the default latest change for a template, rather than the change of the triggering job.

diff --git a/Engine/Source/Programs/Horde/Horde.Build/Jobs/IJob.cs b/Engine/Source/Programs/Horde/Horde.Build/Jobs/IJob.cs
--- a/Engine/Source/Programs/Horde/Horde.Build/Jobs/IJob.cs
+++ b/Engine/Source/Programs/Horde/Horde.Build/Jobs/IJob.cs
@@ -530,6 +530,11 @@ namespace Horde.Build.Jobs
 		/// The triggered job id
 		/// </summary>
 		public JobId? JobId { get; }
+
+		/// <summary>
+		/// Whether to run the latest change, or default change for the template, when starting the new job. Uses same change as the triggering job by default.
+		/// </summary>
+		public bool UseDefaultChangeForTemplate { get; }
 	}
 
 	/// <summary>
diff --git a/Engine/Source/Programs/Horde/Horde.Build/Jobs/JobCollection.cs b/Engine/Source/Programs/Horde/Horde.Build/Jobs/JobCollection.cs
--- a/Engine/Source/Programs/Horde/Horde.Build/Jobs/JobCollection.cs
+++ b/Engine/Source/Programs/Horde/Horde.Build/Jobs/JobCollection.cs
@@ -177,6 +177,7 @@ namespace Horde.Build.Jobs
 			public string Target { get; set; }
 			public TemplateId TemplateRefId { get; set; }
 			public JobId? JobId { get; set; }
+			public bool UseDefaultChangeForTemplate { get; set; }
 
 			[BsonConstructor]
 			private ChainedJobDocument()
@@ -188,6 +189,7 @@ namespace Horde.Build.Jobs
 			{
 				Target = trigger.Trigger;
 				TemplateRefId = trigger.TemplateId;
+				UseDefaultChangeForTemplate = trigger.UseDefaultChangeForTemplate;
 			}
 		}
 
diff --git a/Engine/Source/Programs/Horde/Horde.Build/Jobs/JobService.cs b/Engine/Source/Programs/Horde/Horde.Build/Jobs/JobService.cs
--- a/Engine/Source/Programs/Horde/Horde.Build/Jobs/JobService.cs
+++ b/Engine/Source/Programs/Horde/Horde.Build/Jobs/JobService.cs
@@ -1248,6 +1248,8 @@ namespace Horde.Build.Jobs
 		/// <returns>New job instance</returns>
 		private async Task<IJob?> FireJobTriggerAsync(IJob job, IGraph graph, IChainedJob jobTrigger, StreamConfig streamConfig)
 		{
+			CancellationToken cancellationToken = CancellationToken.None;
+
 			using IScope traceScope = GlobalTracer.Instance.BuildSpan("JobService.FireJobTriggerAsync").StartActive();
 			traceScope.Span.SetTag("Job", job.Id.ToString());
 			traceScope.Span.SetTag("Graph", graph.Id);
@@ -1273,14 +1275,39 @@ namespace Horde.Build.Jobs
 					ITemplate template = await _templateCollection.GetOrAddAsync(templateRefConfig);
 
 					IGraph triggerGraph = await _graphs.AddAsync(template, streamConfig.InitialAgentType);
-					_logger.LogInformation("Creating downstream job {ChainedJobId} from job {JobId}", chainedJobId, newJob.Id);
+
+					int? change = null;
+					int codeChange;
+					if (jobTrigger.UseDefaultChangeForTemplate)
+					{
+						ICommitCollection commits = _perforceService.GetCommits(streamConfig);
+
+						if (templateRefConfig.DefaultChange != null)
+						{
+							change = await EvaluateChangeQueriesAsync(streamConfig.Id, templateRefConfig.DefaultChange, null, commits, cancellationToken);
+						}
+						if (change == null)
+						{
+							change = await commits.GetLatestNumberAsync(cancellationToken);
+						}
+
+						ICommit? commit = await commits.GetLastCodeChangeAsync(change, cancellationToken);
+						codeChange = commit?.Number ?? change.Value;
+					}
+					else
+					{
+						change = job.Change;
+						codeChange = job.CodeChange;
+					}
+
+					_logger.LogInformation("Creating downstream job {ChainedJobId} from job {JobId} at change {Change}", chainedJobId, newJob.Id, change);
 
 					CreateJobOptions options = new CreateJobOptions(templateRefConfig);
 					options.PreflightChange = newJob.PreflightChange;
 					options.PreflightDescription = newJob.PreflightDescription;
 					options.Arguments.AddRange(template.GetDefaultArguments());
 
-					await CreateJobAsync(chainedJobId, streamConfig, jobTrigger.TemplateRefId, template.Hash, triggerGraph, templateRefConfig.Name, newJob.Change, newJob.CodeChange, options);
+					await CreateJobAsync(chainedJobId, streamConfig, jobTrigger.TemplateRefId, template.Hash, triggerGraph, templateRefConfig.Name, change.Value, codeChange, options);
 					return newJob;
 				}
 
diff --git a/Engine/Source/Programs/Horde/Horde.Build/Streams/StreamConfig.cs b/Engine/Source/Programs/Horde/Horde.Build/Streams/StreamConfig.cs
--- a/Engine/Source/Programs/Horde/Horde.Build/Streams/StreamConfig.cs
+++ b/Engine/Source/Programs/Horde/Horde.Build/Streams/StreamConfig.cs
@@ -663,6 +663,11 @@ namespace Horde.Build.Streams
 		/// </summary>
 		[Required]
 		public TemplateId TemplateId { get; set; }
+
+		/// <summary>
+		/// Whether to use the default change for the template rather than the change for the parent job.
+		/// </summary>
+		public bool UseDefaultChangeForTemplate { get; set; }
 	}
 
 	/// <summary>
