Horde: Fix preview endpoint documentation
#skipci

diff --git a/Engine/Source/Programs/Horde/Horde.Build/Dashboard/DashboardController.cs b/Engine/Source/Programs/Horde/Horde.Build/Dashboard/DashboardController.cs
--- a/Engine/Source/Programs/Horde/Horde.Build/Dashboard/DashboardController.cs
+++ b/Engine/Source/Programs/Horde/Horde.Build/Dashboard/DashboardController.cs
@@ -144,7 +144,7 @@ namespace Horde.Build.Dashboard
 		}
 
 		/// <summary>
-		/// Create a new dashboard preview items
+		/// Create a new dashboard preview item
 		/// </summary>
 		/// <returns>Config information needed by the dashboard</returns>
 		[HttpPost]
@@ -174,7 +174,7 @@ namespace Horde.Build.Dashboard
 		}
 
 		/// <summary>
-		/// Create a new dashboard preview items
+		/// Update a dashboard preview item
 		/// </summary>
 		/// <returns>Config information needed by the dashboard</returns>
 		[HttpPut]
@@ -198,7 +198,7 @@ namespace Horde.Build.Dashboard
 		}
 
 		/// <summary>
-		/// Query all the dashboard preview items
+		/// Query dashboard preview items
 		/// </summary>
 		/// <returns>Config information needed by the dashboard</returns>
 		[HttpGet]
