Horde: Fix issue with log view search
#skipci
#rnx

diff --git a/Engine/Source/Programs/Horde/HordeDashboard/src/components/LogView.tsx b/Engine/Source/Programs/Horde/HordeDashboard/src/components/LogView.tsx
--- a/Engine/Source/Programs/Horde/HordeDashboard/src/components/LogView.tsx
+++ b/Engine/Source/Programs/Horde/HordeDashboard/src/components/LogView.tsx
@@ -761,8 +761,12 @@ export const LogList: React.FC<{ logId: string }> = observer(({ logId }) => {
                   if (lineIdx < 0) {
                      lineIdx = 0;
                   }
+                  listRef?.scrollToIndex(0, (index) => { return handler.lineHeight; }, ScrollToMode.top);
                   listRef?.scrollToIndex(lineIdx, (index) => { return handler.lineHeight; }, ScrollToMode.top);
                }
+               else {
+                  listRef?.scrollToIndex(0, (index) => { return handler.lineHeight; }, ScrollToMode.top);
+               }
 
             };
          });
@@ -992,6 +996,7 @@ export const LogList: React.FC<{ logId: string }> = observer(({ logId }) => {
                                  <TextField
                                     suffix={searchState?.results?.length ? `${curSearchIdx + 1}/${searchState?.results?.length}` : undefined}
                                     spellCheck={false}
+                                    autoComplete="off"
                                     deferredValidationTime={1000}
                                     validateOnLoad={false}
                                     componentRef={searchBox}
