Horde: Dashboard preview support
#rnx
#skipci

diff --git a/Engine/Source/Programs/Horde/HordeDashboard/src/backend/Dashboard.ts b/Engine/Source/Programs/Horde/HordeDashboard/src/backend/Dashboard.ts
--- a/Engine/Source/Programs/Horde/HordeDashboard/src/backend/Dashboard.ts
+++ b/Engine/Source/Programs/Horde/HordeDashboard/src/backend/Dashboard.ts
@@ -125,6 +125,16 @@ export class Dashboard {
 
     }
 
+    get preview(): boolean {
+        
+        try {            
+            return !!(process.env.REACT_APP_HORDE_DEBUG_PREVIEW) || window?.location?.hostname?.indexOf("preview") !== -1;
+        } catch (reason) {
+            console.error(reason);
+        }
+        return false;        
+    }
+
     get browser(): WebBrowser {
 
         const agent = window.navigator.userAgent.toLowerCase();
diff --git a/Engine/Source/Programs/Horde/HordeDashboard/src/components/Breadcrumbs.tsx b/Engine/Source/Programs/Horde/HordeDashboard/src/components/Breadcrumbs.tsx
--- a/Engine/Source/Programs/Horde/HordeDashboard/src/components/Breadcrumbs.tsx
+++ b/Engine/Source/Programs/Horde/HordeDashboard/src/components/Breadcrumbs.tsx
@@ -4,6 +4,7 @@ import { mergeStyleSets, MessageBar, MessageBarType, Separator, Spinner, Spinner
 import { observer } from 'mobx-react-lite';
 import React from 'react';
 import { Link } from 'react-router-dom';
+import dashboard from '../backend/Dashboard';
 import notices from '../backend/Notices';
 import { modeColors } from '../styles/Styles';
 
@@ -32,7 +33,7 @@ export const classes = mergeStyleSets({
 
 
 
-export const Breadcrumbs: React.FC<{ items: BreadcrumbItem[], title?: string, suppressHome?: boolean, spinner?: boolean }> = observer((({ items : itemsIn, title, suppressHome, spinner }) => {
+export const Breadcrumbs: React.FC<{ items: BreadcrumbItem[], title?: string, suppressHome?: boolean, spinner?: boolean }> = observer((({ items: itemsIn, title, suppressHome, spinner }) => {
 
    if (notices.updated) { }
 
@@ -82,7 +83,7 @@ export const Breadcrumbs: React.FC<{ items: BreadcrumbItem[], title?: string, su
                <Stack tokens={{ childrenGap: 0 }} disableShrink={true} styles={{ root: { margin: "auto", width: "100%", maxWidth: 1800 } }}>
                   <Stack horizontal tokens={{ childrenGap: 8, padding: 0 }}>{topElements}</Stack>
                   {
-                     <Stack horizontal tokens={{childrenGap: 18}}>
+                     <Stack horizontal tokens={{ childrenGap: 18 }}>
                         <Stack>{bottomElement}</Stack>
                         {!!spinner && <Spinner size={SpinnerSize.large} />}
                      </Stack>
@@ -90,18 +91,24 @@ export const Breadcrumbs: React.FC<{ items: BreadcrumbItem[], title?: string, su
                </Stack>
             </Stack>
          </Stack>
+         {!alert && !!dashboard.preview && <Stack horizontalAlign="center" disableShrink={true} style={{ position: "absolute", width: "100%", pointerEvents: "none" }} >
+            <Stack style={{ paddingTop: 22 }}>
+               <Text variant="large" style={{ fontFamily: "Horde Open Sans Light", opacity: 0.33 }}>Horde Preview</Text>
+            </Stack>
+         </Stack>}
+
          {!!alert && <Stack horizontalAlign="center" disableShrink={true} style={{ position: "absolute", width: "100%", pointerEvents: "none" }} >
-               <Stack horizontal>
-                  <Stack grow />
-                  <Stack>
-                     <MessageBar
-                        messageBarType={MessageBarType.severeWarning} isMultiline={false} >
-                        <Text variant={"small"} style={{ fontFamily: "Horde Open Sans Bold" }}>{alert}</Text>
-                     </MessageBar>
-                  </Stack>
-                  <Stack grow />
+            <Stack horizontal>
+               <Stack grow />
+               <Stack>
+                  <MessageBar
+                     messageBarType={MessageBarType.severeWarning} isMultiline={false} >
+                     <Text variant={"small"} style={{ fontFamily: "Horde Open Sans Bold" }}>{alert}</Text>
+                  </MessageBar>
                </Stack>
-            </Stack>}
+               <Stack grow />
+            </Stack>
+         </Stack>}
       </Stack>
    </Stack>
    );
diff --git a/Engine/Source/Programs/Horde/HordeDashboard/src/components/PreviewChanges.tsx b/Engine/Source/Programs/Horde/HordeDashboard/src/components/PreviewChanges.tsx
new file mode 100644
--- /dev/null
+++ b/Engine/Source/Programs/Horde/HordeDashboard/src/components/PreviewChanges.tsx
@@ -0,0 +1,97 @@
+import { DetailsList, DetailsListLayoutMode, IColumn, IconButton, Modal, PrimaryButton, SelectionMode, Stack, Text } from "@fluentui/react";
+import { Link as ReactRouterLink } from 'react-router-dom';
+import { hordeClasses } from "../styles/Styles";
+
+type ChangeItem = {
+   //
+   id: string;
+   // A summary of the change
+   summary: string;
+   link?: string;
+   slack?: string;
+}
+
+const changeItems: ChangeItem[] = [];
+
+let item = {
+   id: "1",
+   summary: "Improved horizontal space on job detail views to better support various desktop resolutions and scaling",
+   link: "/job/63f7c145d436940d3e099365", 
+   slack: "https://epicgames.slack.com/archives/C044JG73ARY/p1677529747047259"
+}
+changeItems.push(item);
+
+item = {
+   id: "2",
+   summary: "Timeline navigation for steps, agents, pools, and batches.  Click a span in the timeline to pin the tooltip, allowing navigation",
+   link: "/job/63f7c145d436940d3e099365",
+   slack: "https://epicgames.slack.com/archives/C044JG73ARY/p1677530063797479"
+}
+
+changeItems.push(item);
+
+const columns = [
+   { key: 'column1', name: 'ID', minWidth: 100, maxWidth: 100 },
+   { key: 'column2', name: 'Summary', minWidth: 540, maxWidth: 540 },
+   { key: 'column3', name: 'Preview', minWidth: 100, maxWidth: 100 },
+   { key: 'column4', name: 'Discussion', minWidth: 100, maxWidth: 200 },
+
+];
+
+const onRenderItemColumn = (item: ChangeItem, index?: number, columnIn?: IColumn) => {
+
+   const column = columnIn!;
+
+   // simple cases
+   switch (column.name) {
+      case 'ID':
+         return <Stack verticalFill verticalAlign="center"><Text styles={{ root: { fontFamily: "Horde Open Sans SemiBold" } }}>{`PREVIEW-${item.id}`}</Text></Stack>
+      case 'Summary':
+         return <Text style={{ whiteSpace: "pre-wrap" }}>{item.summary}</Text>
+      case 'Preview':
+         return !item.link ? null : <Stack verticalFill verticalAlign="center"><ReactRouterLink to={item.link}><Text>Example Link</Text></ReactRouterLink></Stack>
+      case 'Discussion':
+         return !item.slack ? null : <Stack verticalFill verticalAlign="center"><a href={item.slack} target="_blank" rel="noreferrer"><Text >Thread</Text></a></Stack>
+   }
+
+   return null;
+}
+
+
+export const PreviewChangesModal: React.FC<{ onClose: () => void }> = ({ onClose }) => {
+   return <Modal isOpen={true} isBlocking={true} topOffsetFixed={true} styles={{ main: { padding: 8, width: 1000, hasBeenOpened: false, top: "80px", position: "absolute" } }} className={hordeClasses.modal} onDismiss={() => { onClose() }}>
+      <Stack style={{ paddingRight: 18 }}>
+         <Stack style={{ paddingLeft: 18, paddingTop: 8 }}>
+            <Stack horizontal>
+               <Stack verticalAlign="center">
+                  <Text variant="mediumPlus" styles={{ root: { fontWeight: "unset", fontFamily: "Horde Open Sans SemiBold" } }}>Preview Changes</Text>
+               </Stack>
+               <Stack grow />
+               <Stack>
+                  <IconButton
+                     iconProps={{ iconName: 'Cancel' }}
+                     ariaLabel="Close popup modal"
+                     onClick={() => { onClose() }}
+                  />
+               </Stack>
+            </Stack>
+         </Stack>
+         <Stack style={{ paddingLeft: 24, paddingBottom: 24 }}>
+            <DetailsList
+               items={changeItems}
+               columns={columns}
+               layoutMode={DetailsListLayoutMode.justified}
+               isHeaderVisible={true}
+               selectionMode={SelectionMode.none}
+               onRenderItemColumn={onRenderItemColumn}
+            />
+         </Stack>
+         <Stack horizontal style={{ paddingBottom: 24 }}>
+            <Stack grow />
+            <Stack>
+               <PrimaryButton text="Ok" onClick={() => onClose()} />
+            </Stack>
+         </Stack>
+      </Stack>
+   </Modal>
+}
\ No newline at end of file
diff --git a/Engine/Source/Programs/Horde/HordeDashboard/src/components/TopNav.tsx b/Engine/Source/Programs/Horde/HordeDashboard/src/components/TopNav.tsx
--- a/Engine/Source/Programs/Horde/HordeDashboard/src/components/TopNav.tsx
+++ b/Engine/Source/Programs/Horde/HordeDashboard/src/components/TopNav.tsx
@@ -9,6 +9,7 @@ import { ProjectData } from "../backend/Api";
 import dashboard from '../backend/Dashboard';
 import { ProjectStore } from '../backend/ProjectStore';
 import { modeColors } from '../styles/Styles';
+import { PreviewChangesModal } from './PreviewChanges';
 import { VersionModal } from './VersionModal';
 
 const logoutURL = "/account";
@@ -338,6 +339,7 @@ export const TopNav: React.FC<{ suppressServer?: boolean }> = observer(({ suppre
 
    const [showMenu, setShowMenu] = useState(false);
    const [showVersion, setShowVersion] = useState(false);
+   const [showPreviewChanges, setShowPreviewChanges] = useState(false);
    const navigate = useNavigate();
    const { projectStore } = useBackend();
 
@@ -455,6 +457,15 @@ export const TopNav: React.FC<{ suppressServer?: boolean }> = observer(({ suppre
 
       const style = { ...menuStyles } as Partial<IContextualMenuStyles>;
 
+      if (dashboard.preview) {
+         const previewItem: ICommandBarItemProps = {
+            key: "preview_button",
+            text: "CHANGES",
+            onClick: () => {setShowPreviewChanges(true)}
+         };
+         cbProps.push(previewItem);
+      }
+
       const cbItem: ICommandBarItemProps = {
          key: "admin_button",
          text: "SERVER",
@@ -529,6 +540,7 @@ export const TopNav: React.FC<{ suppressServer?: boolean }> = observer(({ suppre
    return (
       <div style={{ backgroundColor: modeColors.header }}>
          {showVersion && <VersionModal show={true} onClose={() => { setShowVersion(false) }} />}
+         {showPreviewChanges && <PreviewChangesModal onClose={() => { setShowPreviewChanges(false) }} />}
          <Stack tokens={{ maxWidth: 1800, childrenGap: 0 }} disableShrink={true} styles={{ root: { backgroundColor: modeColors.header, margin: "auto", width: "100%" } }}>
             <Stack horizontal verticalAlign='center' styles={{ root: { height: "60px" } }} >
 
@@ -537,7 +549,7 @@ export const TopNav: React.FC<{ suppressServer?: boolean }> = observer(({ suppre
                      <img style={{ width: 32 }} src="/images/horde.svg" alt="" />
                   </Stack>
                   <Stack styles={{ root: { paddingTop: 11 } }}>
-                     <Text styles={{ root: { fontFamily: "Horde Raleway Bold", color: dashboard.darktheme ? "#FFFFFFFF" : modeColors.text } }}>HORDE</Text>
+                     <Text styles={{ root: { fontFamily: "Horde Raleway Bold", color: dashboard.darktheme ? "#FFFFFFFF" : modeColors.text } }}>HORDE{dashboard.preview ? " PREVIEW" : ""}</Text>
                   </Stack>
                </Stack>
                </Link>
