Horde: Fix log view search issue, needed to bump fluent version
#rnx
#skipci

diff --git a/Engine/Source/Programs/Horde/HordeDashboard/package.json b/Engine/Source/Programs/Horde/HordeDashboard/package.json
--- a/Engine/Source/Programs/Horde/HordeDashboard/package.json
+++ b/Engine/Source/Programs/Horde/HordeDashboard/package.json
@@ -7,7 +7,7 @@
   "dependencies": {
     "@blueprintjs/icons": "^3.13.0",
     "@datadog/browser-logs": "^4.30.0",
-    "@fluentui/react": "^8.104.5",
+    "@fluentui/react": "^8.106.2",
     "@fluentui/react-charting": "5.2.0",
     "crypto-js": "^4.1.1",
     "d3": "^7.6.1",
diff --git a/Engine/Source/Programs/Horde/HordeDashboard/src/components/LogView.tsx b/Engine/Source/Programs/Horde/HordeDashboard/src/components/LogView.tsx
--- a/Engine/Source/Programs/Horde/HordeDashboard/src/components/LogView.tsx
+++ b/Engine/Source/Programs/Horde/HordeDashboard/src/components/LogView.tsx
@@ -722,10 +722,7 @@ export const LogList: React.FC<{ logId: string }> = observer(({ logId }) => {
    const doQuery = (newValue: string) => {
 
       handler.currentWarning = undefined;
-      handler.currentError = undefined;
-
-
-      logListKey++;
+      handler.currentError = undefined;      
 
       if (!newValue) {
 
@@ -746,9 +743,7 @@ export const LogList: React.FC<{ logId: string }> = observer(({ logId }) => {
                curSearchIdx = 0;
 
                const lines = logResponse.lines;
-
-               logListKey++;
-
+               
                setSearchState({
                   search: newValue,
                   curRequest: undefined,
@@ -761,12 +756,9 @@ export const LogList: React.FC<{ logId: string }> = observer(({ logId }) => {
                   if (lineIdx < 0) {
                      lineIdx = 0;
                   }
-                  listRef?.scrollToIndex(0, (index) => { return handler.lineHeight; }, ScrollToMode.top);
+                  
                   listRef?.scrollToIndex(lineIdx, (index) => { return handler.lineHeight; }, ScrollToMode.top);
-               }
-               else {
-                  listRef?.scrollToIndex(0, (index) => { return handler.lineHeight; }, ScrollToMode.top);
-               }
+               } 
 
             };
          });
