Horde: Use a timeout on list procedural scroll.  There is way too much competing state on this view.  It needs reworking.
#rnx
#skipci

diff --git a/Engine/Source/Programs/Horde/HordeDashboard/src/components/LogView.tsx b/Engine/Source/Programs/Horde/HordeDashboard/src/components/LogView.tsx
--- a/Engine/Source/Programs/Horde/HordeDashboard/src/components/LogView.tsx
+++ b/Engine/Source/Programs/Horde/HordeDashboard/src/components/LogView.tsx
@@ -434,6 +434,7 @@ export const LogList: React.FC<{ logId: string }> = observer(({ logId }) => {
    globalSearchState = searchState;
    const vw = Math.max(document.documentElement.clientWidth, window.innerWidth || 0);
 
+   let inTimeout = false;
 
    useEffect(() => {
 
@@ -722,7 +723,7 @@ export const LogList: React.FC<{ logId: string }> = observer(({ logId }) => {
    const doQuery = (newValue: string) => {
 
       handler.currentWarning = undefined;
-      handler.currentError = undefined;      
+      handler.currentError = undefined;
 
       if (!newValue) {
 
@@ -743,22 +744,40 @@ export const LogList: React.FC<{ logId: string }> = observer(({ logId }) => {
                curSearchIdx = 0;
 
                const lines = logResponse.lines;
-               
-               setSearchState({
-                  search: newValue,
-                  curRequest: undefined,
-                  results: lines
-               });
 
+            
                if (lines.length) {
+
+                  logListKey++;
+
+                  setSearchState({
+                     search: newValue,
+                     curRequest: undefined,
+                     results: lines
+                  });
+   
+
                   handler.stopTrailing();
                   let lineIdx = lines[curSearchIdx] - 10;
                   if (lineIdx < 0) {
                      lineIdx = 0;
                   }
+
+                  // oof
+                  inTimeout = true;
+                  setTimeout(() => {
+                     inTimeout = false;
+                     listRef?.scrollToIndex(lineIdx, () => { return handler.lineHeight; }, ScrollToMode.top);
+                  }, 250)
                   
-                  listRef?.scrollToIndex(lineIdx, (index) => { return handler.lineHeight; }, ScrollToMode.top);
-               } 
+               } else {
+                  setSearchState({
+                     search: newValue,
+                     curRequest: undefined,
+                     results: lines
+                  });
+   
+               }
 
             };
          });
@@ -1002,7 +1021,7 @@ export const LogList: React.FC<{ logId: string }> = observer(({ logId }) => {
 
                                     onKeyPress={(ev) => {
 
-                                       if (ev.key === "Enter") {
+                                       if (ev.key === "Enter" && !searchState.curRequest && !inTimeout) {
                                           if (searchBox.current?.value === searchState.search) {
                                              searchDown();
                                           } else {
