Horde: Prevent deadlock when disposing JsonRpcAndStorageLogSink

diff --git a/Engine/Source/Programs/Horde/Horde.Agent/Utility/JsonRpcLogSink.cs b/Engine/Source/Programs/Horde/Horde.Agent/Utility/JsonRpcLogSink.cs
--- a/Engine/Source/Programs/Horde/Horde.Agent/Utility/JsonRpcLogSink.cs
+++ b/Engine/Source/Programs/Horde/Horde.Agent/Utility/JsonRpcLogSink.cs
@@ -112,6 +112,8 @@ namespace Horde.Agent.Utility
 			if (_tailTaskStop != null)
 			{
 				_tailTaskStop.Latch();
+				_newTailDataEvent.Latch();
+
 				try
 				{
 					await _tailTask;
