Horde: Catch all types of cancellation exceptions when creating a lease, not just task cancellation exceptions. Causing false-positive errors in DD.

diff --git a/Engine/Source/Programs/Horde/Horde.Build/Agents/AgentService.cs b/Engine/Source/Programs/Horde/Horde.Build/Agents/AgentService.cs
--- a/Engine/Source/Programs/Horde/Horde.Build/Agents/AgentService.cs
+++ b/Engine/Source/Programs/Horde/Horde.Build/Agents/AgentService.cs
@@ -518,7 +518,7 @@ namespace Horde.Build.Agents
 				Task<AgentLease?> task = await source.AssignLeaseAsync(agent, cancellationToken);
 				return task.ContinueWith(x => WrapAssignedLease(source, x, cancellationSource), TaskScheduler.Default);
 			}
-			catch (TaskCanceledException) when (cancellationToken.IsCancellationRequested)
+			catch (OperationCanceledException) when (cancellationToken.IsCancellationRequested)
 			{
 				return Task.FromResult<(ITaskSource, AgentLease)?>(null);
 			}
