Horde: Add clock times to timeline tooltip, fix selected span zoom behavior, shrink tooltip to better use space
#rnx
#skipci

diff --git a/Engine/Source/Programs/Horde/HordeDashboard/src/base/utilities/timeUtils.ts b/Engine/Source/Programs/Horde/HordeDashboard/src/base/utilities/timeUtils.ts
--- a/Engine/Source/Programs/Horde/HordeDashboard/src/base/utilities/timeUtils.ts
+++ b/Engine/Source/Programs/Horde/HordeDashboard/src/base/utilities/timeUtils.ts
@@ -379,7 +379,7 @@ export const getHumanTime = (timeIn: Date | string | undefined): string => {
 };
 
 
-export const getShortNiceTime = (timeIn: Date | string | undefined, relative: boolean = false, includeHour:boolean = false): string => {
+export const getShortNiceTime = (timeIn: Date | string | undefined, relative: boolean = false, includeHour:boolean = false, includeDay: boolean = true): string => {
 
     if (!timeIn) {
         return "";
@@ -401,7 +401,10 @@ export const getShortNiceTime = (timeIn: Date | string | undefined, relative: bo
 
     if (includeHour) {
         const format = dashboard.display24HourClock ? "HH:mm:ss z" : "LT z";
-        timeStr += ` at ${time.format(format)}`;
+        if (includeDay)
+            timeStr += ` at ${time.format(format)}`;
+        else
+            timeStr = time.format(format);
     }
     
     return timeStr;
diff --git a/Engine/Source/Programs/Horde/HordeDashboard/src/components/jobDetailsV2/JobDetailTimeline.tsx b/Engine/Source/Programs/Horde/HordeDashboard/src/components/jobDetailsV2/JobDetailTimeline.tsx
--- a/Engine/Source/Programs/Horde/HordeDashboard/src/components/jobDetailsV2/JobDetailTimeline.tsx
+++ b/Engine/Source/Programs/Horde/HordeDashboard/src/components/jobDetailsV2/JobDetailTimeline.tsx
@@ -10,7 +10,7 @@ import { Link, NavigateFunction, useNavigate } from "react-router-dom";
 import { GetBatchResponse, GetStepResponse, JobStepOutcome, StepData } from "../../backend/Api";
 import dashboard, { StatusColor } from "../../backend/Dashboard";
 import { ISideRailLink } from "../../base/components/SideRail";
-import { msecToElapsed } from "../../base/utilities/timeUtils";
+import { getShortNiceTime, msecToElapsed } from "../../base/utilities/timeUtils";
 import { hordeClasses, modeColors } from "../../styles/Styles";
 import { HistoryModal } from "../HistoryModal";
 import { JobDataView, JobDetailsV2 } from "./JobDetailsViewCommon";
@@ -551,6 +551,15 @@ class TimelineRenderer {
 
          svg!.selectAll(".x-axis").call(xAxis as any);
 
+         const span = dataView.tooltip.span;
+         if (span) {
+            svg!.select(`.cline`)
+               .attr("x1", xScale(span.utcStart.getTime() / 1000) - 1)
+               .attr("x2", xScale(span.utcFinish.getTime() / 1000) - 1)
+               .attr("y1", i => yScale(span.lane) as number + 16)
+               .attr("y2", i => yScale(span.lane) as number + 16)
+         }
+
       }
 
       const zoom = d3.zoom()
@@ -559,8 +568,6 @@ class TimelineRenderer {
          .translateExtent([[margin.left, -Infinity], [width - margin.right, Infinity]])
          .on("zoom", zoomed)
 
-
-
       svg.call(zoom as any);
 
       const closestData = (x: number, y: number): Span | undefined => {
@@ -640,6 +647,9 @@ class TimelineRenderer {
                   .attr("y2", i => yScale(span.lane) as number + 16)
                   .attr("stroke", color)
             }
+         } else {
+            this.svg!.select(`.cline`)
+               .attr("stroke-width", 0)
          }
 
          dataView.tooltip.update(span, d3.pointer(event, container)[0], d3.pointer(event, container)[1]);
@@ -700,7 +710,7 @@ const GraphTooltip: React.FC<{ dataView: TimelineDataView }> = observer(({ dataV
       return null;
    }
 
-   const textSize = 12;
+   const textSize = 11;
 
    let tipX = tooltip.x;
    let offsetX = 32;
@@ -739,7 +749,7 @@ const GraphTooltip: React.FC<{ dataView: TimelineDataView }> = observer(({ dataV
       }
 
       if (agent) {
-         valueElement = <button className="horde-link" style={{borderWidth: 0, padding: 0, backgroundColor: modeColors.background}}>
+         valueElement = <button className="horde-link" style={{ borderWidth: 0, padding: 0, backgroundColor: modeColors.background }}>
             {valueElement}
          </button>
       }
@@ -794,7 +804,7 @@ const GraphTooltip: React.FC<{ dataView: TimelineDataView }> = observer(({ dataV
       elements.push(dataElement("Cost:", `$${span.cost.toFixed(2)}`));
    }
 
-   const elapsed = msecToElapsed(span.utcFinish.getTime() - span.utcStart.getTime(), true, false);
+   let elapsed = msecToElapsed(span.utcFinish.getTime() - span.utcStart.getTime(), true, false);
 
    let spanType = "Wait:";
    if (span.type === 1) {
@@ -804,8 +814,14 @@ const GraphTooltip: React.FC<{ dataView: TimelineDataView }> = observer(({ dataV
       spanType = "Time:";
    }
 
+   const startTime = getShortNiceTime(span.utcStart, false, true, false);
+   const finishTime = getShortNiceTime(span.utcFinish, false, true, false);
+
+   elapsed = `${elapsed} / ${startTime} - ${finishTime}`;
+
    elements.push(dataElement(spanType, elapsed));
 
+
    return <div style={{
       position: "absolute",
       display: "block",
@@ -823,7 +839,7 @@ const GraphTooltip: React.FC<{ dataView: TimelineDataView }> = observer(({ dataV
    }}>
       {!!viewAgent && <HistoryModal agentId={viewAgent} onDismiss={() => setViewAgent("")} />}
       <Stack horizontal>
-         <Stack tokens={{ childrenGap: 6, padding: 16 }}>
+         <Stack tokens={{ childrenGap: 6, padding: 12 }}>
             {elements}
          </Stack>
          {!!tooltip.frozen && <Stack style={{ paddingLeft: 32, paddingRight: 12, paddingTop: 12, cursor: "pointer" }} onClick={() => { tooltip.freeze(false); tooltip.update() }}>
