Horde: Better time format on agent modal and agents modals, also use the 24 hour clock preference
#skipci

diff --git a/Engine/Source/Programs/Horde/HordeDashboard/src/components/AgentView.tsx b/Engine/Source/Programs/Horde/HordeDashboard/src/components/AgentView.tsx
--- a/Engine/Source/Programs/Horde/HordeDashboard/src/components/AgentView.tsx
+++ b/Engine/Source/Programs/Horde/HordeDashboard/src/components/AgentView.tsx
@@ -11,6 +11,7 @@ import { agentStore } from '../backend/AgentStore';
 import { AgentData, BatchUpdatePoolRequest, GetAgentResponse, LeaseState, PoolData } from '../backend/Api';
 import { copyToClipboard } from '../base/utilities/clipboard';
 import { useWindowSize } from '../base/utilities/hooks';
+import { getShortNiceTime } from '../base/utilities/timeUtils';
 import { hexToRGB, hordeClasses, linearInterpolate } from '../styles/Styles';
 import { Breadcrumbs } from './Breadcrumbs';
 import { ConfirmationDialog } from './ConfirmationDialog';
@@ -2396,7 +2397,7 @@ export const AgentView: React.FC = observer(() => {
                let subtitle = "";
                if (!agent.online) {
                   title = `Offline - ${agent.lastShutdownReason}`;
-                  subtitle = `(Last online at ${moment.utc(agent.updateTime).local().format("YYYY-MM-DD HH:mm:ss")})`;
+                  subtitle = `(Last online at ${getShortNiceTime(agent.updateTime, true, true, true)}`;
                }
                else if (agent.pendingShutdown) {
                   title = `Pending Shutdown`;
@@ -2406,7 +2407,7 @@ export const AgentView: React.FC = observer(() => {
                   if (agent.conformAttemptCount && agent.conformAttemptCount > 0) {
                      title = `${agent.pendingFullConform ? "Full " : ""}Conform failed (${agent.conformAttemptCount} attempts)`;
                      if (agent.nextConformTime) {
-                        subtitle = `Next attempt at ${moment.utc(agent.nextConformTime).local().format("YYYY-MM-DD HH:mm:ss")}`;
+                        subtitle = `Next attempt at ${getShortNiceTime(agent.nextConformTime, true, true, true)}`;
                      }
                   }
                   else {
diff --git a/Engine/Source/Programs/Horde/HordeDashboard/src/components/HistoryModal.tsx b/Engine/Source/Programs/Horde/HordeDashboard/src/components/HistoryModal.tsx
--- a/Engine/Source/Programs/Horde/HordeDashboard/src/components/HistoryModal.tsx
+++ b/Engine/Source/Programs/Horde/HordeDashboard/src/components/HistoryModal.tsx
@@ -9,6 +9,7 @@ import backend from "../backend";
 import { agentStore } from "../backend/AgentStore";
 import { AgentData, GetAgentLeaseResponse, GetAgentSessionResponse, JobStepBatchError, LeaseData, SessionData, UpdateAgentRequest } from "../backend/Api";
 import dashboard from "../backend/Dashboard";
+import { getShortNiceTime } from "../base/utilities/timeUtils";
 import { hordeClasses, modeColors } from "../styles/Styles";
 import { BatchStatusIcon, LeaseStatusIcon } from "./StatusIcon";
 
@@ -669,9 +670,9 @@ export const HistoryModal: React.FC<{ agentId: string | undefined, onDismiss: (.
             case 'name':
                return <Stack>{lease.name}</Stack>;
             case 'startTime':
-               return <Stack styles={{ root: { height: '100%', } }} horizontal horizontalAlign={'center'}><Stack.Item align={"center"}>{lease.startTime.toString().substring(0, 24)}</Stack.Item></Stack>
+               return <Stack styles={{ root: { height: '100%', } }} horizontal horizontalAlign={'center'}><Stack.Item align={"center"}>{getShortNiceTime(lease.startTime, false, true, true)}</Stack.Item></Stack>
             case 'endTime':
-               return <Stack styles={{ root: { height: '100%', } }} horizontal horizontalAlign={'center'}><Stack.Item align={"center"}>{lease.finishTime?.toString().substring(0, 24) || ""}</Stack.Item></Stack>
+               return <Stack styles={{ root: { height: '100%', } }} horizontal horizontalAlign={'center'}><Stack.Item align={"center"}>{getShortNiceTime(lease.finishTime, false, true, true)}</Stack.Item></Stack>
             case 'description':
                link = "";
                if (lease.details) {
@@ -701,9 +702,9 @@ export const HistoryModal: React.FC<{ agentId: string | undefined, onDismiss: (.
             case 'id':
                return <Stack styles={{ root: { height: '100%', } }} horizontal><Stack.Item align={"center"}>{session.id}</Stack.Item></Stack>
             case 'startTime':
-               return <Stack styles={{ root: { height: '100%', } }} horizontal horizontalAlign={'center'}><Stack.Item align={"center"}>{session.startTime.toString().substring(0, 24)}</Stack.Item></Stack>
+               return <Stack styles={{ root: { height: '100%', } }} horizontal horizontalAlign={'center'}><Stack.Item align={"center"}>{getShortNiceTime(session.startTime, false, true)}</Stack.Item></Stack>
             case 'endTime':
-               return <Stack styles={{ root: { height: '100%', } }} horizontal horizontalAlign={'center'}><Stack.Item align={"center"}>{session.finishTime?.toString().substring(0, 24) || ""}</Stack.Item></Stack>
+               return <Stack styles={{ root: { height: '100%', } }} horizontal horizontalAlign={'center'}><Stack.Item align={"center"}>{getShortNiceTime(session.finishTime, false, true)}</Stack.Item></Stack>
             default:
                return <span>{session[column!.fieldName as keyof SessionData] as string}</span>;
          }
