Horde: Default all streams on when choosing an automation project, use the custom multilist chooser
#rnx
#skipci

diff --git a/Engine/Source/Programs/Horde/HordeDashboard/src/backend/AutomationTestData.ts b/Engine/Source/Programs/Horde/HordeDashboard/src/backend/AutomationTestData.ts
--- a/Engine/Source/Programs/Horde/HordeDashboard/src/backend/AutomationTestData.ts
+++ b/Engine/Source/Programs/Horde/HordeDashboard/src/backend/AutomationTestData.ts
@@ -48,7 +48,7 @@ export class TestDataHandler {
         this.load();
     }
 
-    setAutomation(automation: string) {
+    setAutomation(automation: string, initStreams: boolean = false) {
         this.state.automation = automation;
 
         if (automation) {
@@ -73,6 +73,11 @@ export class TestDataHandler {
                 this.state.suites = undefined;
             }
 
+            if (initStreams) {
+                const streams = this.getAutomationStreams(automation);
+                this.state.streams = streams.map(s => s);
+            }
+    
         }
 
         if (this.updateSearch()) {
diff --git a/Engine/Source/Programs/Horde/HordeDashboard/src/components/AutomationView.tsx b/Engine/Source/Programs/Horde/HordeDashboard/src/components/AutomationView.tsx
--- a/Engine/Source/Programs/Horde/HordeDashboard/src/components/AutomationView.tsx
+++ b/Engine/Source/Programs/Horde/HordeDashboard/src/components/AutomationView.tsx
@@ -7,7 +7,7 @@ import moment from 'moment';
 import React, { useState } from 'react';
 import { useSearchParams } from 'react-router-dom';
 import backend from '../backend';
-import { GetTestDataRefResponse, GetTestMetaResponse, GetTestResponse, GetTestSuiteResponse, TestOutcome } from '../backend/Api';
+import { GetTestDataRefResponse, GetTestMetaResponse, GetTestResponse, GetTestSuiteResponse, StreamData, TestOutcome } from '../backend/Api';
 import { FilterState, TestDataHandler } from '../backend/AutomationTestData';
 import dashboard, { StatusColor } from '../backend/Dashboard';
 import { projectStore } from '../backend/ProjectStore';
@@ -27,7 +27,6 @@ type DivSelectionType = d3.Selection<HTMLDivElement, unknown, null, undefined>;
 let multiComboBoxId = 0;
 
 const StreamChooser: React.FC<{ handler: TestDataHandler }> = observer(({ handler }) => {
-
    // subscribe
    if (handler.updated) { }
 
@@ -37,66 +36,39 @@ const StreamChooser: React.FC<{ handler: TestDataHandler }> = observer(({ handle
       return null;
    }
 
-   const options: IContextualMenuItem[] = [];
-
-   const projects = projectStore.projects.sort((a, b) => a.order - b.order);
-
-   const streams = handler.getAutomationStreams(automation);
-
-   let anySelected = false;
-
-   projects.forEach(p => {
+   const cstreams: Set<string> = new Set(handler.state.streams ?? []);
+   const streams = handler.getAutomationStreams(automation).sort((a, b) => a.localeCompare(b)).map(a => projectStore.streamById(a)).filter(a => !!a) as StreamData[];
 
-      let projectStreams = p.streams?.filter(s => !!streams.find(stream => stream === s.id));
-
-      if (!projectStreams?.length) {
-         return;
-      }
-
-      projectStreams = projectStreams.sort((a, b) => a.name.localeCompare(b.name));
+   if (!streams.length) {
+      return null;
+   }
 
-      const streamItems: IContextualMenuItem[] = [];
+   const options: IComboBoxOption[] = [];
+   streams.forEach(stream => {
+      options.push({ key: stream.id, text: stream.fullname ?? "???" });
+   });
 
-      let anyChecked = false;
-      projectStreams.forEach(s => {
-         const checked = handler.state.streams?.find(id => id === s.id);
-         if (checked) {
-            anyChecked = true;
-            anySelected = true;
-         }
-         streamItems.push({
-            key: `stream_select_${p.id}_${s.id}}`, canCheck: true, isChecked: checked, text: s.name, onClick: (ev, item) => {
+   return <Stack style={{ paddingTop: 12, paddingBottom: 4 }}>
+      <Stack style={{ paddingTop: 0, paddingBottom: 4 }}>
+         <Label>Streams</Label>
+      </Stack>
+      <Stack tokens={{ childrenGap: 6 }} style={{ paddingLeft: 12 }}>
+         <MultiOptionChooser options={options} initialKeysIn={handler.state.streams ?? []} updateKeys={(keys) => {
 
-               if (!item) {
-                  return;
+            streams.forEach(t => {
+               const selected = keys.find(k => k === t.id);
+               if (!selected && cstreams.has(t.id)) {
+                  handler.removeStream(t.id);
                }
-               // don't close
-               ev?.preventDefault()
-
-               !checked ? handler.addStream(s.id) : handler.removeStream(s.id);
-            }
-         });
-      });
-
-      const subMenuProps: IContextualMenuProps = {
-         shouldFocusOnMount: true,
-         subMenuHoverDelay: 0,
-         items: streamItems,
-      };
-
-      options.push({ key: `project_stream_select_${p.id}`, canCheck: true, isChecked: anyChecked, text: p.name, subMenuProps: subMenuProps });
-   })
-
-   const menuProps: IContextualMenuProps = {
-      shouldFocusOnMount: true,
-      subMenuHoverDelay: 0,
-      items: options,
-   };
-
-   return <DefaultButton style={{ width: 270, textAlign: "left" }} menuProps={menuProps} text={anySelected ? "Select" : "None Selected"} />
+               else if (selected && !cstreams.has(t.id)) {
+                  handler.addStream(t.id);
+               }
+            });
+         }} />
+      </Stack>
+   </Stack>
 });
 
-
 const AutomationChooser: React.FC<{ handler: TestDataHandler }> = observer(({ handler }) => {
 
    // subscribe
@@ -107,7 +79,12 @@ const AutomationChooser: React.FC<{ handler: TestDataHandler }> = observer(({ ha
    const options: IContextualMenuItem[] = [];
 
    handler.automation.forEach(a => {
-      options.push({ key: `automation_${a}`, text: a, onClick: () => { multiComboBoxId++; handler.setAutomation(a); } });
+      options.push({
+         key: `automation_${a}`, text: a, onClick: () => {
+            multiComboBoxId++;
+            handler.setAutomation(a, true);
+         }
+      });
    });
 
    const menuProps: IContextualMenuProps = {
@@ -439,14 +416,9 @@ const AutomationSidebarLeft: React.FC<{ handler: TestDataHandler }> = ({ handler
                <AutomationChooser handler={handler} />
             </Stack>
          </Stack>
-         {!!handler.state.automation && <Stack>
-            <Stack style={{ paddingTop: 12, paddingBottom: 4 }}>
-               <Label>Streams</Label>
-            </Stack>
-            <Stack tokens={{ childrenGap: 6 }} style={{ paddingLeft: 12 }}>
-               <StreamChooser handler={handler} />
-            </Stack>
-         </Stack>}
+         <Stack>
+            <StreamChooser handler={handler} />
+         </Stack>
          <Stack>
             <TestChooser handler={handler} />
          </Stack>
