Horde: Link to pool instead of agents page on step agent rows
#skipci
#rnx

diff --git a/Engine/Source/Programs/Horde/HordeDashboard/src/components/jobDetailsV2/JobDetailViewSteps.tsx b/Engine/Source/Programs/Horde/HordeDashboard/src/components/jobDetailsV2/JobDetailViewSteps.tsx
--- a/Engine/Source/Programs/Horde/HordeDashboard/src/components/jobDetailsV2/JobDetailViewSteps.tsx
+++ b/Engine/Source/Programs/Horde/HordeDashboard/src/components/jobDetailsV2/JobDetailViewSteps.tsx
@@ -244,7 +244,7 @@ export const StepsPanelInner: React.FC<{ jobDetails: JobDetailsV2, depStepId?: s
                if (item.agentPool && (item.batch?.state === JobStepBatchState.Ready || item.batch?.error === JobStepBatchError.NoAgentsOnline || item.batch?.error === JobStepBatchError.NoAgentsInPool)) {
                   return <Stack horizontal disableShrink={true} horizontalAlign="start">
                      {(error || warning) && <Icon styles={{ root: { paddingTop: 3, paddingRight: 8 } }} className={style!.className} iconName="Square" />}
-                     <Link to={`/agents?agent=${encodeURI(item.agentPool)}&exact=true`} ><Text nowrap={true}>{batchText}</Text></Link>
+                     <Link to={`/pools?pool=${encodeURI(item.agentPool)}`} ><Text nowrap={true}>{batchText}</Text></Link>
                   </Stack>                     
                }
                return <Stack horizontal disableShrink={true} horizontalAlign="start">
