Horde: Fix automation dashboard selection crash
#skipci
#rnx

diff --git a/Engine/Source/Programs/Horde/HordeDashboard/src/components/AutomationView.tsx b/Engine/Source/Programs/Horde/HordeDashboard/src/components/AutomationView.tsx
--- a/Engine/Source/Programs/Horde/HordeDashboard/src/components/AutomationView.tsx
+++ b/Engine/Source/Programs/Horde/HordeDashboard/src/components/AutomationView.tsx
@@ -141,14 +141,12 @@ const MultiOptionChooser: React.FC<{ options: IComboBoxOption[], initialKeysIn:
             }
             setSelectedKeys(updatedKeys);
          }
-
-         updateKeys([...updatedKeys.filter(k => k !== 'selectAll')]);
       }
    };
 
    const comboBoxStyles: Partial<IComboBoxStyles> = { root: { width: 270 } };
 
-   return <ComboBox key={`multi_option_id_${multiComboBoxId}`} placeholder="None" defaultSelectedKey={initialKeys} multiSelect options={options} onChange={onChange} styles={comboBoxStyles} />
+   return <ComboBox key={`multi_option_id_${multiComboBoxId}`} placeholder="None" defaultSelectedKey={initialKeys} multiSelect options={options} onChange={onChange} onMenuDismissed={() => { updateKeys([...selectedKeys.filter(k => k !== 'selectAll')]); }} styles={comboBoxStyles} />
 };
 
 const TestChooser: React.FC<{ handler: TestDataHandler }> = observer(({ handler }) => {
