Horde: Fix automation dashboard crash, again
#rnx
#skipci

diff --git a/Engine/Source/Programs/Horde/HordeDashboard/src/backend/AutomationTestData.ts b/Engine/Source/Programs/Horde/HordeDashboard/src/backend/AutomationTestData.ts
--- a/Engine/Source/Programs/Horde/HordeDashboard/src/backend/AutomationTestData.ts
+++ b/Engine/Source/Programs/Horde/HordeDashboard/src/backend/AutomationTestData.ts
@@ -883,6 +883,10 @@ export class TestDataHandler {
         const state = this.state;
 
         if (!state.automation || !state.streams?.length || (!state.tests?.length && !state.suites?.length)) {
+            this.refs = [];
+            this.streamStatus = new Map();
+            this.setQueryLoading(false);
+            this.setQueryUpdated();
             return;
         }
 
diff --git a/Engine/Source/Programs/Horde/HordeDashboard/src/components/AutomationView.tsx b/Engine/Source/Programs/Horde/HordeDashboard/src/components/AutomationView.tsx
--- a/Engine/Source/Programs/Horde/HordeDashboard/src/components/AutomationView.tsx
+++ b/Engine/Source/Programs/Horde/HordeDashboard/src/components/AutomationView.tsx
@@ -53,7 +53,7 @@ const StreamChooser: React.FC<{ handler: TestDataHandler }> = observer(({ handle
          <Label>Streams</Label>
       </Stack>
       <Stack tokens={{ childrenGap: 6 }} style={{ paddingLeft: 12 }}>
-         <MultiOptionChooser options={options} initialKeysIn={handler.state.streams ?? []} updateKeys={(keys) => {
+         <MultiOptionChooser id="stream" optionsIn={options} initialKeysIn={handler.state.streams ?? []} updateKeys={(keys) => {
 
             streams.forEach(t => {
                const selected = keys.find(k => k === t.id);
@@ -97,56 +97,27 @@ const AutomationChooser: React.FC<{ handler: TestDataHandler }> = observer(({ ha
 });
 
 
-const MultiOptionChooser: React.FC<{ options: IComboBoxOption[], initialKeysIn: string[], updateKeys: (selectedKeys: string[]) => void }> = ({ options, initialKeysIn, updateKeys }) => {
+const MultiOptionChooser: React.FC<{ id: string, optionsIn: IComboBoxOption[], initialKeysIn: string[], updateKeys: (selectedKeys: string[]) => void }> = ({ id, optionsIn, initialKeysIn, updateKeys }) => {
 
-   let initialKeys = [...initialKeysIn];
+   const comboBoxRef = React.useRef<IComboBox>(null);
 
-   const [selectedKeys, setSelectedKeys] = React.useState<string[]>(initialKeys);
+   let initialKeys = [...initialKeysIn];
+   let options = [...optionsIn];
 
    if (options.length === initialKeys.length) {
-      selectedKeys.push('selectAll');
+      initialKeys.push('selectAll');
    }
 
    options.unshift({ key: 'selectAll', text: 'Select All', itemType: SelectableOptionMenuItemType.SelectAll });
 
-   const selectableOptions = options.filter(
-      option =>
-         (option.itemType === SelectableOptionMenuItemType.Normal || option.itemType === undefined) && !option.disabled,
-   );
+   const comboBoxStyles: Partial<IComboBoxStyles> = { root: { width: 270 } };   
 
-   const onChange = (
-      event: React.FormEvent<IComboBox>,
-      option?: IComboBoxOption,
-      index?: number,
-      value?: string,
-   ): void => {
-      const selected = option?.selected;
-      const currentSelectedOptionKeys = selectedKeys.filter(key => key !== 'selectAll');
-      const selectAllState = currentSelectedOptionKeys.length === selectableOptions.length;
-
-      let updatedKeys: string[] = [];
-
-      if (option) {
-         if (option?.itemType === SelectableOptionMenuItemType.SelectAll) {
-            updatedKeys = selectAllState ? [] : ['selectAll', ...selectableOptions.map(o => o.key as string)];
-            selectAllState
-               ? setSelectedKeys(updatedKeys)
-               : setSelectedKeys(updatedKeys);
-         } else {
-            updatedKeys = selected
-               ? [...currentSelectedOptionKeys, option!.key as string]
-               : currentSelectedOptionKeys.filter(k => k !== option.key);
-            if (updatedKeys.length === selectableOptions.length) {
-               updatedKeys.push('selectAll');
-            }
-            setSelectedKeys(updatedKeys);
-         }
+   return <ComboBox componentRef={comboBoxRef} key={`multi_option_${id}_${multiComboBoxId}`} placeholder="None" defaultSelectedKey={initialKeys} multiSelect options={options} onResolveOptions={ () =>options} onMenuDismiss={() => {
+      if (comboBoxRef?.current?.selectedOptions) {
+         const selectedKeys = comboBoxRef.current.selectedOptions.map(o => o.key as string).filter(k => k !== 'selectAll');
+         setTimeout(() => { multiComboBoxId++; updateKeys(selectedKeys) }, 250);
       }
-   };
-
-   const comboBoxStyles: Partial<IComboBoxStyles> = { root: { width: 270 } };
-
-   return <ComboBox key={`multi_option_id_${multiComboBoxId}`} placeholder="None" defaultSelectedKey={initialKeys} multiSelect options={options} onChange={onChange} onMenuDismissed={() => { updateKeys([...selectedKeys.filter(k => k !== 'selectAll')]); }} styles={comboBoxStyles} />
+   }} styles={comboBoxStyles} />
 };
 
 const TestChooser: React.FC<{ handler: TestDataHandler }> = observer(({ handler }) => {
@@ -173,7 +144,7 @@ const TestChooser: React.FC<{ handler: TestDataHandler }> = observer(({ handler
          <Label>Tests</Label>
       </Stack>
       <Stack tokens={{ childrenGap: 6 }} style={{ paddingLeft: 12 }}>
-         <MultiOptionChooser options={options} initialKeysIn={handler.state.tests ?? []} updateKeys={(keys) => {
+         <MultiOptionChooser id="test" optionsIn={options} initialKeysIn={handler.state.tests ?? []} updateKeys={(keys) => {
 
             handler.streamTests.forEach(t => {
                const selected = keys.find(k => k === t.name);
@@ -213,7 +184,7 @@ const SuiteChooser: React.FC<{ handler: TestDataHandler }> = observer(({ handler
          <Label>Suites</Label>
       </Stack>
       <Stack tokens={{ childrenGap: 6 }} style={{ paddingLeft: 12 }}>
-         <MultiOptionChooser options={options} initialKeysIn={handler.state.suites ?? []} updateKeys={(keys) => {
+         <MultiOptionChooser id="suite" optionsIn={options} initialKeysIn={handler.state.suites ?? []} updateKeys={(keys) => {
 
             handler.streamSuites.forEach(t => {
                const selected = keys.find(k => k === t.name);
@@ -248,7 +219,7 @@ const PlatformChooser: React.FC<{ handler: TestDataHandler }> = observer(({ hand
          <Label>Platforms</Label>
       </Stack>
       <Stack tokens={{ childrenGap: 6 }} style={{ paddingLeft: 12 }}>
-         <MultiOptionChooser options={options} initialKeysIn={handler.state.platforms ?? []} updateKeys={(keys) => {
+         <MultiOptionChooser id="platform" optionsIn={options} initialKeysIn={handler.state.platforms ?? []} updateKeys={(keys) => {
 
             handler.platforms.forEach(p => {
                const selected = keys.find(k => k === p);
@@ -282,7 +253,7 @@ const ConfigChooser: React.FC<{ handler: TestDataHandler }> = observer(({ handle
          <Label>Configurations</Label>
       </Stack>
       <Stack tokens={{ childrenGap: 6 }} style={{ paddingLeft: 12 }}>
-         <MultiOptionChooser options={options} initialKeysIn={handler.state.configurations ?? []} updateKeys={(keys) => {
+         <MultiOptionChooser id="config" optionsIn={options} initialKeysIn={handler.state.configurations ?? []} updateKeys={(keys) => {
 
             handler.configurations.forEach(p => {
                const selected = keys.find(k => k === p);
@@ -317,7 +288,7 @@ const TargetChooser: React.FC<{ handler: TestDataHandler }> = observer(({ handle
          <Label>Targets</Label>
       </Stack>
       <Stack tokens={{ childrenGap: 6 }} style={{ paddingLeft: 12 }}>
-         <MultiOptionChooser options={options} initialKeysIn={handler.state.targets ?? []} updateKeys={(keys) => {
+         <MultiOptionChooser id="target" optionsIn={options} initialKeysIn={handler.state.targets ?? []} updateKeys={(keys) => {
 
             handler.targets.forEach(p => {
                const selected = keys.find(k => k === p);
@@ -351,7 +322,7 @@ const RHIChooser: React.FC<{ handler: TestDataHandler }> = observer(({ handler }
          <Label>RHI</Label>
       </Stack>
       <Stack tokens={{ childrenGap: 6 }} style={{ paddingLeft: 12 }}>
-         <MultiOptionChooser options={options} initialKeysIn={handler.state.rhi ?? []} updateKeys={(keys) => {
+         <MultiOptionChooser id="rhi" optionsIn={options} initialKeysIn={handler.state.rhi ?? []} updateKeys={(keys) => {
             handler.rhi.forEach(p => {
                const selected = keys.find(k => k === p);
                if (!selected && crhi.has(p)) {
@@ -385,7 +356,7 @@ const VariationChooser: React.FC<{ handler: TestDataHandler }> = observer(({ han
          <Label>Variation</Label>
       </Stack>
       <Stack tokens={{ childrenGap: 6 }} style={{ paddingLeft: 12 }}>
-         <MultiOptionChooser options={options} initialKeysIn={handler.state.variation ?? []} updateKeys={(keys) => {
+         <MultiOptionChooser id="variation" optionsIn={options} initialKeysIn={handler.state.variation ?? []} updateKeys={(keys) => {
 
             handler.variation.forEach(p => {
                const selected = keys.find(k => k === p);
