Horde: Guard against exception in timeline
#skipci
#rnx

diff --git a/Engine/Source/Programs/Horde/HordeDashboard/src/components/AutomationView.tsx b/Engine/Source/Programs/Horde/HordeDashboard/src/components/AutomationView.tsx
--- a/Engine/Source/Programs/Horde/HordeDashboard/src/components/AutomationView.tsx
+++ b/Engine/Source/Programs/Horde/HordeDashboard/src/components/AutomationView.tsx
@@ -1067,7 +1067,7 @@ class TestGraphRenderer extends AutomationGraph {
 
 }
 
-const TestGraph: React.FC<{ testId: string, streamId: string, handler: TestDataHandler }> = observer(({ testId, streamId, handler }) => {
+const TestGraph: React.FC<{ testId: string, streamId: string, handler: TestDataHandler }> = ({ testId, streamId, handler }) => {
 
    const graph_container_id = `${testId}_${streamId}_automation_graph_container`;
 
@@ -1094,7 +1094,7 @@ const TestGraph: React.FC<{ testId: string, streamId: string, handler: TestDataH
       </Stack>
    </Stack>;
 
-})
+}
 
 const AutomationTestView: React.FC<{ test: GetTestResponse, handler: TestDataHandler }> = observer(({ test, handler }) => {
 
diff --git a/Engine/Source/Programs/Horde/HordeDashboard/src/components/jobDetailsV2/JobDetailTimeline.tsx b/Engine/Source/Programs/Horde/HordeDashboard/src/components/jobDetailsV2/JobDetailTimeline.tsx
--- a/Engine/Source/Programs/Horde/HordeDashboard/src/components/jobDetailsV2/JobDetailTimeline.tsx
+++ b/Engine/Source/Programs/Horde/HordeDashboard/src/components/jobDetailsV2/JobDetailTimeline.tsx
@@ -198,6 +198,11 @@ class TimelineDataView extends JobDataView {
 
       });
 
+      if (!batches.length) {
+         console.error("No batches to render in timeline");
+         return;
+      }
+
       batches = batches.sort((a, b) => {
          return a.utcWaitStart.getTime() - b.utcWaitStart.getTime();
       });
