Horde: Gracefully shutdown the agent before running any upgrade commands, preventing connections being abruptly terminated to the server.

- More generally, allow leases to set the session result and run a callback after graceful shutdown.
- Shutdown and restart commands use the same mechanism.

diff --git a/Engine/Source/Programs/Horde/Horde.Agent/Leases/Handlers/RestartHandler.cs b/Engine/Source/Programs/Horde/Horde.Agent/Leases/Handlers/RestartHandler.cs
--- a/Engine/Source/Programs/Horde/Horde.Agent/Leases/Handlers/RestartHandler.cs
+++ b/Engine/Source/Programs/Horde/Horde.Agent/Leases/Handlers/RestartHandler.cs
@@ -3,6 +3,7 @@
 using System.Threading;
 using System.Threading.Tasks;
 using Horde.Agent.Services;
+using Horde.Agent.Utility;
 using HordeCommon.Rpc.Tasks;
 using Microsoft.Extensions.Logging;
 
@@ -20,8 +21,9 @@ namespace Horde.Agent.Leases.Handlers
 		/// <inheritdoc/>
 		public override Task<LeaseResult> ExecuteAsync(ISession session, string leaseId, RestartTask task, CancellationToken cancellationToken)
 		{
-			_logger.LogInformation("Setting shutdown flag");
-			return Task.FromResult(new LeaseResult(SessionResult.Shutdown));
+			_logger.LogInformation("Scheduling restart task");
+			SessionResult result = new SessionResult((logger, ctx) => Shutdown.ExecuteAsync(true, logger, ctx));
+			return Task.FromResult(new LeaseResult(result));
 		}
 	}
 }
diff --git a/Engine/Source/Programs/Horde/Horde.Agent/Leases/Handlers/ShutdownHandler.cs b/Engine/Source/Programs/Horde/Horde.Agent/Leases/Handlers/ShutdownHandler.cs
--- a/Engine/Source/Programs/Horde/Horde.Agent/Leases/Handlers/ShutdownHandler.cs
+++ b/Engine/Source/Programs/Horde/Horde.Agent/Leases/Handlers/ShutdownHandler.cs
@@ -3,6 +3,7 @@
 using System.Threading;
 using System.Threading.Tasks;
 using Horde.Agent.Services;
+using Horde.Agent.Utility;
 using HordeCommon.Rpc.Tasks;
 using Microsoft.Extensions.Logging;
 
@@ -20,8 +21,9 @@ namespace Horde.Agent.Leases.Handlers
 		/// <inheritdoc/>
 		public override Task<LeaseResult> ExecuteAsync(ISession session, string leaseId, ShutdownTask task, CancellationToken cancellationToken)
 		{
-			_logger.LogInformation("Setting shutdown flag");
-			return Task.FromResult(new LeaseResult(SessionResult.Shutdown));
+			_logger.LogInformation("Scheduling shutdown task");
+			SessionResult result = new SessionResult((logger, ctx) => Shutdown.ExecuteAsync(false, logger, ctx));
+			return Task.FromResult(new LeaseResult(result));
 		}
 	}
 }
diff --git a/Engine/Source/Programs/Horde/Horde.Agent/Leases/Handlers/UpgradeHandler.cs b/Engine/Source/Programs/Horde/Horde.Agent/Leases/Handlers/UpgradeHandler.cs
--- a/Engine/Source/Programs/Horde/Horde.Agent/Leases/Handlers/UpgradeHandler.cs
+++ b/Engine/Source/Programs/Horde/Horde.Agent/Leases/Handlers/UpgradeHandler.cs
@@ -5,14 +5,12 @@ using System.Collections.Generic;
 using System.Diagnostics;
 using System.IO;
 using System.IO.Compression;
-using System.Net.Http;
 using System.Reflection;
 using System.Text;
 using System.Threading;
 using System.Threading.Tasks;
 using EpicGames.Core;
 using Grpc.Core;
-using Horde.Agent.Parser;
 using Horde.Agent.Services;
 using Horde.Agent.Utility;
 using HordeCommon.Rpc;
@@ -75,59 +73,66 @@ namespace Horde.Agent.Leases.Handlers
 				// Get the current process and assembly. This may be different if running through dotnet.exe rather than a native PE image.
 				FileReference assemblyFileName = new FileReference(Assembly.GetExecutingAssembly().Location);
 
-				// Spawn the other process
-				using (Process process = new Process())
-				{
-					StringBuilder arguments = new StringBuilder();
+				StringBuilder arguments = new StringBuilder();
 
-					DirectoryReference targetDir = assemblyFileName.Directory;
+				DirectoryReference targetDir = assemblyFileName.Directory;
 
-					// We were launched via an external application (presumably dotnet.exe). Do the same thing again.
-					FileReference newAssemblyFileName = FileReference.Combine(extractedDir, assemblyFileName.MakeRelativeTo(targetDir));
-					if (!FileReference.Exists(newAssemblyFileName))
-					{
-						logger.LogError("Unable to find {AgentExe} in extracted archive", newAssemblyFileName);
-						return LeaseResult.Failed;
-					}
+				// We were launched via an external application (presumably dotnet.exe). Do the same thing again.
+				FileReference newAssemblyFileName = FileReference.Combine(extractedDir, assemblyFileName.MakeRelativeTo(targetDir));
+				if (!FileReference.Exists(newAssemblyFileName))
+				{
+					logger.LogError("Unable to find {AgentExe} in extracted archive", newAssemblyFileName);
+					return LeaseResult.Failed;
+				}
 
-					process.StartInfo.FileName = "dotnet";
+				StringBuilder currentArguments = new StringBuilder();
+				foreach (string arg in Program.Args)
+				{
+					currentArguments.AppendArgument(arg);
+				}
 
-					StringBuilder currentArguments = new StringBuilder();
-					foreach (string arg in Program.Args)
-					{
-						currentArguments.AppendArgument(arg);
-					}
+				arguments.AppendArgument(newAssemblyFileName.FullName);
+				arguments.AppendArgument("Service");
+				arguments.AppendArgument("Upgrade");
+				arguments.AppendArgument("-ProcessId=", Environment.ProcessId.ToString());
+				arguments.AppendArgument("-TargetDir=", targetDir.FullName);
+				arguments.AppendArgument("-Arguments=", currentArguments.ToString());
 
-					arguments.AppendArgument(newAssemblyFileName.FullName);
-					arguments.AppendArgument("Service");
-					arguments.AppendArgument("Upgrade");
-					arguments.AppendArgument("-ProcessId=", Environment.ProcessId.ToString());
-					arguments.AppendArgument("-TargetDir=", targetDir.FullName);
-					arguments.AppendArgument("-Arguments=", currentArguments.ToString());
+				// Spawn the other process
+				SessionResult result = new SessionResult((logger, ctx) => RunUpgradeAsync("dotnet", arguments.ToString(), logger, ctx));
+				return new LeaseResult(result);
+			}
 
-					process.StartInfo.Arguments = arguments.ToString();
-					process.StartInfo.UseShellExecute = false;
-					process.EnableRaisingEvents = true;
+			return LeaseResult.Success;
+		}
 
-					StringBuilder launchCommand = new StringBuilder();
-					launchCommand.AppendArgument(process.StartInfo.FileName);
-					launchCommand.Append(' ');
-					launchCommand.Append(arguments);
-					logger.LogInformation("Launching: {Launch}", launchCommand.ToString());
+		static async Task RunUpgradeAsync(string fileName, string arguments, ILogger logger, CancellationToken cancellationToken)
+		{
+			using (Process process = new Process())
+			{
+				StringBuilder launchCommand = new StringBuilder();
+				launchCommand.AppendArgument(fileName);
+				launchCommand.Append(' ');
+				launchCommand.Append(arguments);
+				logger.LogInformation("Launching: {Launch}", launchCommand.ToString());
 
-					TaskCompletionSource<int> exitCodeSource = new TaskCompletionSource<int>();
-					process.Exited += (sender, args) => { exitCodeSource.SetResult(process.ExitCode); };
+				process.StartInfo.FileName = fileName;
+				process.StartInfo.Arguments = arguments.ToString();
+				process.StartInfo.UseShellExecute = false;
+				process.EnableRaisingEvents = true;
 
-					process.Start();
+				TaskCompletionSource<int> exitCodeSource = new TaskCompletionSource<int>();
+				process.Exited += (sender, args) => { exitCodeSource.SetResult(process.ExitCode); };
 
-					using (cancellationToken.Register(() => { exitCodeSource.SetResult(0); }))
-					{
-						await exitCodeSource.Task;
-					}
+				process.Start();
+
+				using (cancellationToken.Register(() => { exitCodeSource.SetResult(0); }))
+				{
+					await exitCodeSource.Task;
 				}
-			}
 
-			return LeaseResult.Success;
+				logger.LogInformation("Upgrade did not terminate initiating process; reconnecting to server...");
+			}
 		}
 
 		/// <summary>
diff --git a/Engine/Source/Programs/Horde/Horde.Agent/Leases/LeaseManager.cs b/Engine/Source/Programs/Horde/Horde.Agent/Leases/LeaseManager.cs
--- a/Engine/Source/Programs/Horde/Horde.Agent/Leases/LeaseManager.cs
+++ b/Engine/Source/Programs/Horde/Horde.Agent/Leases/LeaseManager.cs
@@ -24,32 +24,6 @@ namespace Horde.Agent.Leases
 {
 	using ByteString = Google.Protobuf.ByteString;
 
-	/// <summary>
-	/// Outcome from a session
-	/// </summary>
-	enum SessionResult
-	{
-		/// <summary>
-		/// Create a new session (server migration, etc...)
-		/// </summary>
-		Continue,
-
-		/// <summary>
-		/// Session completed normally, terminate the application.
-		/// </summary>
-		Terminate,
-
-		/// <summary>
-		/// Machine should be scheduled to restart
-		/// </summary>
-		Restart,
-
-		/// <summary>
-		/// Machine should be scheduled to shut down
-		/// </summary>
-		Shutdown,
-	}
-
 	/// <summary>
 	/// Implements the message handling loop for an agent. Runs asynchronously until disposed.
 	/// </summary>
@@ -114,7 +88,7 @@ namespace Horde.Agent.Leases
 		/// <summary>
 		/// Result from executing this session
 		/// </summary>
-		SessionResult _result = SessionResult.Continue;
+		SessionResult? _sessionResult;
 
 		/// <summary>
 		/// Task completion source used to trigger the background thread to update the leases. Must take a lock on LockObject before 
@@ -174,12 +148,12 @@ namespace Horde.Agent.Leases
 			catch (OperationCanceledException ex) when (stoppingToken.IsCancellationRequested)
 			{
 				_logger.LogInformation(ex, "Execution cancelled");
-				result = SessionResult.Terminate;
+				result = new SessionResult(SessionOutcome.Terminate);
 			}
 			catch (Exception ex)
 			{
 				_logger.LogError(ex, "Exception while executing session. Restarting.");
-				result = SessionResult.Continue;
+				result = new SessionResult(SessionOutcome.BackOff);
 			}
 
 			while (_activeLeases.Count > 0)
@@ -273,9 +247,9 @@ namespace Horde.Agent.Leases
 					stopping = true;
 				}
 
-				if (_result != SessionResult.Continue)
+				if (_sessionResult != null)
 				{
-					_logger.LogInformation("Session termination requested (result: {Result})", _result);
+					_logger.LogInformation("Session termination requested (result: {Result})", _sessionResult.Outcome);
 					stopping = true;
 				}
 
@@ -291,7 +265,7 @@ namespace Horde.Agent.Leases
 					{
 						updateSessionRequest.Leases.Add(new Lease(leaseInfo.Lease));
 					}
-					if (_result != SessionResult.Continue && _activeLeases.Count == 0)
+					if (_sessionResult != null && _activeLeases.Count == 0)
 					{
 						stopping = true;
 					}
@@ -352,7 +326,7 @@ namespace Horde.Agent.Leases
 							if (atLeastOneLeaseFinished && shutdownAfterFinishedLease)
 							{
 								_logger.LogInformation("At least one lease executed. Requesting shutdown.");
-								_result = SessionResult.Terminate;
+								_sessionResult = new SessionResult(SessionOutcome.Terminate);
 							}
 
 							// Remove any leases which have completed
@@ -385,10 +359,10 @@ namespace Horde.Agent.Leases
 						}
 
 						// If there's nothing still running and cancellation was requested, exit
-						if (_activeLeases.Count == 0 && updateSessionRequest.Status == AgentStatus.Stopping)
+						if (_activeLeases.Count == 0 && _sessionResult != null)
 						{
 							_logger.LogInformation("No leases are active. Agent is stopping.");
-							break;
+							return _sessionResult;
 						}
 					}
 				}
@@ -408,7 +382,6 @@ namespace Horde.Agent.Leases
 					await Task.Delay(TimeSpan.FromSeconds(10.0), stoppingToken);
 				}
 			}
-			return _result;
 		}
 
 		/// <summary>
@@ -537,10 +510,10 @@ namespace Horde.Agent.Leases
 				}
 				_logger.LogInformation("Transitioning lease {LeaseId} to {State}, outcome={Outcome}", leaseInfo.Lease.Id, leaseInfo.Lease.State, leaseInfo.Lease.Outcome);
 
-				if (result.SessionResult.HasValue && result.SessionResult.Value != SessionResult.Continue && _result == SessionResult.Continue)
+				if (result.SessionResult != null && _sessionResult == null)
 				{
-					_logger.LogInformation("Lease {LeaseId} is setting session result to {Result}", leaseInfo.Lease.Id, result.SessionResult.Value);
-					_result = result.SessionResult.Value;
+					_logger.LogInformation("Lease {LeaseId} is setting session result to {Result}", leaseInfo.Lease.Id, result.SessionResult.Outcome);
+					_sessionResult = result.SessionResult;
 				}
 			}
 			_updateLeasesEvent.Set();
diff --git a/Engine/Source/Programs/Horde/Horde.Agent/Leases/LeaseResult.cs b/Engine/Source/Programs/Horde/Horde.Agent/Leases/LeaseResult.cs
--- a/Engine/Source/Programs/Horde/Horde.Agent/Leases/LeaseResult.cs
+++ b/Engine/Source/Programs/Horde/Horde.Agent/Leases/LeaseResult.cs
@@ -1,5 +1,6 @@
 // Copyright Epic Games, Inc. All Rights Reserved.
 
+using Horde.Agent.Services;
 using HordeCommon;
 
 namespace Horde.Agent.Leases
diff --git a/Engine/Source/Programs/Horde/Horde.Agent/Services/SessionResult.cs b/Engine/Source/Programs/Horde/Horde.Agent/Services/SessionResult.cs
new file mode 100644
--- /dev/null
+++ b/Engine/Source/Programs/Horde/Horde.Agent/Services/SessionResult.cs
@@ -0,0 +1,63 @@
+// Copyright Epic Games, Inc. All Rights Reserved.
+
+using System;
+using System.Threading;
+using System.Threading.Tasks;
+using Microsoft.Extensions.Logging;
+
+namespace Horde.Agent.Services
+{
+	/// <summary>
+	/// Outcome from a session
+	/// </summary>
+	enum SessionOutcome
+	{
+		/// <summary>
+		/// Continue the session with a backoff
+		/// </summary>
+		BackOff,
+
+		/// <summary>
+		/// Session completed normally, terminate the application.
+		/// </summary>
+		Terminate,
+
+		/// <summary>
+		/// Runs the given callback then attempts to reconnect
+		/// </summary>
+		RunCallback,
+	}
+
+	/// <summary>
+	/// Result from executing a session
+	/// </summary>
+	class SessionResult
+	{
+		/// <summary>
+		/// The outcome code
+		/// </summary>
+		public SessionOutcome Outcome { get; }
+
+		/// <summary>
+		/// Callback for running upgrades
+		/// </summary>
+		public Func<ILogger, CancellationToken, Task>? CallbackAsync { get; }
+
+		/// <summary>
+		/// Constructor
+		/// </summary>
+		public SessionResult(SessionOutcome outcome)
+		{
+			Outcome = outcome;
+		}
+
+		/// <summary>
+		/// Constructor
+		/// </summary>
+		public SessionResult(Func<ILogger, CancellationToken, Task> callbackAsync)
+		{
+			Outcome = SessionOutcome.RunCallback;
+			CallbackAsync = callbackAsync;
+		}
+	}
+}
diff --git a/Engine/Source/Programs/Horde/Horde.Agent/Services/WorkerService.cs b/Engine/Source/Programs/Horde/Horde.Agent/Services/WorkerService.cs
--- a/Engine/Source/Programs/Horde/Horde.Agent/Services/WorkerService.cs
+++ b/Engine/Source/Programs/Horde/Horde.Agent/Services/WorkerService.cs
@@ -81,9 +81,10 @@ namespace Horde.Agent.Services
 			int failureCount = 0;
 			while (!stoppingToken.IsCancellationRequested)
 			{
-				SessionResult result = SessionResult.Continue;
+				SessionResult? result = null;
 
 				Stopwatch sessionTime = Stopwatch.StartNew();
+
 				await using (AsyncServiceScope scope = _serviceProvider.CreateAsyncScope())
 				{
 					try
@@ -120,31 +121,19 @@ namespace Horde.Agent.Services
 					}
 				}
 
-				if (result == SessionResult.Terminate)
-				{
-					break;
-				}
-				else if (result == SessionResult.Shutdown || result == SessionResult.Restart)
+				if (result != null)
 				{
-					bool restart = (result == SessionResult.Restart);
-					_logger.LogInformation("Initiating shutdown (restart={Restart})", restart);
-
-					if (Shutdown.InitiateShutdown(restart, _logger))
+					if (result.Outcome == SessionOutcome.BackOff)
 					{
-						for (int idx = 10; idx > 0; idx--)
-						{
-							_logger.LogInformation("Waiting for shutdown ({Count})", idx);
-							try
-							{
-								await Task.Delay(TimeSpan.FromSeconds(60.0), stoppingToken);
-								_logger.LogInformation("Shutdown aborted.");
-							}
-							catch (OperationCanceledException)
-							{
-								_logger.LogInformation("Agent is shutting down.");
-								return;
-							}
-						}
+						await Task.Delay(TimeSpan.FromSeconds(30.0), stoppingToken);
+					}
+					else if (result.Outcome == SessionOutcome.Terminate)
+					{
+						break;
+					}
+					else if (result.Outcome == SessionOutcome.RunCallback)
+					{
+						await result.CallbackAsync!(_logger, stoppingToken);
 					}
 				}
 
diff --git a/Engine/Source/Programs/Horde/Horde.Agent/Utility/Shutdown.cs b/Engine/Source/Programs/Horde/Horde.Agent/Utility/Shutdown.cs
--- a/Engine/Source/Programs/Horde/Horde.Agent/Utility/Shutdown.cs
+++ b/Engine/Source/Programs/Horde/Horde.Agent/Utility/Shutdown.cs
@@ -4,6 +4,8 @@ using System;
 using System.Diagnostics;
 using System.Diagnostics.CodeAnalysis;
 using System.Runtime.InteropServices;
+using System.Threading;
+using System.Threading.Tasks;
 using Microsoft.Extensions.Logging;
 
 namespace Horde.Agent.Utility
@@ -49,7 +51,7 @@ namespace Horde.Agent.Utility
 		[SuppressMessage("Style", "IDE1006:Naming Styles", Justification = "Native struct")]
 		[return: MarshalAs(UnmanagedType.Bool)]
 		static extern bool AdjustTokenPrivileges(IntPtr TokenHandle,
-		   [MarshalAs(UnmanagedType.Bool)]bool DisableAllPrivileges,
+		   [MarshalAs(UnmanagedType.Bool)] bool DisableAllPrivileges,
 		   ref TOKEN_PRIVILEGES NewState,
 		   uint Zero,
 		   IntPtr Null1,
@@ -123,11 +125,11 @@ namespace Horde.Agent.Utility
 			{
 				string shutdownArgs;
 				if (restartAfterShutdown)
-                {
+				{
 					shutdownArgs = "sudo shutdown -r +0 \"Horde Agent is restarting\"";
 				}
 				else
-                {
+				{
 					shutdownArgs = "sudo shutdown +0 \"Horde Agent is shutting down\"";
 				}
 
@@ -156,5 +158,34 @@ namespace Horde.Agent.Utility
 				return false;
 			}
 		}
+
+		/// <summary>
+		/// Attempt to initiate a shutdown and return if it failed
+		/// </summary>
+		/// <param name="restart">Whether to restart</param>
+		/// <param name="logger">Logger for output</param>
+		/// <param name="cancellationToken">Cancellation token for the operation</param>
+		public static async Task ExecuteAsync(bool restart, ILogger logger, CancellationToken cancellationToken)
+		{
+			logger.LogInformation("Initiating shutdown (restart={Restart})", restart);
+
+			if (Shutdown.InitiateShutdown(restart, logger))
+			{
+				for (int idx = 10; idx > 0; idx--)
+				{
+					logger.LogInformation("Waiting for shutdown ({Count})", idx);
+					try
+					{
+						await Task.Delay(TimeSpan.FromMinutes(2.0), cancellationToken);
+						logger.LogInformation("Shutdown aborted.");
+					}
+					catch (OperationCanceledException)
+					{
+						logger.LogInformation("Agent is shutting down.");
+						return;
+					}
+				}
+			}
+		}
 	}
 }
