Horde: Only clear *.ib_mon files if collecting has been enabled

diff --git a/Engine/Source/Programs/Horde/Horde.Agent/Execution/JobExecutor.cs b/Engine/Source/Programs/Horde/Horde.Agent/Execution/JobExecutor.cs
--- a/Engine/Source/Programs/Horde/Horde.Agent/Execution/JobExecutor.cs
+++ b/Engine/Source/Programs/Horde/Horde.Agent/Execution/JobExecutor.cs
@@ -714,7 +714,11 @@ namespace Horde.Agent.Execution
 				}
 			}
 
-			_xgeMetadataExtractor?.ClearLocalIbMonFiles();
+			if (_jobOptions.CollectIbMonFilesAsArtifacts is true)
+			{
+				_xgeMetadataExtractor?.ClearLocalIbMonFiles();	
+			}
+			
 			if (_jobOptions.UseNewTempStorage ?? false)
 			{
 				bool result = await ExecuteWithTempStorageAsync(step, workspaceDir, arguments.ToString(), useP4, logger, cancellationToken);
