Horde: Store tools in new storage backend.

Data is currently stored as a single file node; may be worth changing to a proper chunked tree in the future.

diff --git a/Engine/Source/Programs/Horde/Horde.Build.Tests/ToolTests.cs b/Engine/Source/Programs/Horde/Horde.Build.Tests/ToolTests.cs
--- a/Engine/Source/Programs/Horde/Horde.Build.Tests/ToolTests.cs
+++ b/Engine/Source/Programs/Horde/Horde.Build.Tests/ToolTests.cs
@@ -4,10 +4,11 @@ using System;
 using System.Collections.Generic;
 using System.IO;
 using System.Text;
+using System.Threading;
 using System.Threading.Tasks;
 using Horde.Build.Server;
+using Horde.Build.Storage;
 using Horde.Build.Tools;
-using Horde.Build.Utilities;
 using HordeCommon;
 using Microsoft.Extensions.DependencyInjection;
 using Microsoft.VisualStudio.TestTools.UnitTesting;
@@ -17,15 +18,25 @@ namespace Horde.Build.Tests
 	[TestClass]
     public class ToolTests : TestSetup
     {
+		ToolId _toolId = new ToolId("ugs");
+
+		public ToolTests()
+		{
+			GlobalConfig globalConfig = new GlobalConfig();
+			globalConfig.Storage.Backends.Clear();
+			globalConfig.Storage.Backends.Add(new BackendConfig { Id = new BackendId("tools-backend"), Type = StorageBackendType.Memory });
+			globalConfig.Storage.Namespaces.Clear();
+			globalConfig.Storage.Namespaces.Add(new NamespaceConfig { Id = Namespace.Tools, Backend = new BackendId("tools-backend") });
+			globalConfig.Tools.Add(new ToolConfig(_toolId) { Name = "UnrealGameSync", Description = "Tool for syncing content from source control" });
+			SetConfig(globalConfig);
+		}
+
         [TestMethod]
         public async Task AddTool()
         {
 			ToolCollection collection = ServiceProvider.GetRequiredService<ToolCollection>();
 
-			ToolId toolId = new ToolId("ugs");
-			SetConfig(new GlobalConfig { Tools = new List<ToolConfig> { new ToolConfig(toolId) { Name = "UnrealGameSync", Description = "Tool for syncing content from source control" } } });
-
-			ITool tool = Deref(await collection.GetAsync(toolId, GlobalConfig.CurrentValue));
+			ITool tool = Deref(await collection.GetAsync(_toolId, GlobalConfig.CurrentValue));
 			Assert.AreEqual(tool.Id, new ToolId("ugs"));
 			Assert.AreEqual(tool.Config.Name, "UnrealGameSync");
 			Assert.AreEqual(tool.Config.Description, "Tool for syncing content from source control");
@@ -41,10 +52,7 @@ namespace Horde.Build.Tests
 		{
 			ToolCollection collection = ServiceProvider.GetRequiredService<ToolCollection>();
 
-			ToolId toolId = new ToolId("ugs");
-			SetConfig(new GlobalConfig { Tools = new List<ToolConfig> { new ToolConfig(toolId) { Name = "UnrealGameSync", Description = "Tool for syncing content from source control" } } });
-
-			ITool tool = Deref(await collection.GetAsync(toolId, GlobalConfig.CurrentValue));
+			ITool tool = Deref(await collection.GetAsync(_toolId, GlobalConfig.CurrentValue));
 			Assert.AreEqual(new ToolId("ugs"), tool.Id);
 			Assert.AreEqual("UnrealGameSync", tool.Config.Name);
 			Assert.AreEqual("Tool for syncing content from source control", tool.Config.Description);
@@ -55,7 +63,7 @@ namespace Horde.Build.Tests
 			ToolDeploymentId deploymentId;
 			using (MemoryStream stream = new MemoryStream(deploymentData))
 			{
-				tool = Deref(await collection.CreateDeploymentAsync(tool, new ToolDeploymentConfig { Version = "1.0", Duration = TimeSpan.FromMinutes(5.0), CreatePaused = true }, stream, GlobalConfig.CurrentValue));
+				tool = Deref(await collection.CreateDeploymentAsync(tool, new ToolDeploymentConfig { Version = "1.0", Duration = TimeSpan.FromMinutes(5.0), CreatePaused = true }, stream, GlobalConfig.CurrentValue, CancellationToken.None));
 				Assert.AreEqual(1, tool.Deployments.Count);
 				Assert.IsNull(tool.Deployments[0].StartedAt);
 				deploymentId = tool.Deployments[^1].Id;
@@ -93,7 +101,7 @@ namespace Horde.Build.Tests
 			byte[] outputDeploymentData;
 			using (MemoryStream stream = new MemoryStream())
 			{
-				Stream dataStream = await collection.GetDeploymentPayloadAsync(tool, tool.Deployments[0]);
+				Stream dataStream = await collection.GetDeploymentPayloadAsync(tool, tool.Deployments[0], CancellationToken.None);
 				await dataStream.CopyToAsync(stream);
 				outputDeploymentData = stream.ToArray();
 			}
diff --git a/Engine/Source/Programs/Horde/Horde.Build/Storage/StorageConfig.cs b/Engine/Source/Programs/Horde/Horde.Build/Storage/StorageConfig.cs
--- a/Engine/Source/Programs/Horde/Horde.Build/Storage/StorageConfig.cs
+++ b/Engine/Source/Programs/Horde/Horde.Build/Storage/StorageConfig.cs
@@ -70,6 +70,7 @@ namespace Horde.Build.Storage
 		internal void PostLoad(GlobalConfig globalConfig)
 		{
 			// Create a lookup for backend configs
+			_backendLookup.Clear();
 			foreach (BackendConfig backendConfig in Backends)
 			{
 				_backendLookup.Add(backendConfig.Id, backendConfig);
@@ -83,6 +84,7 @@ namespace Horde.Build.Storage
 			}
 
 			// Validate the backend config for each namespace
+			_namespaceLookup.Clear();
 			foreach (NamespaceConfig namespaceConfig in Namespaces)
 			{
 				BackendConfig? backendConfig;
diff --git a/Engine/Source/Programs/Horde/Horde.Build/Tools/ITool.cs b/Engine/Source/Programs/Horde/Horde.Build/Tools/ITool.cs
--- a/Engine/Source/Programs/Horde/Horde.Build/Tools/ITool.cs
+++ b/Engine/Source/Programs/Horde/Horde.Build/Tools/ITool.cs
@@ -101,7 +101,7 @@ namespace Horde.Build.Tools
 		/// <summary>
 		/// Reference to this tool in Horde Storage.
 		/// </summary>
-		RefId RefId { get; }
+		RefName RefName { get; }
 	}
 
 	/// <summary>
diff --git a/Engine/Source/Programs/Horde/Horde.Build/Tools/ToolCollection.cs b/Engine/Source/Programs/Horde/Horde.Build/Tools/ToolCollection.cs
--- a/Engine/Source/Programs/Horde/Horde.Build/Tools/ToolCollection.cs
+++ b/Engine/Source/Programs/Horde/Horde.Build/Tools/ToolCollection.cs
@@ -2,13 +2,14 @@
 
 using EpicGames.Core;
 using EpicGames.Horde.Storage;
+using EpicGames.Horde.Storage.Nodes;
 using EpicGames.Serialization;
 using Horde.Build.Server;
 using Horde.Build.Storage;
 using Horde.Build.Utilities;
 using HordeCommon;
+using Microsoft.Extensions.Caching.Memory;
 using Microsoft.Extensions.Logging;
-using Microsoft.Extensions.Options;
 using MongoDB.Bson;
 using MongoDB.Bson.Serialization;
 using MongoDB.Bson.Serialization.Attributes;
@@ -18,6 +19,7 @@ using System;
 using System.Collections.Generic;
 using System.IO;
 using System.Linq;
+using System.Threading;
 using System.Threading.Tasks;
 
 namespace Horde.Build.Tools
@@ -89,7 +91,7 @@ namespace Horde.Build.Tools
 			public string MimeType { get; set; }
 
 			[BsonElement("ref")]
-			public RefId RefId { get; set; }
+			public RefName RefName { get; set; }
 
 			public ToolDeployment(ToolDeploymentId id)
 			{
@@ -98,13 +100,13 @@ namespace Horde.Build.Tools
 				MimeType = ToolDeploymentConfig.DefaultMimeType;
 			}
 
-			public ToolDeployment(ToolDeploymentId id, ToolDeploymentConfig options, RefId refId)
+			public ToolDeployment(ToolDeploymentId id, ToolDeploymentConfig options, RefName refName)
 			{
 				Id = id;
 				Version = options.Version;
 				Duration = options.Duration;
 				MimeType = options.MimeType;
-				RefId = refId;
+				RefName = refName;
 			}
 
 			public void UpdateTemporalState(DateTime utcNow)
@@ -157,9 +159,9 @@ namespace Horde.Build.Tools
 		}
 
 		private readonly VersionedCollection<ToolId, Tool> _tools;
-		private readonly ILegacyStorageClient _storage;
+		private readonly StorageService _storageService;
 		private readonly IClock _clock;
-		private readonly NamespaceId _namespaceId;
+		private readonly IMemoryCache _cache;
 		private readonly ILogger _logger;
 
 		private static readonly RedisKey s_baseKey = "tools/v1/";
@@ -169,12 +171,12 @@ namespace Horde.Build.Tools
 		/// <summary>
 		/// Constructor
 		/// </summary>
-		public ToolCollection(MongoService mongoService, RedisService redisService, ILegacyStorageClient storage, IClock clock, ILogger<ToolCollection> logger)
+		public ToolCollection(MongoService mongoService, RedisService redisService, StorageService storageService, IMemoryCache cache, IClock clock, ILogger<ToolCollection> logger)
 		{
 			_tools = new VersionedCollection<ToolId, Tool>(mongoService, "Tools", redisService, s_baseKey, s_types);
-			_storage = storage;
+			_storageService = storageService;
 			_clock = clock;
-			_namespaceId = Namespace.Tools;
+			_cache = cache;
 			_logger = logger;
 		}
 
@@ -231,23 +233,26 @@ namespace Horde.Build.Tools
 		/// <param name="options">Options for the new deployment</param>
 		/// <param name="stream">Stream containing the tool data</param>
 		/// <param name="globalConfig">The current configuration</param>
+		/// <param name="cancellationToken">Cancellation token for the operation</param>
 		/// <returns>Updated tool document, or null if it does not exist</returns>
-		public async Task<ITool?> CreateDeploymentAsync(ITool tool, ToolDeploymentConfig options, Stream stream, GlobalConfig globalConfig)
+		public async Task<ITool?> CreateDeploymentAsync(ITool tool, ToolDeploymentConfig options, Stream stream, GlobalConfig globalConfig, CancellationToken cancellationToken)
 		{
-			// Upload the tool data first
-			IoHash hash = await _storage.WriteBlobAsync(_namespaceId, stream);
+			ToolDeploymentId deploymentId = ToolDeploymentId.GenerateNewId();
+			RefName refName = new RefName($"{tool.Id}/{deploymentId}");
 
-			ToolDeploymentData deploymentData = new ToolDeploymentData();
-			deploymentData.Id = ToolDeploymentId.GenerateNewId();
-			deploymentData.Version = options.Version;
-			deploymentData.Data = new CbBinaryAttachment(hash);
+			// Upload the tool data first
+			IStorageClient client = await _storageService.GetClientAsync(Namespace.Tools, cancellationToken);
 
-			BucketId bucketId = GetBucket(tool.Id);
-			RefId refId = new RefId(deploymentData.Id.ToString());
-			await _storage.SetRefAsync(_namespaceId, bucketId, refId, deploymentData);
+			NodeHandle handle;
+			using (TreeWriter writer = new TreeWriter(client, refName))
+			{
+				FileNodeWriter nodeWriter = new FileNodeWriter(writer, new ChunkingOptions());
+				handle = await nodeWriter.CreateAsync(stream, cancellationToken);
+			}
+			await client.WriteRefTargetAsync(refName, handle, cancellationToken: cancellationToken);
 
 			// Create the new deployment object
-			ToolDeployment deployment = new ToolDeployment(deploymentData.Id, options, refId);
+			ToolDeployment deployment = new ToolDeployment(deploymentId, options, refName);
 
 			// Start the deployment
 			DateTime utcNow = _clock.UtcNow;
@@ -260,7 +265,7 @@ namespace Horde.Build.Tools
 			Tool? newTool = (Tool)tool;
 			for (; ; )
 			{
-				newTool = await TryAddDeploymentAsync(newTool, deployment);
+				newTool = await TryAddDeploymentAsync(newTool, deployment, cancellationToken);
 				if (newTool != null)
 				{
 					break;
@@ -278,7 +283,7 @@ namespace Horde.Build.Tools
 			return newTool;
 		}
 
-		async ValueTask<Tool?> TryAddDeploymentAsync(Tool tool, ToolDeployment deployment)
+		async ValueTask<Tool?> TryAddDeploymentAsync(Tool tool, ToolDeployment deployment, CancellationToken cancellationToken)
 		{
 			Tool? newTool = tool;
 
@@ -291,7 +296,9 @@ namespace Horde.Build.Tools
 				{
 					return null;
 				}
-				await _storage.DeleteRefAsync(_namespaceId, GetBucket(newTool.Id), tool.Deployments[0].RefId);
+
+				IStorageClient client = await _storageService.GetClientAsync(Namespace.Tools, cancellationToken);
+				await client.DeleteRefAsync(tool.Deployments[0].RefName, cancellationToken);
 			}
 
 			// Add the new deployment
@@ -318,8 +325,6 @@ namespace Horde.Build.Tools
 				return null;
 			}
 
-			BucketId bucketId = GetBucket(tool.Id);
-
 			ToolDeployment deployment = tool.Deployments[idx];
 			switch (action)
 			{
@@ -360,13 +365,28 @@ namespace Horde.Build.Tools
 		/// </summary>
 		/// <param name="tool">Identifier for the tool</param>
 		/// <param name="deployment">The deployment</param>
+		/// <param name="cancellationToken">Cancellation token for the operation</param>
 		/// <returns>Stream for the data</returns>
-		public async Task<Stream> GetDeploymentPayloadAsync(ITool tool, IToolDeployment deployment)
+		public async Task<Stream> GetDeploymentPayloadAsync(ITool tool, IToolDeployment deployment, CancellationToken cancellationToken)
 		{
-			ToolDeploymentData data = await _storage.GetRefAsync<ToolDeploymentData>(_namespaceId, GetBucket(tool.Id), deployment.RefId);
-			return await _storage.ReadBlobAsync(_namespaceId, data.Data);
-		}
+			_ = tool;
 
-		private static BucketId GetBucket(ToolId toolId) => new BucketId(toolId.ToString());
+			IStorageClient client = await _storageService.GetClientAsync(Namespace.Tools, cancellationToken);
+			TreeReader reader = new TreeReader(client, _cache, _logger);
+			FileNode node = await reader.ReadNodeAsync<FileNode>(deployment.RefName, DateTime.UtcNow - TimeSpan.FromDays(2.0), cancellationToken);
+
+			MemoryStream memoryStream = new MemoryStream();
+			try
+			{
+				await node.CopyToStreamAsync(reader, memoryStream, cancellationToken);
+				memoryStream.Position = 0;
+				return memoryStream;
+			}
+			catch
+			{
+				await memoryStream.DisposeAsync();
+				throw;
+			}
+		}
 	}
 }
diff --git a/Engine/Source/Programs/Horde/Horde.Build/Tools/ToolsController.cs b/Engine/Source/Programs/Horde/Horde.Build/Tools/ToolsController.cs
--- a/Engine/Source/Programs/Horde/Horde.Build/Tools/ToolsController.cs
+++ b/Engine/Source/Programs/Horde/Horde.Build/Tools/ToolsController.cs
@@ -15,6 +15,7 @@ using System;
 using System.Collections.Generic;
 using System.IO;
 using System.Linq;
+using System.Threading;
 using System.Threading.Tasks;
 
 namespace Horde.Build.Tools
@@ -80,8 +81,8 @@ namespace Horde.Build.Tools
 		/// <inheritdoc cref="IToolDeployment.MimeType"/>
 		public string MimeType => _deployment.MimeType;
 
-		/// <inheritdoc cref="IToolDeployment.RefId"/>
-		public RefId RefId => _deployment.RefId;
+		/// <inheritdoc cref="IToolDeployment.RefName"/>
+		public RefName RefName => _deployment.RefName;
 
 		/// <summary>
 		/// Constructor
@@ -153,9 +154,10 @@ namespace Horde.Build.Tools
 		/// <returns>Information about the registered agent</returns>
 		[HttpPost]
 		[Route("/api/v1/tools/{id}/deployments")]
-		public async Task<ActionResult<CreateDeploymentResponse>> CreateDeploymentAsync(ToolId id, [FromForm] ToolDeploymentConfig options, [FromForm] IFormFile file)
+		public async Task<ActionResult<CreateDeploymentResponse>> CreateDeploymentAsync(ToolId id, [FromForm] ToolDeploymentConfig options, [FromForm] IFormFile file, CancellationToken cancellationToken)
 		{
 			ITool? tool = await _toolCollection.GetAsync(id, _globalConfig.Value);
+
 			if (tool == null)
 			{
 				return NotFound(id);
@@ -167,7 +169,7 @@ namespace Horde.Build.Tools
 
 			using (Stream stream = file.OpenReadStream())
 			{
-				tool = await _toolCollection.CreateDeploymentAsync(tool, options, stream, _globalConfig.Value);
+				tool = await _toolCollection.CreateDeploymentAsync(tool, options, stream, _globalConfig.Value, cancellationToken);
 				if (tool == null)
 				{
 					return NotFound(id);
@@ -183,7 +185,7 @@ namespace Horde.Build.Tools
 		/// <returns>Information about the registered agent</returns>
 		[HttpGet]
 		[Route("/api/v1/tools/{id}/deployments/{deploymentId}")]
-		public async Task<ActionResult> GetDeploymentAsync(ToolId id, ToolDeploymentId deploymentId, [FromQuery] GetToolAction action = GetToolAction.Info)
+		public async Task<ActionResult> GetDeploymentAsync(ToolId id, ToolDeploymentId deploymentId, [FromQuery] GetToolAction action = GetToolAction.Info, CancellationToken cancellationToken = default)
 		{
 			ITool? tool = await _toolCollection.GetAsync(id, _globalConfig.Value);
 			if (tool == null)
@@ -201,10 +203,10 @@ namespace Horde.Build.Tools
 				return NotFound(id, deploymentId);
 			}
 
-			return await GetDeploymentResponseAsync(tool, deployment, action);
+			return await GetDeploymentResponseAsync(tool, deployment, action, cancellationToken);
 		}
 
-		private async Task<ActionResult> GetDeploymentResponseAsync(ITool tool, IToolDeployment deployment, GetToolAction action)
+		private async Task<ActionResult> GetDeploymentResponseAsync(ITool tool, IToolDeployment deployment, GetToolAction action, CancellationToken cancellationToken)
 		{
 			if (action == GetToolAction.Info)
 			{
@@ -213,7 +215,7 @@ namespace Horde.Build.Tools
 			}
 			else
 			{
-				Stream stream = await _toolCollection.GetDeploymentPayloadAsync(tool, deployment);
+				Stream stream = await _toolCollection.GetDeploymentPayloadAsync(tool, deployment, cancellationToken);
 				return new FileStreamResult(stream, "application/zip");
 			}
 		}
@@ -273,7 +275,7 @@ namespace Horde.Build.Tools
 		/// <returns>Information about the registered agent</returns>
 		[HttpGet]
 		[Route("/api/v1/tools/{id}")]
-		public async Task<ActionResult> GetToolAsync(ToolId id, GetToolAction action = GetToolAction.Info)
+		public async Task<ActionResult> GetToolAsync(ToolId id, GetToolAction action = GetToolAction.Info, CancellationToken cancellationToken = default)
 		{
 			ITool? tool = await _toolCollection.GetAsync(id, _globalConfig.Value);
 			if (tool == null)
@@ -296,7 +298,7 @@ namespace Horde.Build.Tools
 					return NotFound(LogEvent.Create(LogLevel.Error, "Tool {ToolId} does not currently have any deployments", id));
 				}
 
-				return await GetDeploymentPayloadAsync(tool, tool.Deployments[^1]);
+				return await GetDeploymentPayloadAsync(tool, tool.Deployments[^1], cancellationToken);
 			}
 		}
 
@@ -306,10 +308,11 @@ namespace Horde.Build.Tools
 		/// <param name="id">The tool identifier</param>
 		/// <param name="phase">Value indicating the client's preference for deployment to receive.</param>
 		/// <param name="action">Information about the returned deployment</param>
+		/// <param name="cancellationToken">Cancellation token for the operation</param>
 		/// <returns>Information about the registered agent</returns>
 		[HttpGet]
 		[Route("/api/v1/tools/{id}/deployments")]
-		public async Task<ActionResult> FindDeploymentAsync(ToolId id, [FromQuery] double phase = 0.0, [FromQuery] GetToolAction action = GetToolAction.Info)
+		public async Task<ActionResult> FindDeploymentAsync(ToolId id, [FromQuery] double phase = 0.0, [FromQuery] GetToolAction action = GetToolAction.Info, CancellationToken cancellationToken = default)
 		{
 			ITool? tool = await _toolCollection.GetAsync(id, _globalConfig.Value);
 			if (tool == null)
@@ -342,13 +345,13 @@ namespace Horde.Build.Tools
 			}
 			else
 			{
-				return await GetDeploymentPayloadAsync(tool, tool.Deployments[idx]);
+				return await GetDeploymentPayloadAsync(tool, tool.Deployments[idx], cancellationToken);
 			}
 		}
 
-		async Task<ActionResult> GetDeploymentPayloadAsync(ITool tool, IToolDeployment deployment)
+		async Task<ActionResult> GetDeploymentPayloadAsync(ITool tool, IToolDeployment deployment, CancellationToken cancellationToken)
 		{
-			Stream stream = await _toolCollection.GetDeploymentPayloadAsync(tool, deployment);
+			Stream stream = await _toolCollection.GetDeploymentPayloadAsync(tool, deployment, cancellationToken);
 			return new FileStreamResult(stream, deployment.MimeType);
 		}
 
