Horde: Ignore steps that haven't run before when trying to build a list of suspects for issues.

diff --git a/Engine/Source/Programs/Horde/Horde.Build/Issues/IssueService.cs b/Engine/Source/Programs/Horde/Horde.Build/Issues/IssueService.cs
--- a/Engine/Source/Programs/Horde/Horde.Build/Issues/IssueService.cs
+++ b/Engine/Source/Programs/Horde/Horde.Build/Issues/IssueService.cs
@@ -923,12 +923,16 @@ namespace Horde.Build.Issues
 				List<NewIssueSuspectData> newSuspects = new List<NewIssueSuspectData>();
 				if (spans.Count > 0)
 				{
-					HashSet<int> suspectChanges = new HashSet<int>(spans[0].Suspects.Select(x => x.OriginatingChange ?? x.Change));
-					for (int spanIdx = 1; spanIdx < spans.Count; spanIdx++)
+					List<IIssueSpan> spansWithSuspects = spans.Where(x => x.LastSuccess != null).ToList();
+					if (spansWithSuspects.Count > 0)
 					{
-						suspectChanges.IntersectWith(spans[spanIdx].Suspects.Select(x => x.OriginatingChange ?? x.Change));
+						HashSet<int> suspectChanges = new HashSet<int>(spansWithSuspects[0].Suspects.Select(x => x.OriginatingChange ?? x.Change));
+						for (int spanIdx = 1; spanIdx < spansWithSuspects.Count; spanIdx++)
+						{
+							suspectChanges.IntersectWith(spansWithSuspects[spanIdx].Suspects.Select(x => x.OriginatingChange ?? x.Change));
+						}
+						newSuspects = spansWithSuspects[0].Suspects.Where(x => suspectChanges.Contains(x.OriginatingChange ?? x.Change)).Select(x => new NewIssueSuspectData(x.AuthorId, x.OriginatingChange ?? x.Change)).ToList();
 					}
-					newSuspects = spans[0].Suspects.Where(x => suspectChanges.Contains(x.OriginatingChange ?? x.Change)).Select(x => new NewIssueSuspectData(x.AuthorId, x.OriginatingChange ?? x.Change)).ToList();
 				}
 
 				// Create the combined fingerprint for this issue
@@ -1226,7 +1230,7 @@ namespace Horde.Build.Issues
 			}
 			else
 			{
-				_logger.LogInformation("Natched to issue {IssueId}", issue.Id);
+				_logger.LogInformation("Matched to issue {IssueId}", issue.Id);
 			}
 			return issue;
 		}
