Horde: Allow specifying a range of required resources when requesting a compute lease, and return the allocated resources to the requestor.

diff --git a/Engine/Source/Programs/Horde/Horde.Build.Tests/TestSetup.cs b/Engine/Source/Programs/Horde/Horde.Build.Tests/TestSetup.cs
--- a/Engine/Source/Programs/Horde/Horde.Build.Tests/TestSetup.cs
+++ b/Engine/Source/Programs/Horde/Horde.Build.Tests/TestSetup.cs
@@ -101,7 +101,6 @@ namespace Horde.Build.Tests
 		public AgentSoftwareService AgentSoftwareService => ServiceProvider.GetRequiredService<AgentSoftwareService>();
 		public AgentService AgentService => ServiceProvider.GetRequiredService<AgentService>();
 		public ICommitService CommitService => ServiceProvider.GetRequiredService<ICommitService>();
-		public ComputeService ComputeService => ServiceProvider.GetRequiredService<ComputeService>();
 		public GlobalsService GlobalsService => ServiceProvider.GetRequiredService<GlobalsService>();
 		public MongoService MongoService => ServiceProvider.GetRequiredService<MongoService>();
 		public ITemplateCollection TemplateCollection => ServiceProvider.GetRequiredService<ITemplateCollection>();
@@ -243,7 +242,6 @@ namespace Horde.Build.Tests
 			services.AddSingleton<AgentSoftwareService>();
 			services.AddSingleton<FleetService>();
 			services.AddSingleton<ConsistencyService>();
-			services.AddSingleton<ComputeService>();
 			services.AddSingleton<RequestTrackerService>();
 			services.AddSingleton<GlobalsService>();
 			services.AddSingleton<CredentialService>();
diff --git a/Engine/Source/Programs/Horde/Horde.Build/Agents/Fleet/ComputeQueueAwsMetricStrategy.cs b/Engine/Source/Programs/Horde/Horde.Build/Agents/Fleet/ComputeQueueAwsMetricStrategy.cs
--- a/Engine/Source/Programs/Horde/Horde.Build/Agents/Fleet/ComputeQueueAwsMetricStrategy.cs
+++ b/Engine/Source/Programs/Horde/Horde.Build/Agents/Fleet/ComputeQueueAwsMetricStrategy.cs
@@ -45,7 +45,7 @@ namespace Horde.Build.Agents.Fleet
 
 	static class ComputeServiceExtensions
 	{
-		public static Task<int> GetNumQueuedTasksForPoolAsync(this ComputeService computeService, ClusterId clusterId, IPool pool)
+		public static Task<int> GetNumQueuedTasksForPoolAsync(this ComputeTaskSource computeTaskSource, ClusterId clusterId, IPool pool)
 		{
 			// Will need to reimplement this functionality in ComputeService if we want to use this strategy, but conditions on tasks
 			// may reference properties that are specific to an agent rather than a pool...
@@ -60,7 +60,7 @@ namespace Horde.Build.Agents.Fleet
 	public class ComputeQueueAwsMetricStrategy : IPoolSizeStrategy
 	{
 		private readonly IAmazonCloudWatch _cloudWatch;
-		private readonly ComputeService _computeService;
+		private readonly ComputeTaskSource _computeTaskSource;
 		private readonly ComputeQueueAwsMetricSettings _settings;
 		private readonly ILogger<ComputeQueueAwsMetricStrategy> _logger;
 
@@ -68,14 +68,14 @@ namespace Horde.Build.Agents.Fleet
 		/// Constructor
 		/// </summary>
 		/// <param name="cloudWatch"></param>
-		/// <param name="computeService"></param>
+		/// <param name="computeTaskSource"></param>
 		/// <param name="settings"></param>
 		/// <param name="logger"></param>
-		public ComputeQueueAwsMetricStrategy(IAmazonCloudWatch cloudWatch, ComputeService computeService,
+		public ComputeQueueAwsMetricStrategy(IAmazonCloudWatch cloudWatch, ComputeTaskSource computeTaskSource,
 			ComputeQueueAwsMetricSettings settings, ILogger<ComputeQueueAwsMetricStrategy> logger)
 		{
 			_cloudWatch = cloudWatch;
-			_computeService = computeService;
+			_computeTaskSource = computeTaskSource;
 			_settings = settings;
 			_logger = logger;
 		}
@@ -91,7 +91,7 @@ namespace Horde.Build.Agents.Fleet
 			Dictionary<string, List<MetricDatum>> metricsPerCloudWatchNamespace = new();
 			int numAgents = agents.Count;
 			int totalCpuCores = agents.Select(x => x.Resources.TryGetValue(KnownPropertyNames.LogicalCores, out int numCpuCores) ? numCpuCores : 0).Sum();
-			int numQueuedComputeTasks =	await _computeService.GetNumQueuedTasksForPoolAsync(new ClusterId(_settings.ComputeClusterId), pool);
+			int numQueuedComputeTasks =	await _computeTaskSource.GetNumQueuedTasksForPoolAsync(new ClusterId(_settings.ComputeClusterId), pool);
 
 			double numQueuedTasksPerAgent = numQueuedComputeTasks / (double)Math.Max(numAgents, 1);
 			double numQueuedTasksPerCores = numQueuedComputeTasks / (double)Math.Max(totalCpuCores, 1);
diff --git a/Engine/Source/Programs/Horde/Horde.Build/Agents/IAgent.cs b/Engine/Source/Programs/Horde/Horde.Build/Agents/IAgent.cs
--- a/Engine/Source/Programs/Horde/Horde.Build/Agents/IAgent.cs
+++ b/Engine/Source/Programs/Horde/Horde.Build/Agents/IAgent.cs
@@ -597,10 +597,11 @@ namespace Horde.Build.Agents
 		/// </summary>
 		/// <param name="agent">The agent to create a lease for</param>
 		/// <param name="requirements">Requirements for the lease</param>
+		/// <param name="assignedResources">Receives the allocated resources</param>
 		/// <returns>True if the new lease can be granted</returns>
-		public static bool MeetsRequirements(this IAgent agent, Requirements requirements)
+		public static bool MeetsRequirements(this IAgent agent, Requirements requirements, Dictionary<string, int> assignedResources)
 		{
-			return MeetsRequirements(agent, requirements.Condition, requirements.Resources, requirements.Exclusive);
+			return MeetsRequirements(agent, requirements.Condition, requirements.Resources, requirements.Exclusive, assignedResources);
 		}
 
 		/// <summary>
@@ -610,8 +611,9 @@ namespace Horde.Build.Agents
 		/// <param name="exclusive">Whether t</param>
 		/// <param name="condition">Condition to satisfy</param>
 		/// <param name="resources">Resources required to execute</param>
+		/// <param name="assignedResources">Resources allocated to the task</param>
 		/// <returns>True if the new lease can be granted</returns>
-		public static bool MeetsRequirements(this IAgent agent, Condition? condition, Dictionary<string, int>? resources, bool exclusive)
+		public static bool MeetsRequirements(this IAgent agent, Condition? condition, Dictionary<string, ResourceRequirements>? resources, bool exclusive, Dictionary<string, int> assignedResources)
 		{
 			if (!agent.Enabled || agent.Status != AgentStatus.Ok)
 			{
@@ -631,7 +633,7 @@ namespace Horde.Build.Agents
 			}
 			if (resources != null)
 			{
-				foreach ((string name, int count) in resources)
+				foreach ((string name, ResourceRequirements resourceRequirements) in resources)
 				{
 					int remainingCount;
 					if (!agent.Resources.TryGetValue(name, out remainingCount))
@@ -647,10 +649,21 @@ namespace Horde.Build.Agents
 							remainingCount -= leaseCount;
 						}
 					}
-					if (remainingCount < count)
+					if (remainingCount < resourceRequirements.Min)
 					{
 						return false;
 					}
+
+					int allocatedCount;
+					if (resourceRequirements.Max != null)
+					{
+						allocatedCount = Math.Min(resourceRequirements.Max.Value, remainingCount);
+					}
+					else
+					{
+						allocatedCount = resourceRequirements.Min;
+					}
+					assignedResources.Add(name, allocatedCount);
 				}
 			}
 			return true;
diff --git a/Engine/Source/Programs/Horde/Horde.Build/Compute/ComputeController.cs b/Engine/Source/Programs/Horde/Horde.Build/Compute/ComputeController.cs
--- a/Engine/Source/Programs/Horde/Horde.Build/Compute/ComputeController.cs
+++ b/Engine/Source/Programs/Horde/Horde.Build/Compute/ComputeController.cs
@@ -1,6 +1,7 @@
 // Copyright Epic Games, Inc. All Rights Reserved.
 
 using System;
+using System.Collections.Generic;
 using System.Net;
 using EpicGames.Core;
 using EpicGames.Horde.Compute;
@@ -53,6 +54,11 @@ namespace Horde.Build.Compute
 		/// AES IV for the channel, as a hex string
 		/// </summary>
 		public string AesIv { get; set; } = String.Empty;
+
+		/// <summary>
+		/// Resources assigned to this machine
+		/// </summary>
+		public Dictionary<string, int> Resources { get; set; } = new Dictionary<string, int>();
 	}
 
 	/// <summary>
diff --git a/Engine/Source/Programs/Horde/Horde.Build/Compute/ComputeService.cs b/Engine/Source/Programs/Horde/Horde.Build/Compute/ComputeService.cs
deleted file mode 100644
diff --git a/Engine/Source/Programs/Horde/Horde.Build/Compute/ComputeTaskSource.cs b/Engine/Source/Programs/Horde/Horde.Build/Compute/ComputeTaskSource.cs
--- a/Engine/Source/Programs/Horde/Horde.Build/Compute/ComputeTaskSource.cs
+++ b/Engine/Source/Programs/Horde/Horde.Build/Compute/ComputeTaskSource.cs
@@ -42,14 +42,20 @@ namespace Horde.Build.Compute
 		/// </summary>
 		public ComputeTask Task { get; }
 
+		/// <summary>
+		/// Resources assigned to this lease
+		/// </summary>
+		public Dictionary<string, int> AssignedResources { get; }
+
 		/// <summary>
 		/// Constructor
 		/// </summary>
-		public ComputeResource(IPAddress ip, int port, ComputeTask task)
+		public ComputeResource(IPAddress ip, int port, ComputeTask task, Dictionary<string, int> assignedResources)
 		{
 			Ip = ip;
 			Port = port;
 			Task = task;
+			AssignedResources = assignedResources;
 		}
 	}
 
@@ -63,7 +69,7 @@ namespace Horde.Build.Compute
 			public IAgent Agent { get; }
 			public IPAddress Ip { get; }
 			public int Port { get; }
-			public TaskCompletionSource<(ComputeTask, Requirements)?> Lease { get; } = new TaskCompletionSource<(ComputeTask, Requirements)?>(TaskCreationOptions.RunContinuationsAsynchronously);
+			public TaskCompletionSource<AgentLease?> Lease { get; } = new TaskCompletionSource<AgentLease?>(TaskCreationOptions.RunContinuationsAsynchronously);
 
 			public Waiter(IAgent agent, IPAddress ip, int port)
 			{
@@ -105,16 +111,21 @@ namespace Horde.Build.Compute
 				LinkedList<Waiter>? waiters;
 				if (_waiters.TryGetValue(clusterId, out waiters))
 				{
-					ComputeTask? task = null;
+					ComputeTask? computeTask = null;
 					for (LinkedListNode<Waiter>? node = waiters.First; node != null; node = node.Next)
 					{
-						if (node.Value.Agent.MeetsRequirements(requirements))
+						Dictionary<string, int> assignedResources = new Dictionary<string, int>(); 
+						if (node.Value.Agent.MeetsRequirements(requirements, assignedResources))
 						{
-							task ??= CreateComputeTask();
-							if (node.Value.Lease.TrySetResult((task, requirements)))
+							computeTask ??= CreateComputeTask();
+
+							byte[] payload = Any.Pack(computeTask).ToByteArray();
+							AgentLease lease = new AgentLease(LeaseId.GenerateNewId(), "Compute task", null, null, null, LeaseState.Pending, assignedResources, requirements.Exclusive, payload);
+
+							if (node.Value.Lease.TrySetResult(lease))
 							{
 								Waiter waiter = node.Value;
-								return new ComputeResource(waiter.Ip, waiter.Port, task);
+								return new ComputeResource(waiter.Ip, waiter.Port, computeTask, assignedResources);
 							}
 						}
 					}
@@ -187,15 +198,9 @@ namespace Horde.Build.Compute
 				{
 					using (IDisposable disposable = cancellationToken.Register(() => waiter.Lease.TrySetResult(null)))
 					{
-						(ComputeTask, Requirements)? result = await waiter.Lease.Task;
-						if (result != null)
+						AgentLease? lease = await waiter.Lease.Task;
+						if (lease != null)
 						{
-							(ComputeTask task, Requirements requirements) = result.Value;
-
-							string leaseName = "Compute task";
-							byte[] payload = Any.Pack(task).ToByteArray();
-
-							AgentLease lease = new AgentLease(LeaseId.GenerateNewId(), leaseName, null, null, null, LeaseState.Pending, requirements.Resources, requirements.Exclusive, payload);
 							_logger.LogInformation("Created compute lease for agent {AgentId}", agent.Id);
 							return lease;
 						}
diff --git a/Engine/Source/Programs/Horde/Horde.Build/Startup.cs b/Engine/Source/Programs/Horde/Horde.Build/Startup.cs
--- a/Engine/Source/Programs/Horde/Horde.Build/Startup.cs
+++ b/Engine/Source/Programs/Horde/Horde.Build/Startup.cs
@@ -695,7 +695,6 @@ namespace Horde.Build
 				services.AddHostedService<JobTaskSource>(provider => provider.GetRequiredService<JobTaskSource>());
 				services.AddSingleton<ConformTaskSource>();
 				services.AddHostedService<ConformTaskSource>(provider => provider.GetRequiredService<ConformTaskSource>());
-				services.AddSingleton<ComputeService>();
 				services.AddSingleton<ComputeTaskSource>();
 
 				services.AddSingleton<ITaskSource, UpgradeTaskSource>();
diff --git a/Engine/Source/Programs/Shared/EpicGames.Horde/Compute/Requirements.cs b/Engine/Source/Programs/Shared/EpicGames.Horde/Compute/Requirements.cs
--- a/Engine/Source/Programs/Shared/EpicGames.Horde/Compute/Requirements.cs
+++ b/Engine/Source/Programs/Shared/EpicGames.Horde/Compute/Requirements.cs
@@ -22,7 +22,7 @@ namespace EpicGames.Horde.Compute
 		/// Resources used by the process
 		/// </summary>
 		[CbField("r")]
-		public Dictionary<string, int> Resources { get; } = new Dictionary<string, int>();
+		public Dictionary<string, ResourceRequirements> Resources { get; } = new Dictionary<string, ResourceRequirements>();
 
 		/// <summary>
 		/// Whether we require exclusive access to the device
@@ -64,9 +64,9 @@ namespace EpicGames.Horde.Compute
 			{
 				list.Add($"\"{Condition}\"");
 			}
-			foreach ((string name, int count) in Resources)
+			foreach ((string name, ResourceRequirements allocation) in Resources)
 			{
-				list.Add($"{name}: {count}");
+				list.Add($"{name}: {allocation.Min}-{allocation.Max}");
 			}
 			if (Exclusive)
 			{
@@ -75,4 +75,22 @@ namespace EpicGames.Horde.Compute
 			return String.Join(", ", list);
 		}
 	}
+
+	/// <summary>
+	/// Specifies requirements for resource allocation
+	/// </summary>
+	public class ResourceRequirements
+	{
+		/// <summary>
+		/// Minimum allocation of the requested resource
+		/// </summary>
+		[CbField("min")]
+		public int Min { get; set; } = 1;
+
+		/// <summary>
+		/// Maximum allocation of the requested resource. Allocates as much as possible unless capped.
+		/// </summary>
+		[CbField("max")]
+		public int? Max { get; set; }
+	}
 }
