Horde: Prevent errors due to Slack notifications being too long. Cap number of jobs listed in autoscaling messages to 10.

diff --git a/Engine/Source/Programs/Horde/Horde.Build/Notifications/Sinks/SlackNotificationSink.cs b/Engine/Source/Programs/Horde/Horde.Build/Notifications/Sinks/SlackNotificationSink.cs
--- a/Engine/Source/Programs/Horde/Horde.Build/Notifications/Sinks/SlackNotificationSink.cs
+++ b/Engine/Source/Programs/Horde/Horde.Build/Notifications/Sinks/SlackNotificationSink.cs
@@ -396,11 +396,19 @@ namespace Horde.Build.Notifications.Sinks
 		
 		private async Task SendJobScheduledOnEmptyAutoScaledPoolMessageAsync(string recipient, List<JobScheduledNotification> notifications)
 		{
-			string jobIds = String.Join(", ", notifications.Select(x => x.JobId));
+			const int MaxItems = 10;
+			string jobIds = StringUtils.FormatList(notifications.Select(x => x.JobId.ToString()).ToArray(), MaxItems);
 				
 			StringBuilder sb = new();
-			foreach (JobScheduledNotification notification in notifications)
+			for(int idx = 0; idx < notifications.Count; idx++)
 			{
+				if (idx >= MaxItems && notifications.Count > MaxItems + 2)
+				{
+					sb.AppendLine($"...and {notifications.Count - idx} others.");
+					break;
+				}
+
+				JobScheduledNotification notification = notifications[idx];
 				string jobUrl = _settings.DashboardUrl + "/job/" + notification.JobId;
 				sb.AppendLine($"Job `{notification.JobName}` with ID <{jobUrl}|{notification.JobId}> in pool `{notification.PoolName}`");
 			}
