Horde: Make agent id selectable in history modal, fix link cursor on agent id in timeline tooltip
#rnx
#skipci

diff --git a/Engine/Source/Programs/Horde/HordeDashboard/src/components/HistoryModal.tsx b/Engine/Source/Programs/Horde/HordeDashboard/src/components/HistoryModal.tsx
--- a/Engine/Source/Programs/Horde/HordeDashboard/src/components/HistoryModal.tsx
+++ b/Engine/Source/Programs/Horde/HordeDashboard/src/components/HistoryModal.tsx
@@ -345,7 +345,7 @@ export const HistoryModal: React.FC<{ agentId: string | undefined, onDismiss: (.
          return <Modal isOpen={true} isBlocking={true} topOffsetFixed={true} styles={{ main: { padding: 16, width: 1200, hasBeenOpened: false, top: "80px", position: "absolute" } }} className={hordeClasses.modal}>
             <Stack tokens={{ childrenGap: 40 }} styles={{ root: { padding: 8 } }}>
                <Stack grow verticalAlign="center">
-                  <Text variant="mediumPlus" styles={{ root: { fontWeight: "unset", fontFamily: "Horde Open Sans SemiBold" } }}>Loading Agent {agentId}</Text>
+                  <Text variant="mediumPlus" styles={{ root: { fontWeight: "unset", fontFamily: "Horde Open Sans SemiBold", fontSize: "20px" } }}>Agent {agentId}</Text>
                </Stack>
                <Stack verticalAlign="center">
                   <Spinner size={SpinnerSize.large} />
@@ -785,11 +785,6 @@ export const HistoryModal: React.FC<{ agentId: string | undefined, onDismiss: (.
          modalProps={{
             isBlocking: false,
             topOffsetFixed: true,
-            dragOptions: {
-               closeMenuItemText: "Close",
-               moveMenuItemText: "Move",
-               menu: ContextualMenu
-            },
             styles: {
                main: {
                   hasBeenOpened: false,
@@ -813,7 +808,7 @@ export const HistoryModal: React.FC<{ agentId: string | undefined, onDismiss: (.
          dialogContentProps={{
             type: DialogType.close,
             onDismiss: onDismissHistoryModal,
-            title: state.selectedAgent?.name,
+            title: `Agent ${state.selectedAgent?.name}`,
          }}
       >
          <Stack>
diff --git a/Engine/Source/Programs/Horde/HordeDashboard/src/components/jobDetailsV2/JobDetailTimeline.tsx b/Engine/Source/Programs/Horde/HordeDashboard/src/components/jobDetailsV2/JobDetailTimeline.tsx
--- a/Engine/Source/Programs/Horde/HordeDashboard/src/components/jobDetailsV2/JobDetailTimeline.tsx
+++ b/Engine/Source/Programs/Horde/HordeDashboard/src/components/jobDetailsV2/JobDetailTimeline.tsx
@@ -754,7 +754,7 @@ const GraphTooltip: React.FC<{ dataView: TimelineDataView }> = observer(({ dataV
       }
 
       if (agent) {
-         valueElement = <button className="horde-link" style={{ borderWidth: 0, padding: 0, backgroundColor: modeColors.background }}>
+         valueElement = <button className="horde-link" style={{ borderWidth: 0, padding: 0, backgroundColor: modeColors.background, cursor: "pointer" }}>
             {valueElement}
          </button>
       }
