Horde: Ignore cancellation exceptions when stopping compute listener.

diff --git a/Engine/Source/Programs/Horde/Horde.Agent/Services/ComputeListenerService.cs b/Engine/Source/Programs/Horde/Horde.Agent/Services/ComputeListenerService.cs
--- a/Engine/Source/Programs/Horde/Horde.Agent/Services/ComputeListenerService.cs
+++ b/Engine/Source/Programs/Horde/Horde.Agent/Services/ComputeListenerService.cs
@@ -138,7 +138,7 @@ namespace Horde.Agent.Services
 		public async Task StopAsync(CancellationToken cancellationToken)
 		{
 			_cancellationSource.Cancel();
-			await _serverTask;
+			await _serverTask.IgnoreCanceledExceptionsAsync();
 		}
 
 		async Task RunTcpListenerAsync(IPAddress address, int port, CancellationToken cancellationToken)
