Horde: Prioritize agents with pending shutdowns when scaling out pools

Also ensures agents with a pending shutdown is not counted towards active agents.

diff --git a/Engine/Source/Programs/Horde/Horde.Build.Tests/Fleet/FleetServiceTest.cs b/Engine/Source/Programs/Horde/Horde.Build.Tests/Fleet/FleetServiceTest.cs
--- a/Engine/Source/Programs/Horde/Horde.Build.Tests/Fleet/FleetServiceTest.cs
+++ b/Engine/Source/Programs/Horde/Horde.Build.Tests/Fleet/FleetServiceTest.cs
@@ -52,6 +52,9 @@ namespace Horde.Build.Tests.Fleet
 
 	public class FakeFleetManager : IFleetManager
 	{
+		public const string ScaleOutMessage = "Pool scaled out by fake fleet manager";
+		public const string ScaleInMessage = "Pool scaled in by fake fleet manager";
+		
 		public int AgentsAddedCount { get; private set; }
 		public int AgentsRemovedCount { get; private set; }
 
@@ -60,14 +63,14 @@ namespace Horde.Build.Tests.Fleet
 		public Task<ScaleResult> ExpandPoolAsync(IPool pool, IReadOnlyList<IAgent> agents, int count, CancellationToken cancellationToken = default)
 		{
 			AgentsAddedCount += count;
-			ScaleResult scaleResult = new (FleetManagerOutcome.Success, count, 0, "Pool scaled out");
+			ScaleResult scaleResult = new (FleetManagerOutcome.Success, count, 0, ScaleOutMessage);
 			return Task.FromResult(ForceResult ?? scaleResult);
 		}
 
 		public Task<ScaleResult> ShrinkPoolAsync(IPool pool, IReadOnlyList<IAgent> agents, int count, CancellationToken cancellationToken = default)
 		{
 			AgentsRemovedCount += count;
-			ScaleResult scaleResult = new (FleetManagerOutcome.Success, 0, count, "Pool scaled in");
+			ScaleResult scaleResult = new (FleetManagerOutcome.Success, 0, count, ScaleInMessage);
 			return Task.FromResult(ForceResult ?? scaleResult);
 		}
 
@@ -123,7 +126,44 @@ namespace Horde.Build.Tests.Fleet
 			Assert.IsTrue(_fleetManagerSpy.ExpandAgents[0].Enabled);
 			Assert.IsTrue(_fleetManagerSpy.ExpandAgents[1].Enabled);
 		}
+
+		[TestMethod]
+		public async Task ScaleOutWithPendingShutdowns()
+		{
+			using FleetService service = GetFleetService(new FakeFleetManager());
+			IPool pool = await PoolService.CreatePoolAsync("testPool", null, true, 0, 0, sizeStrategy: PoolSizeStrategy.LeaseUtilization);
+			IAgent agent1 = await CreateAgentAsync(pool, true);
+			IAgent agent2 = await CreateAgentAsync(pool, true, requestShutdown: true);
+
+			PoolWithAgents poolWithAgents = (await service.GetPoolsWithAgentsAsync())[0];
+			Assert.AreEqual(1, poolWithAgents.Agents.Count);
+			ScaleResult result = await service.ScalePoolAsync(pool, poolWithAgents.Agents, new PoolSizeResult(poolWithAgents.Agents.Count, 2), CancellationToken.None);
+//			Assert.AreEqual(FleetManagerOutcome.Success, result.Outcome);
+			Assert.AreEqual("Scaled out by only cancelling shutdowns", result.Message);
+			Assert.AreEqual(1, result.AgentsAddedCount);
+
+			Assert.IsFalse((await AgentService.GetAgentAsync(agent2.Id))!.RequestShutdown);
+		}
 		
+		[TestMethod]
+		public async Task ScaleOutWithPendingShutdownsWithStoppedAgents()
+		{
+			using FleetService service = GetFleetService(new FakeFleetManager());
+			IPool pool = await PoolService.CreatePoolAsync("testPool", null, true, 0, 0, sizeStrategy: PoolSizeStrategy.LeaseUtilization);
+			IAgent agent1 = await CreateAgentAsync(pool, true);
+			IAgent agent2 = await CreateAgentAsync(pool, false);
+			IAgent agent3 = await CreateAgentAsync(pool, true, requestShutdown: true);
+
+			PoolWithAgents poolWithAgents = (await service.GetPoolsWithAgentsAsync())[0];
+			Assert.AreEqual(1, poolWithAgents.Agents.Count);
+			ScaleResult result = await service.ScalePoolAsync(pool, poolWithAgents.Agents, new PoolSizeResult(poolWithAgents.Agents.Count, 3), CancellationToken.None);
+			Assert.AreEqual(FleetManagerOutcome.Success, result.Outcome);
+			Assert.AreEqual(FakeFleetManager.ScaleOutMessage, result.Message);
+			Assert.AreEqual(2, result.AgentsAddedCount);
+
+			Assert.IsFalse((await AgentService.GetAgentAsync(agent3.Id))!.RequestShutdown);
+		}
+
 		[TestMethod]
 		public async Task ScaleOutCooldown()
 		{
diff --git a/Engine/Source/Programs/Horde/Horde.Build.Tests/TestSetup.cs b/Engine/Source/Programs/Horde/Horde.Build.Tests/TestSetup.cs
--- a/Engine/Source/Programs/Horde/Horde.Build.Tests/TestSetup.cs
+++ b/Engine/Source/Programs/Horde/Horde.Build.Tests/TestSetup.cs
@@ -348,10 +348,15 @@ namespace Horde.Build.Tests
 		}
 		
 		private static int s_agentIdCounter = 1;
-		public async Task<IAgent> CreateAgentAsync(IPool pool, bool enabled = true)
+		public async Task<IAgent> CreateAgentAsync(IPool pool, bool enabled = true, bool requestShutdown = false)
 		{
 			IAgent agent = await AgentService.CreateAgentAsync("TestAgent" + s_agentIdCounter++, enabled, new List<PoolId> { pool.Id });
 			agent = await AgentService.CreateSessionAsync(agent, AgentStatus.Ok, new List<string>(), new Dictionary<string, int>(), null);
+			if (requestShutdown)
+			{
+				await AgentCollection.TryUpdateSettingsAsync(agent, requestShutdown: true);
+			}
+			
 			return agent;
 		}
 
diff --git a/Engine/Source/Programs/Horde/Horde.Build/Agents/Fleet/FleetService.cs b/Engine/Source/Programs/Horde/Horde.Build/Agents/Fleet/FleetService.cs
--- a/Engine/Source/Programs/Horde/Horde.Build/Agents/Fleet/FleetService.cs
+++ b/Engine/Source/Programs/Horde/Horde.Build/Agents/Fleet/FleetService.cs
@@ -189,7 +189,7 @@ namespace Horde.Build.Agents.Fleet
 
 		internal async Task<List<PoolWithAgents>> GetPoolsWithAgentsAsync()
 		{
-			List<IAgent> agents = await _agentCollection.FindAsync(status: AgentStatus.Ok, enabled: true);
+			List<IAgent> agents = (await _agentCollection.FindAsync(status: AgentStatus.Ok, enabled: true)).Where(x => !x.RequestShutdown).ToList();
 			List<IAgent> GetAgentsInPool(PoolId poolId) => agents.FindAll(a => a.GetPools().Any(p => p == poolId));
 			List<IPool> pools = await _poolCollection.GetAsync();
 
@@ -265,7 +265,9 @@ namespace Horde.Build.Agents.Fleet
 					}
 					else
 					{
-						result = await fleetManager.ExpandPoolAsync(pool, agents, deltaAgentCount, cancellationToken);
+						result = await ExpandWithPendingShutdownsFirstAsync(pool, deltaAgentCount, (agentsToAdd) =>
+							fleetManager.ExpandPoolAsync(pool, agents, agentsToAdd, cancellationToken));
+						
 						scaleOutTime = _clock.UtcNow;
 					}
 				}
@@ -278,14 +280,14 @@ namespace Horde.Build.Agents.Fleet
 					else
 					{
 						result = await fleetManager.ShrinkPoolAsync(pool, agents, -deltaAgentCount, cancellationToken);
-						scaleInTime = _clock.UtcNow;	
+						scaleInTime = _clock.UtcNow;
 					}
 				}
 			}
 			catch (Exception ex)
 			{
 				_logger.LogInformation(ex, "Failed to scale {PoolName}", pool.Name);
-				return new ScaleResult(FleetManagerOutcome.Failure, 0, 0);
+				return new ScaleResult(FleetManagerOutcome.Failure, 0, 0, "Exception during scaling. See log.");
 			}
 
 			if (pool.LastScaleResult == null || !pool.LastScaleResult.Equals(result))
@@ -402,6 +404,50 @@ namespace Horde.Build.Agents.Fleet
 			}
 		}
 
+		/// <summary>
+		/// </summary>
+		
+		
+		/// <summary>
+		/// Cancel a number of pending agent shutdowns for a pool
+		/// 
+		/// By preventing a shutdown, the agent can immediately be put back
+		/// to service without requiring the provisioning of a new one
+		/// </summary>
+		/// <param name="pool">Pool to cancel shutdowns in</param>
+		/// <param name="count">Number of shutdowns to cancel</param>
+		/// <returns>Number of pending shutdown cancelled</returns>
+		private async Task<int> CancelPendingShutdownsAsync(IPool pool, int count)
+		{
+			List<IAgent> agents = await _agentCollection.FindAsync(status: AgentStatus.Ok, enabled: true, poolId: pool.Id);
+			int numShutdownsCancelled = 0;
+			foreach (IAgent agent in agents)
+			{
+				if (agent.RequestShutdown && numShutdownsCancelled < count)
+				{
+					await _agentCollection.TryUpdateSettingsAsync(agent, requestShutdown: false);
+					numShutdownsCancelled++;
+				}
+			}
+			
+			return numShutdownsCancelled;
+		}
+
+		private async Task<ScaleResult> ExpandWithPendingShutdownsFirstAsync(IPool pool, int agentsToAdd, Func<int, Task<ScaleResult>> scaleOutFunc)
+		{
+			int numShutdownsCancelled = await CancelPendingShutdownsAsync(pool, agentsToAdd);
+			agentsToAdd -= numShutdownsCancelled;
+			GlobalTracer.Instance.ActiveSpan?.SetTag("NumShutdownsCancelled", numShutdownsCancelled);
+
+			if (agentsToAdd <= 0)
+			{
+				return new ScaleResult(FleetManagerOutcome.Success, numShutdownsCancelled, 0, "Scaled out by only cancelling shutdowns");
+			}
+			
+			ScaleResult result = await scaleOutFunc(agentsToAdd);
+			return new ScaleResult(result.Outcome, result.AgentsAddedCount + numShutdownsCancelled, result.AgentsRemovedCount, result.Message);
+		}
+
 		private static T DeserializeConfig<T>(string json)
 		{
 			json = String.IsNullOrEmpty(json) ? "{}" : json;
diff --git a/Engine/Source/Programs/Horde/Horde.Build/Agents/IAgentCollection.cs b/Engine/Source/Programs/Horde/Horde.Build/Agents/IAgentCollection.cs
--- a/Engine/Source/Programs/Horde/Horde.Build/Agents/IAgentCollection.cs
+++ b/Engine/Source/Programs/Horde/Horde.Build/Agents/IAgentCollection.cs
@@ -78,7 +78,7 @@ namespace Horde.Build.Agents
 		/// <param name="shutdownReason">The reason for shutting down agent, ex. Autoscaler/Manual/Unexpected</param>
 		/// <param name="pools">List of pools for the agent</param>
 		/// <param name="comment">New comment</param>
-		/// <returns>Version of the software that needs to be installed on the agent. Null if the agent is running the correct version.</returns>
+		/// <returns>New agent state if update was successful</returns>
 		Task<IAgent?> TryUpdateSettingsAsync(IAgent agent, bool? enabled = null, bool? requestConform = null, bool? requestFullConform = null, bool? requestRestart = null, bool? requestShutdown = null, string? shutdownReason = null, List<PoolId>? pools = null, string? comment = null);
 
 		/// <summary>
