Horde: Prevent use of empty strings for block directory names. Speculative fix for temp storage extraction on Mac.

diff --git a/Engine/Source/Programs/Horde/Horde.Agent/Utility/TempStorage.cs b/Engine/Source/Programs/Horde/Horde.Agent/Utility/TempStorage.cs
--- a/Engine/Source/Programs/Horde/Horde.Agent/Utility/TempStorage.cs
+++ b/Engine/Source/Programs/Horde/Horde.Agent/Utility/TempStorage.cs
@@ -607,6 +607,23 @@ namespace Horde.Storage.Utility
 			return new FileEntry(localFileListLocation.GetFileName(), FileEntryFlags.None, fileNodeWriter.Length, handle);
 		}
 
+		/// <summary>
+		/// Gets the directory name for storing data from a block
+		/// </summary>
+		/// <param name="blockName"></param>
+		/// <returns></returns>
+		static string GetBlockDirectoryName(string blockName)
+		{
+			if (String.IsNullOrEmpty(blockName))
+			{
+				return "block";
+			}
+			else
+			{
+				return $"block-{blockName}";
+			}
+		}
+
 		/// <summary>
 		/// Saves the given files (that should be rooted at the branch root) to a shared temp storage manifest with the given temp storage node and game.
 		/// </summary>
@@ -621,7 +638,8 @@ namespace Horde.Storage.Utility
 		/// <returns>The created manifest instance (which has already been saved to disk).</returns>
 		public static async Task<DirectoryEntry> ArchiveBlockAsync(DirectoryReference manifestDir, string nodeName, string blockName, DirectoryReference workspaceDir, IEnumerable<FileReference> files, TreeWriter writer, ILogger logger, CancellationToken cancellationToken)
 		{
-			logger.LogInformation("Creating output block \"{BlockName}\"", blockName);
+			string blockDirectoryName = GetBlockDirectoryName(blockName);
+			logger.LogInformation("Creating output block \"{BlockName}\" (as {DirectoryName})", blockName, blockDirectoryName);
 
 			// Create a manifest for the given build products
 			FileInfo[] fileInfos = files.Select(x => new FileInfo(x.FullName)).ToArray();
@@ -634,7 +652,7 @@ namespace Horde.Storage.Utility
 			DirectoryNode rootNode = new DirectoryNode();
 			await rootNode.CopyFilesAsync(workspaceDir, Enumerable.Concat(files, new[] { manifestLocation }), new ChunkingOptions(), writer, new CopyStatsLogger(logger), cancellationToken);
 
-			return new DirectoryEntry(blockName, rootNode);
+			return new DirectoryEntry(blockDirectoryName, rootNode);
 		}
 
 		/// <summary>
@@ -664,14 +682,16 @@ namespace Horde.Storage.Utility
 			}
 			else
 			{
+				string blockDirectoryName = GetBlockDirectoryName(blockName);
+
 				// Read the shared manifest
 				RefName refName = GetRefNameForNode(refPrefix, nodeName);
-				logger.LogInformation("Reading node \"{NodeName}\" block \"{BlockName}\" from temp storage (ref: {RefName}, local: {LocalFile})", nodeName, blockName, refName, localManifestFile);
+				logger.LogInformation("Reading node \"{NodeName}\" block \"{BlockName}\" from temp storage (ref: {RefName}, local: {LocalFile}, blockdir: {BlockDir})", nodeName, blockName, refName, localManifestFile, blockDirectoryName);
 
 				DirectoryNode node = await reader.ReadNodeAsync<DirectoryNode>(refName, cancellationToken: cancellationToken);
 
 				DirectoryEntry? rootDirEntry;
-				if (!node.TryGetDirectoryEntry(blockName, out rootDirEntry))
+				if (!node.TryGetDirectoryEntry(blockDirectoryName, out rootDirEntry))
 				{
 					throw new TempStorageException($"Missing block \"{blockName}\" from node \"{nodeName}\"");
 				}
