Horde: Update changelist metadata whenever a change-commit trigger callback is received from the Perforce server. Should allow fixing issues with Swarm/Robomerge changing commit author after submission.

* Modified route to include cluster name (api/v1/perforce/{cluster}/trigger/{type}).
* Refreshing a change just adds it back into the list of changes that need to be queried as part of the regular commit update. Existing changes are replaced.

diff --git a/Engine/Source/Programs/Horde/Horde.Build.Tests/Stubs/Services/PerforceServiceStub.cs b/Engine/Source/Programs/Horde/Horde.Build.Tests/Stubs/Services/PerforceServiceStub.cs
--- a/Engine/Source/Programs/Horde/Horde.Build.Tests/Stubs/Services/PerforceServiceStub.cs
+++ b/Engine/Source/Programs/Horde/Horde.Build.Tests/Stubs/Services/PerforceServiceStub.cs
@@ -11,7 +11,6 @@ using EpicGames.Perforce;
 using Horde.Build.Perforce;
 using Horde.Build.Streams;
 using Horde.Build.Users;
-using Horde.Build.Utilities;
 using Microsoft.VisualStudio.TestTools.UnitTesting;
 
 namespace Horde.Build.Tests.Stubs.Services
@@ -193,7 +192,7 @@ namespace Horde.Build.Tests.Stubs.Services
 			return Task.CompletedTask;
 		}
 
-		public Task UpdateChangelistDescription(string clusterName, int change, Func<string, string> description, CancellationToken cancellationToken)
+		public Task UpdateChangelistDescriptionAsync(string clusterName, int change, Func<string, string> description, CancellationToken cancellationToken)
 		{
 			return Task.CompletedTask;
 		}
@@ -262,5 +261,10 @@ namespace Horde.Build.Tests.Stubs.Services
 		{
 			return new CommitCollection(this, streamConfig);
 		}
+
+		public Task RefreshCachedCommitAsync(string clusterName, int change)
+		{
+			throw new NotImplementedException();
+		}
 	}
 }
diff --git a/Engine/Source/Programs/Horde/Horde.Build/Jobs/JobService.cs b/Engine/Source/Programs/Horde/Horde.Build/Jobs/JobService.cs
--- a/Engine/Source/Programs/Horde/Horde.Build/Jobs/JobService.cs
+++ b/Engine/Source/Programs/Horde/Horde.Build/Jobs/JobService.cs
@@ -1140,7 +1140,7 @@ namespace Horde.Build.Jobs
 
 				_logger.LogInformation("Updating description for {ClonedPreflightChange}", clonedPreflightChange);
 
-				await _perforceService.UpdateChangelistDescription(streamConfig.ClusterName, clonedPreflightChange, x => x.TrimEnd() + $"\n#preflight {job.Id}");
+				await _perforceService.UpdateChangelistDescriptionAsync(streamConfig.ClusterName, clonedPreflightChange, x => x.TrimEnd() + $"\n#preflight {job.Id}");
 
 				_logger.LogInformation("Submitting change {Change} (through {ChangeCopy}) after successful completion of {JobId}", job.PreflightChange, clonedPreflightChange, job.Id);
 				(change, message) = await _perforceService.SubmitShelvedChangeAsync(streamConfig, clonedPreflightChange, job.PreflightChange);
diff --git a/Engine/Source/Programs/Horde/Horde.Build/Perforce/IPerforceService.cs b/Engine/Source/Programs/Horde/Horde.Build/Perforce/IPerforceService.cs
--- a/Engine/Source/Programs/Horde/Horde.Build/Perforce/IPerforceService.cs
+++ b/Engine/Source/Programs/Horde/Horde.Build/Perforce/IPerforceService.cs
@@ -197,7 +197,14 @@ namespace Horde.Build.Perforce
 		/// <param name="updateFunc">Function to update the description</param>
 		/// <param name="cancellationToken">Cancellation token for the operation</param>
 		/// <returns>Async task</returns>
-		public Task UpdateChangelistDescription(string clusterName, int change, Func<string, string> updateFunc, CancellationToken cancellationToken = default);
+		public Task UpdateChangelistDescriptionAsync(string clusterName, int change, Func<string, string> updateFunc, CancellationToken cancellationToken = default);
+
+		/// <summary>
+		/// Refreshes any cached metadata for the given change
+		/// </summary>
+		/// <param name="clusterName">Name of the Perforce cluster</param>
+		/// <param name="change">Change to update</param>
+		public Task RefreshCachedCommitAsync(string clusterName, int change);
 
 		/// <summary>
 		/// Creates a commit source for the given stream
diff --git a/Engine/Source/Programs/Horde/Horde.Build/Perforce/PerforceController.cs b/Engine/Source/Programs/Horde/Horde.Build/Perforce/PerforceController.cs
--- a/Engine/Source/Programs/Horde/Horde.Build/Perforce/PerforceController.cs
+++ b/Engine/Source/Programs/Horde/Horde.Build/Perforce/PerforceController.cs
@@ -76,13 +76,15 @@ namespace Horde.Build.Perforce
 	[Route("[controller]")]
 	public class PublicPerforceController : ControllerBase
 	{
+		private readonly IPerforceService _perforceService;
 		private readonly ILogger<PerforceController> _logger;
 
 		/// <summary>
 		/// Constructor
 		/// </summary>
-		public PublicPerforceController(ILogger<PerforceController> logger)
+		public PublicPerforceController(IPerforceService perforceService, ILogger<PerforceController> logger)
 		{
+			_perforceService = perforceService;
 			_logger = logger;
 		}
 
@@ -91,11 +93,17 @@ namespace Horde.Build.Perforce
 		/// </summary>
 		/// <returns>200 OK on success</returns>
 		[HttpPost]
-		[Route("/api/v1/perforce/trigger/{type}")]
-		public ActionResult TriggerCallback(string type, [FromQuery] long? changelist = null, [FromQuery(Name = "user")] string? perforceUser = null)
+		[Route("/api/v1/perforce/{cluster}/trigger/{type}")]
+		public async Task<ActionResult> TriggerCallback(string cluster, string type, [FromQuery] long? changelist = null, [FromQuery(Name = "user")] string? perforceUser = null)
 		{
 			// Currently just a placeholder until correct triggers are in place.
-			_logger.LogDebug("Received Perforce trigger callback. Type={Type} Changelist={Changelist} User={User}", type, changelist, perforceUser);
+			_logger.LogDebug("Received Perforce trigger callback. Cluster={Cluster} Type={Type} Changelist={Changelist} User={User}", cluster, type, changelist, perforceUser);
+
+			if (type == "change-commit" && changelist != null)
+			{
+				await _perforceService.RefreshCachedCommitAsync(cluster, (int)changelist.Value);
+			}
+
 			string content = "{\"message\": \"Trigger received\"}";
 			return new ContentResult { ContentType = "text/plain", StatusCode = (int)HttpStatusCode.OK, Content = content };
 		}
diff --git a/Engine/Source/Programs/Horde/Horde.Build/Perforce/PerforceService.cs b/Engine/Source/Programs/Horde/Horde.Build/Perforce/PerforceService.cs
--- a/Engine/Source/Programs/Horde/Horde.Build/Perforce/PerforceService.cs
+++ b/Engine/Source/Programs/Horde/Horde.Build/Perforce/PerforceService.cs
@@ -854,7 +854,7 @@ namespace Horde.Build.Perforce
 		}
 
 		/// <inheritdoc/>
-		public async Task UpdateChangelistDescription(string clusterName, int change, Func<string, string> updateFunc, CancellationToken cancellationToken)
+		public async Task UpdateChangelistDescriptionAsync(string clusterName, int change, Func<string, string> updateFunc, CancellationToken cancellationToken)
 		{
 			using IScope scope = GlobalTracer.Instance.BuildSpan("PerforceService.UpdateChangelistDescription").StartActive();
 			scope.Span.SetTag("ClusterName", clusterName);
@@ -881,6 +881,9 @@ namespace Horde.Build.Perforce
 			}
 		}
 
+		/// <inheritdoc/>
+		public virtual Task RefreshCachedCommitAsync(string clusterName, int change) => Task.CompletedTask;
+
 		static async Task ResetClientAsync(IPerforceConnection perforce, CancellationToken cancellationToken)
 		{
 			await perforce.RevertAsync(-1, null, RevertOptions.KeepWorkspaceFiles, FileSpecList.Any, cancellationToken);
diff --git a/Engine/Source/Programs/Horde/Horde.Build/Perforce/PerforceServiceCache.cs b/Engine/Source/Programs/Horde/Horde.Build/Perforce/PerforceServiceCache.cs
--- a/Engine/Source/Programs/Horde/Horde.Build/Perforce/PerforceServiceCache.cs
+++ b/Engine/Source/Programs/Horde/Horde.Build/Perforce/PerforceServiceCache.cs
@@ -10,7 +10,6 @@ using System.Threading.Tasks;
 using EpicGames.Core;
 using EpicGames.Perforce;
 using EpicGames.Redis;
-using Horde.Build.Projects;
 using Horde.Build.Server;
 using Horde.Build.Streams;
 using Horde.Build.Users;
@@ -393,6 +392,9 @@ namespace Horde.Build.Perforce
 
 			using (IPooledPerforceConnection perforce = await ConnectAsync(clusterName, null, cancellationToken))
 			{
+				// Remove any changes we need to update
+				int[] refreshNumbers = await _redisService.GetDatabase().SetPopAsync(GetRefreshSetKey(clusterName), 100);
+
 				// Get the changelist range to query
 				FileSpecList spec = FileSpecList.Any;
 				if (state.MaxChange > 0)
@@ -402,7 +404,12 @@ namespace Horde.Build.Perforce
 
 				// Find the changes within that range, and abort if there's nothing new
 				List<ChangesRecord> changes = await perforce.GetChangesAsync(ChangesOptions.None, MaxChanges, ChangeStatus.Submitted, spec, cancellationToken);
-				if (changes.Count == 0)
+
+				List<int> changeNumbers = new List<int>();
+				changeNumbers.AddRange(refreshNumbers.Where(x => x <= state.MaxChange));
+				changeNumbers.AddRange(changes.Select(x => x.Number));
+
+				if (changeNumbers.Count == 0)
 				{
 					return null;
 				}
@@ -425,7 +432,7 @@ namespace Horde.Build.Perforce
 				const int MaxFiles = 1000;
 
 				// Describe the changes and create records for them
-				List<DescribeRecord> describeRecords = await perforce.DescribeAsync(DescribeOptions.None, MaxFiles, changes.Select(x => x.Number).ToArray(), cancellationToken);
+				List<DescribeRecord> describeRecords = await perforce.DescribeAsync(DescribeOptions.None, MaxFiles, changeNumbers.ToArray(), cancellationToken);
 				foreach (StreamInfo streamInfo in streamInfos)
 				{
 					foreach (DescribeRecord describeRecord in describeRecords)
@@ -463,18 +470,21 @@ namespace Horde.Build.Perforce
 				}
 
 				// Update the cache state
-				Dictionary<StreamId, int> minChanges = new Dictionary<StreamId, int>(streamInfos.Count);
-				foreach (StreamInfo streamInfo in streamInfos)
+				if (changes.Count > 0)
 				{
-					int minChange;
-					if (reset || !state.MinChanges.TryGetValue(streamInfo.StreamConfig.Id, out minChange))
+					Dictionary<StreamId, int> minChanges = new Dictionary<StreamId, int>(streamInfos.Count);
+					foreach (StreamInfo streamInfo in streamInfos)
 					{
-						minChange = changes[^1].Number;
+						int minChange;
+						if (reset || !state.MinChanges.TryGetValue(streamInfo.StreamConfig.Id, out minChange))
+						{
+							minChange = changes[^1].Number;
+						}
+						minChanges.Add(streamInfo.StreamConfig.Id, minChange);
 					}
-					minChanges.Add(streamInfo.StreamConfig.Id, minChange);
+					state.MinChanges = minChanges;
+					state.MaxChange = changes[0].Number;
 				}
-				state.MinChanges = minChanges;
-				state.MaxChange = changes[0].Number;
 				return state;
 			}
 		}
@@ -485,6 +495,17 @@ namespace Horde.Build.Perforce
 			await _commits.ReplaceOneAsync(filter, commitDoc, new ReplaceOptions { IsUpsert = true }, cancellationToken);
 		}
 
+#pragma warning disable CA1308 // Normalize strings to uppercase
+		static RedisSetKey<int> GetRefreshSetKey(string clusterName) => new RedisSetKey<int>($"perforce/{clusterName.ToLowerInvariant()}/refresh");
+#pragma warning restore CA1308 // Normalize strings to uppercase
+
+		/// <inheritdoc/>
+		public override async Task RefreshCachedCommitAsync(string clusterName, int change)
+		{
+			RedisSetKey<int> refreshSetKey = GetRefreshSetKey(clusterName);
+			await _redisService.GetDatabase().SetAddAsync(refreshSetKey, change, CommandFlags.FireAndForget);
+		}
+
 		#endregion
 
 		#region ICommitSource implementation
