Horde: Update Perforce trigger endpoint to be compatible with incoming HTTP calls

diff --git a/Engine/Source/Programs/Horde/Horde.Build/Perforce/PerforceController.cs b/Engine/Source/Programs/Horde/Horde.Build/Perforce/PerforceController.cs
--- a/Engine/Source/Programs/Horde/Horde.Build/Perforce/PerforceController.cs
+++ b/Engine/Source/Programs/Horde/Horde.Build/Perforce/PerforceController.cs
@@ -1,5 +1,6 @@
 // Copyright Epic Games, Inc. All Rights Reserved.
 
+using System;
 using System.Collections.Generic;
 using System.Net;
 using System.Threading.Tasks;
@@ -69,6 +70,62 @@ namespace Horde.Build.Perforce
 		}
 	}
 
+	/// <summary>
+	/// Perforce trigger received as JSON
+	/// </summary>
+	public class PerforceTriggerPayload
+	{
+		/// <summary>
+		/// Type of trigger (change-commit etc)
+		/// </summary>
+		public string TriggerType { get; }
+		
+		/// <summary>
+		/// Author of changelist (Perforce user) 
+		/// </summary>
+		public string User { get; }
+		
+		/// <summary>
+		/// Change number
+		/// </summary>
+		public string Changelist { get; }
+		
+		/// <summary>
+		/// Change number
+		/// </summary>
+		public int ChangelistNumber { get; }
+		
+		/// <summary>
+		/// Change root of changelist
+		/// </summary>
+		public string ChangeRoot { get; }
+
+		/// <summary>
+		/// Constructor
+		/// </summary>
+		/// <param name="triggerType"></param>
+		/// <param name="user"></param>
+		/// <param name="changelist"></param>
+		/// <param name="changeRoot"></param>
+		public PerforceTriggerPayload(string triggerType, string user, string changelist, string changeRoot)
+		{
+			TriggerType = triggerType;
+			User = user;
+			Changelist = changelist;
+			ChangelistNumber = Convert.ToInt32(changelist);
+			ChangeRoot = changeRoot;
+		}
+
+		/// <summary>
+		/// Format as a string for debugging purposes
+		/// </summary>
+		/// <returns></returns>
+		public override string ToString()
+		{
+			return $"TriggerType={TriggerType} User={User} CL={Changelist} Root={ChangeRoot}";
+		}
+	}
+	
 	/// <summary>
 	/// Controller for Perforce triggers and callbacks
 	/// </summary>
@@ -93,19 +150,15 @@ namespace Horde.Build.Perforce
 		/// </summary>
 		/// <returns>200 OK on success</returns>
 		[HttpPost]
-		[Route("/api/v1/perforce/{cluster}/trigger/{type}")]
-		public async Task<ActionResult> TriggerCallback(string cluster, string type, [FromQuery] long? changelist = null, [FromQuery(Name = "user")] string? perforceUser = null)
+		[Route("/api/v1/perforce/trigger")]
+		public ActionResult TriggerCallback([FromBody]PerforceTriggerPayload payload)
 		{
 			// Currently just a placeholder until correct triggers are in place.
-			_logger.LogDebug("Received Perforce trigger callback. Cluster={Cluster} Type={Type} Changelist={Changelist} User={User}", cluster, type, changelist, perforceUser);
-
-			if (type == "change-commit" && changelist != null)
-			{
-				await _perforceService.RefreshCachedCommitAsync(cluster, (int)changelist.Value);
-			}
+			_logger.LogDebug("Received Perforce trigger callback. Type={Type} CL={Changelist} User={User} Root={Root}", 
+				payload.TriggerType, payload.Changelist, payload.User, payload.ChangeRoot);
 
 			string content = "{\"message\": \"Trigger received\"}";
-			return new ContentResult { ContentType = "text/plain", StatusCode = (int)HttpStatusCode.OK, Content = content };
+			return new ContentResult { ContentType = "application/json", StatusCode = (int)HttpStatusCode.OK, Content = content };
 		}
 	}
 }
