Horde: Don't show negative time when client and server don't agree, also fix warning
#rnx
#skipci

diff --git a/Engine/Source/Programs/Horde/HordeDashboard/src/base/utilities/timeUtils.ts b/Engine/Source/Programs/Horde/HordeDashboard/src/base/utilities/timeUtils.ts
--- a/Engine/Source/Programs/Horde/HordeDashboard/src/base/utilities/timeUtils.ts
+++ b/Engine/Source/Programs/Horde/HordeDashboard/src/base/utilities/timeUtils.ts
@@ -269,8 +269,10 @@ export const getElapsedString = (start: Moment, end: Moment, includeSeconds:bool
         }
 
         if (includeSeconds || !duration) {
-            if (d.seconds()) {
+            if (d.seconds() > 0) {
                 duration += `${d.seconds()}s `;
+            } else {
+                duration += `0s `;
             }
         }
     }
diff --git a/Engine/Source/Programs/Horde/HordeDashboard/src/components/HistoryModal.tsx b/Engine/Source/Programs/Horde/HordeDashboard/src/components/HistoryModal.tsx
--- a/Engine/Source/Programs/Horde/HordeDashboard/src/components/HistoryModal.tsx
+++ b/Engine/Source/Programs/Horde/HordeDashboard/src/components/HistoryModal.tsx
@@ -1,6 +1,6 @@
 // Copyright Epic Games, Inc. All Rights Reserved.
 
-import { ConstrainMode, ContextualMenu, DefaultButton, DetailsHeader, DetailsList, DetailsListLayoutMode, DetailsRow, Dialog, DialogFooter, DialogType, GroupedList, GroupHeader, IColumn, IContextualMenuItem, IContextualMenuProps, IDetailsHeaderProps, IDetailsHeaderStyles, IDetailsListProps, IGroup, ITooltipHostStyles, mergeStyleSets, Modal, Pivot, PivotItem, PrimaryButton, ScrollablePane, ScrollbarVisibility, Selection, SelectionMode, Spinner, SpinnerSize, Stack, Sticky, StickyPositionType, Text } from "@fluentui/react";
+import { ConstrainMode, DefaultButton, DetailsHeader, DetailsList, DetailsListLayoutMode, DetailsRow, Dialog, DialogFooter, DialogType, GroupedList, GroupHeader, IColumn, IContextualMenuItem, IContextualMenuProps, IDetailsHeaderProps, IDetailsHeaderStyles, IDetailsListProps, IGroup, ITooltipHostStyles, mergeStyleSets, Modal, Pivot, PivotItem, PrimaryButton, ScrollablePane, ScrollbarVisibility, Selection, SelectionMode, Spinner, SpinnerSize, Stack, Sticky, StickyPositionType, Text } from "@fluentui/react";
 import { action, makeObservable, observable } from "mobx";
 import { observer } from "mobx-react-lite";
 import React, { useState } from "react";
