Horde: Store tools as regular DirectoryNode/FileNode trees, but convert to/from zip files for upload and download. Allows for patching using bundle data where possible.

diff --git a/Engine/Source/Programs/Horde/Horde.Build.Tests/ToolTests.cs b/Engine/Source/Programs/Horde/Horde.Build.Tests/ToolTests.cs
--- a/Engine/Source/Programs/Horde/Horde.Build.Tests/ToolTests.cs
+++ b/Engine/Source/Programs/Horde/Horde.Build.Tests/ToolTests.cs
@@ -3,6 +3,8 @@
 using System;
 using System.Collections.Generic;
 using System.IO;
+using System.IO.Compression;
+using System.Linq;
 using System.Text;
 using System.Threading;
 using System.Threading.Tasks;
@@ -58,10 +60,25 @@ namespace Horde.Build.Tests
 			Assert.AreEqual("Tool for syncing content from source control", tool.Config.Description);
 			Assert.AreEqual(0, tool.Deployments.Count);
 
-			byte[] deploymentData = Encoding.UTF8.GetBytes("hello world");
+			const string FileName = "test.txt";
+			byte[] fileData = Encoding.UTF8.GetBytes("hello world");
+
+			byte[] zipData;
+			using (MemoryStream stream = new MemoryStream())
+			{
+				using (ZipArchive archive = new ZipArchive(stream, ZipArchiveMode.Create))
+				{
+					ZipArchiveEntry entry = archive.CreateEntry(FileName);
+					using (Stream entryStream = entry.Open())
+					{
+						await entryStream.WriteAsync(fileData);
+					}
+				}
+				zipData = stream.ToArray();
+			}
 
 			ToolDeploymentId deploymentId;
-			using (MemoryStream stream = new MemoryStream(deploymentData))
+			using (MemoryStream stream = new MemoryStream(zipData))
 			{
 				tool = Deref(await collection.CreateDeploymentAsync(tool, new ToolDeploymentConfig { Version = "1.0", Duration = TimeSpan.FromMinutes(5.0), CreatePaused = true }, stream, GlobalConfig.CurrentValue, CancellationToken.None));
 				Assert.AreEqual(1, tool.Deployments.Count);
@@ -98,14 +115,22 @@ namespace Horde.Build.Tests
 			Assert.IsTrue((tool.Deployments[0].Progress - 0.5) < 0.1);
 
 			// Get the deployment data
-			byte[] outputDeploymentData;
-			using (MemoryStream stream = new MemoryStream())
+			using Stream dataStream = await collection.GetDeploymentZipAsync(tool, tool.Deployments[0], CancellationToken.None);
+			using (ZipArchive archive = new ZipArchive(dataStream, ZipArchiveMode.Read))
 			{
-				Stream dataStream = await collection.GetDeploymentPayloadAsync(tool, tool.Deployments[0], CancellationToken.None);
-				await dataStream.CopyToAsync(stream);
-				outputDeploymentData = stream.ToArray();
+				ZipArchiveEntry entry = archive.Entries.First();
+				Assert.AreEqual(entry.Name, FileName);
+
+				byte[] outputFileData;
+				using (MemoryStream stream = new MemoryStream())
+				{
+					using Stream entryStream = entry.Open();
+					await entryStream.CopyToAsync(stream);
+					outputFileData = stream.ToArray();
+				}
+
+				Assert.IsTrue(fileData.AsSpan().SequenceEqual(outputFileData.AsSpan()));
 			}
-			Assert.IsTrue(deploymentData.AsSpan().SequenceEqual(outputDeploymentData));
 		}
 	}
 }
\ No newline at end of file
diff --git a/Engine/Source/Programs/Horde/Horde.Build/Tools/ToolCollection.cs b/Engine/Source/Programs/Horde/Horde.Build/Tools/ToolCollection.cs
--- a/Engine/Source/Programs/Horde/Horde.Build/Tools/ToolCollection.cs
+++ b/Engine/Source/Programs/Horde/Horde.Build/Tools/ToolCollection.cs
@@ -16,8 +16,10 @@ using MongoDB.Bson.Serialization.Attributes;
 using MongoDB.Driver;
 using StackExchange.Redis;
 using System;
+using System.Buffers;
 using System.Collections.Generic;
 using System.IO;
+using System.IO.Pipelines;
 using System.Linq;
 using System.Threading;
 using System.Threading.Tasks;
@@ -246,10 +248,10 @@ namespace Horde.Build.Tools
 			NodeHandle handle;
 			using (TreeWriter writer = new TreeWriter(client, refName))
 			{
-				FileNodeWriter nodeWriter = new FileNodeWriter(writer, new ChunkingOptions());
-				handle = await nodeWriter.CreateAsync(stream, cancellationToken);
+				DirectoryNode directoryNode = new DirectoryNode();
+				await directoryNode.CopyFromZipStreamAsync(stream, writer, cancellationToken);
+				handle = await writer.WriteAsync(refName, directoryNode, cancellationToken: cancellationToken);
 			}
-			await client.WriteRefTargetAsync(refName, handle, cancellationToken: cancellationToken);
 
 			// Create the new deployment object
 			ToolDeployment deployment = new ToolDeployment(deploymentId, options, refName);
@@ -367,26 +369,15 @@ namespace Horde.Build.Tools
 		/// <param name="deployment">The deployment</param>
 		/// <param name="cancellationToken">Cancellation token for the operation</param>
 		/// <returns>Stream for the data</returns>
-		public async Task<Stream> GetDeploymentPayloadAsync(ITool tool, IToolDeployment deployment, CancellationToken cancellationToken)
+		public async Task<Stream> GetDeploymentZipAsync(ITool tool, IToolDeployment deployment, CancellationToken cancellationToken)
 		{
 			_ = tool;
 
 			IStorageClient client = await _storageService.GetClientAsync(Namespace.Tools, cancellationToken);
 			TreeReader reader = new TreeReader(client, _cache, _logger);
-			FileNode node = await reader.ReadNodeAsync<FileNode>(deployment.RefName, DateTime.UtcNow - TimeSpan.FromDays(2.0), cancellationToken);
+			DirectoryNode node = await reader.ReadNodeAsync<DirectoryNode>(deployment.RefName, DateTime.UtcNow - TimeSpan.FromDays(2.0), cancellationToken);
 
-			MemoryStream memoryStream = new MemoryStream();
-			try
-			{
-				await node.CopyToStreamAsync(reader, memoryStream, cancellationToken);
-				memoryStream.Position = 0;
-				return memoryStream;
-			}
-			catch
-			{
-				await memoryStream.DisposeAsync();
-				throw;
-			}
+			return node.ToZipStream(reader);
 		}
 	}
 }
diff --git a/Engine/Source/Programs/Horde/Horde.Build/Tools/ToolsController.cs b/Engine/Source/Programs/Horde/Horde.Build/Tools/ToolsController.cs
--- a/Engine/Source/Programs/Horde/Horde.Build/Tools/ToolsController.cs
+++ b/Engine/Source/Programs/Horde/Horde.Build/Tools/ToolsController.cs
@@ -215,7 +215,7 @@ namespace Horde.Build.Tools
 			}
 			else
 			{
-				Stream stream = await _toolCollection.GetDeploymentPayloadAsync(tool, deployment, cancellationToken);
+				Stream stream = await _toolCollection.GetDeploymentZipAsync(tool, deployment, cancellationToken);
 				return new FileStreamResult(stream, "application/zip");
 			}
 		}
@@ -351,7 +351,7 @@ namespace Horde.Build.Tools
 
 		async Task<ActionResult> GetDeploymentPayloadAsync(ITool tool, IToolDeployment deployment, CancellationToken cancellationToken)
 		{
-			Stream stream = await _toolCollection.GetDeploymentPayloadAsync(tool, deployment, cancellationToken);
+			Stream stream = await _toolCollection.GetDeploymentZipAsync(tool, deployment, cancellationToken);
 			return new FileStreamResult(stream, deployment.MimeType);
 		}
 
