Horde: Clear issue quarantine if being resolved or force closed
#rnx

diff --git a/Engine/Source/Programs/Horde/Horde.Build/Issues/IssueCollection.cs b/Engine/Source/Programs/Horde/Horde.Build/Issues/IssueCollection.cs
--- a/Engine/Source/Programs/Horde/Horde.Build/Issues/IssueCollection.cs
+++ b/Engine/Source/Programs/Horde/Horde.Build/Issues/IssueCollection.cs
@@ -1080,6 +1080,12 @@ namespace Horde.Build.Issues
 					updates.Add(Builders<Issue>.Update.Set(x => x.QuarantineTimeUtc, DateTime.UtcNow));
 				}
 			}
+			else if ((newResolvedById != null && newResolvedById.Value != UserId.Empty) || (newForceClosedById != null && newForceClosedById.Value != UserId.Empty))
+			{
+				// Clear quarantine if being resolved or if being force closed
+				updates.Add(Builders<Issue>.Update.Unset(x => x.QuarantinedByUserId));
+				updates.Add(Builders<Issue>.Update.Unset(x => x.QuarantineTimeUtc));
+			}
 
 			if (newForceClosedById != null)
 			{
