Horde: Fix tools endpoint returning incorrect response objects.

diff --git a/Engine/Source/Programs/Horde/Horde.Build/Tools/ToolsController.cs b/Engine/Source/Programs/Horde/Horde.Build/Tools/ToolsController.cs
--- a/Engine/Source/Programs/Horde/Horde.Build/Tools/ToolsController.cs
+++ b/Engine/Source/Programs/Horde/Horde.Build/Tools/ToolsController.cs
@@ -270,7 +270,7 @@ namespace Horde.Build.Tools
 		}
 
 		/// <summary>
-		/// Register an agent to perform remote work.
+		/// Gets information about a particular tool
 		/// </summary>
 		/// <returns>Information about the registered agent</returns>
 		[HttpGet]
@@ -289,7 +289,7 @@ namespace Horde.Build.Tools
 
 			if (action == GetToolAction.Info)
 			{
-				return Ok(tool);
+				return Ok(new GetToolResponse(tool));
 			}
 			else
 			{
@@ -303,7 +303,7 @@ namespace Horde.Build.Tools
 		}
 
 		/// <summary>
-		/// Register an agent to perform remote work.
+		/// Finds deployments of a particular tool.
 		/// </summary>
 		/// <param name="id">The tool identifier</param>
 		/// <param name="phase">Value indicating the client's preference for deployment to receive.</param>
@@ -341,7 +341,7 @@ namespace Horde.Build.Tools
 
 			if (action == GetToolAction.Info)
 			{
-				return Ok(tool.Deployments[idx]);
+				return Ok(new GetToolDeploymentResponse(tool.Deployments[idx]));
 			}
 			else
 			{
