Horde: Allow direct access to blobs for tools.

diff --git a/Engine/Source/Programs/Horde/Horde.Build/Tools/ITool.cs b/Engine/Source/Programs/Horde/Horde.Build/Tools/ITool.cs
--- a/Engine/Source/Programs/Horde/Horde.Build/Tools/ITool.cs
+++ b/Engine/Source/Programs/Horde/Horde.Build/Tools/ITool.cs
@@ -93,11 +93,6 @@ namespace Horde.Build.Tools
 		/// </summary>
 		TimeSpan Duration { get; }
 
-		/// <summary>
-		/// Mime-type for the data
-		/// </summary>
-		string MimeType { get; }
-
 		/// <summary>
 		/// Reference to this tool in Horde Storage.
 		/// </summary>
diff --git a/Engine/Source/Programs/Horde/Horde.Build/Tools/ToolCollection.cs b/Engine/Source/Programs/Horde/Horde.Build/Tools/ToolCollection.cs
--- a/Engine/Source/Programs/Horde/Horde.Build/Tools/ToolCollection.cs
+++ b/Engine/Source/Programs/Horde/Horde.Build/Tools/ToolCollection.cs
@@ -1,6 +1,5 @@
 // Copyright Epic Games, Inc. All Rights Reserved.
 
-using EpicGames.Core;
 using EpicGames.Horde.Storage;
 using EpicGames.Horde.Storage.Nodes;
 using EpicGames.Serialization;
@@ -16,10 +15,8 @@ using MongoDB.Bson.Serialization.Attributes;
 using MongoDB.Driver;
 using StackExchange.Redis;
 using System;
-using System.Buffers;
 using System.Collections.Generic;
 using System.IO;
-using System.IO.Pipelines;
 using System.Linq;
 using System.Threading;
 using System.Threading.Tasks;
@@ -89,9 +86,6 @@ namespace Horde.Build.Tools
 			[BsonElement("dur")]
 			public TimeSpan Duration { get; set; }
 
-			[BsonElement("mtp")]
-			public string MimeType { get; set; }
-
 			[BsonElement("ref")]
 			public RefName RefName { get; set; }
 
@@ -99,7 +93,6 @@ namespace Horde.Build.Tools
 			{
 				Id = id;
 				Version = String.Empty;
-				MimeType = ToolDeploymentConfig.DefaultMimeType;
 			}
 
 			public ToolDeployment(ToolDeploymentId id, ToolDeploymentConfig options, RefName refName)
@@ -107,7 +100,6 @@ namespace Horde.Build.Tools
 				Id = id;
 				Version = options.Version;
 				Duration = options.Duration;
-				MimeType = options.MimeType;
 				RefName = refName;
 			}
 
diff --git a/Engine/Source/Programs/Horde/Horde.Build/Tools/ToolConfig.cs b/Engine/Source/Programs/Horde/Horde.Build/Tools/ToolConfig.cs
--- a/Engine/Source/Programs/Horde/Horde.Build/Tools/ToolConfig.cs
+++ b/Engine/Source/Programs/Horde/Horde.Build/Tools/ToolConfig.cs
@@ -94,20 +94,12 @@ namespace Horde.Build.Tools
 	/// </summary>
 	public class ToolDeploymentConfig
 	{
-		/// <summary>
-		/// Default mime types for deployment data
-		/// </summary>
-		public const string DefaultMimeType = "application/zip";
-
 		/// <inheritdoc cref="IToolDeployment.Version"/>
 		public string Version { get; set; } = "Unknown";
 
 		/// <inheritdoc cref="IToolDeployment.Duration"/>
 		public TimeSpan Duration { get; set; }
 
-		/// <inheritdoc cref="IToolDeployment.MimeType"/>
-		public string MimeType { get; set; } = DefaultMimeType;
-
 		/// <summary>
 		/// Whether to create the deployment in a paused state
 		/// </summary>
diff --git a/Engine/Source/Programs/Horde/Horde.Build/Tools/ToolsController.cs b/Engine/Source/Programs/Horde/Horde.Build/Tools/ToolsController.cs
--- a/Engine/Source/Programs/Horde/Horde.Build/Tools/ToolsController.cs
+++ b/Engine/Source/Programs/Horde/Horde.Build/Tools/ToolsController.cs
@@ -4,6 +4,7 @@ using EpicGames.Core;
 using EpicGames.Horde.Storage;
 using Horde.Build.Acls;
 using Horde.Build.Server;
+using Horde.Build.Storage;
 using Horde.Build.Utilities;
 using HordeCommon;
 using Microsoft.AspNetCore.Authorization;
@@ -46,10 +47,10 @@ namespace Horde.Build.Tools
 		/// <summary>
 		/// Constructor
 		/// </summary>
-		public GetToolResponse(ITool tool)
+		public GetToolResponse(ITool tool, List<GetToolDeploymentResponse> deployments)
 		{
 			_tool = tool;
-			Deployments = tool.Deployments.ConvertAll(x => new GetToolDeploymentResponse(x));
+			Deployments = deployments;
 		}
 	}
 
@@ -78,16 +79,22 @@ namespace Horde.Build.Tools
 		/// <inheritdoc cref="IToolDeployment.Duration"/>
 		public TimeSpan Duration => _deployment.Duration;
 
-		/// <inheritdoc cref="IToolDeployment.MimeType"/>
-		public string MimeType => _deployment.MimeType;
-
 		/// <inheritdoc cref="IToolDeployment.RefName"/>
 		public RefName RefName => _deployment.RefName;
 
+		/// <summary>
+		/// Node for downloading this deployment
+		/// </summary>
+		public NodeHandle RootNode { get; }
+
 		/// <summary>
 		/// Constructor
 		/// </summary>
-		public GetToolDeploymentResponse(IToolDeployment deployment) => _deployment = deployment;
+		public GetToolDeploymentResponse(IToolDeployment deployment, NodeHandle rootNode)
+		{
+			_deployment = deployment;
+			RootNode = rootNode;
+		}
 	}
 
 	/// <summary>
@@ -179,47 +186,6 @@ namespace Horde.Build.Tools
 			return new CreateDeploymentResponse { Id = tool.Deployments[^1].Id };
 		}
 
-		/// <summary>
-		/// Register an agent to perform remote work.
-		/// </summary>
-		/// <returns>Information about the registered agent</returns>
-		[HttpGet]
-		[Route("/api/v1/tools/{id}/deployments/{deploymentId}")]
-		public async Task<ActionResult> GetDeploymentAsync(ToolId id, ToolDeploymentId deploymentId, [FromQuery] GetToolAction action = GetToolAction.Info, CancellationToken cancellationToken = default)
-		{
-			ITool? tool = await _toolCollection.GetAsync(id, _globalConfig.Value);
-			if (tool == null)
-			{
-				return NotFound(id);
-			}
-			if (!tool.Config.Authorize(AclAction.DownloadTool, User))
-			{
-				return Forbid(AclAction.DownloadTool, id);
-			}
-
-			IToolDeployment? deployment = tool.Deployments.FirstOrDefault(x => x.Id == deploymentId);
-			if(deployment == null)
-			{
-				return NotFound(id, deploymentId);
-			}
-
-			return await GetDeploymentResponseAsync(tool, deployment, action, cancellationToken);
-		}
-
-		private async Task<ActionResult> GetDeploymentResponseAsync(ITool tool, IToolDeployment deployment, GetToolAction action, CancellationToken cancellationToken)
-		{
-			if (action == GetToolAction.Info)
-			{
-				GetToolDeploymentResponse response = new GetToolDeploymentResponse(deployment);
-				return Ok(response);
-			}
-			else
-			{
-				Stream stream = await _toolCollection.GetDeploymentZipAsync(tool, deployment, cancellationToken);
-				return new FileStreamResult(stream, "application/zip");
-			}
-		}
-
 		/// <summary>
 		/// Updates the state of an active deployment.
 		/// </summary>
@@ -256,15 +222,17 @@ namespace Horde.Build.Tools
 	public class PublicToolsController : HordeControllerBase
 	{
 		readonly ToolCollection _toolCollection;
+		readonly StorageService _storageService;
 		readonly IOptionsSnapshot<GlobalConfig> _globalConfig;
 		readonly IClock _clock;
 
 		/// <summary>
 		/// Constructor
 		/// </summary>
-		public PublicToolsController(ToolCollection toolCollection, IClock clock, IOptionsSnapshot<GlobalConfig> globalConfig)
+		public PublicToolsController(ToolCollection toolCollection, StorageService storageService, IClock clock, IOptionsSnapshot<GlobalConfig> globalConfig)
 		{
 			_toolCollection = toolCollection;
+			_storageService = storageService;
 			_clock = clock;
 			_globalConfig = globalConfig;
 		}
@@ -289,7 +257,13 @@ namespace Horde.Build.Tools
 
 			if (action == GetToolAction.Info)
 			{
-				return Ok(new GetToolResponse(tool));
+				List<GetToolDeploymentResponse> deploymentResponses = new List<GetToolDeploymentResponse>();
+				foreach (IToolDeployment deployment in tool.Deployments)
+				{
+					GetToolDeploymentResponse deploymentResponse = await GetDeploymentInfoResponseAsync(deployment, cancellationToken);
+					deploymentResponses.Add(deploymentResponse);
+				}
+				return Ok(new GetToolResponse(tool, deploymentResponses));
 			}
 			else
 			{
@@ -298,7 +272,7 @@ namespace Horde.Build.Tools
 					return NotFound(LogEvent.Create(LogLevel.Error, "Tool {ToolId} does not currently have any deployments", id));
 				}
 
-				return await GetDeploymentPayloadAsync(tool, tool.Deployments[^1], cancellationToken);
+				return await GetDeploymentResponseAsync(tool, tool.Deployments[^1], action, cancellationToken);
 			}
 		}
 
@@ -339,20 +313,87 @@ namespace Horde.Build.Tools
 				}
 			}
 
+			return await GetDeploymentResponseAsync(tool, tool.Deployments[idx], action, cancellationToken);
+		}
+
+		/// <summary>
+		/// Gets information about a specific tool deployment.
+		/// </summary>
+		/// <returns>Information about the registered agent</returns>
+		[HttpGet]
+		[Route("/api/v1/tools/{id}/deployments/{deploymentId}")]
+		public async Task<ActionResult> GetDeploymentAsync(ToolId id, ToolDeploymentId deploymentId, [FromQuery] GetToolAction action = GetToolAction.Info, CancellationToken cancellationToken = default)
+		{
+			ITool? tool = await _toolCollection.GetAsync(id, _globalConfig.Value);
+			if (tool == null)
+			{
+				return NotFound(id);
+			}
+			if (!AuthorizeDownload(tool))
+			{
+				return Forbid(AclAction.DownloadTool, id);
+			}
+
+			IToolDeployment? deployment = tool.Deployments.FirstOrDefault(x => x.Id == deploymentId);
+			if (deployment == null)
+			{
+				return NotFound(id, deploymentId);
+			}
+
+			return await GetDeploymentResponseAsync(tool, deployment, action, cancellationToken);
+		}
+
+		private async Task<ActionResult> GetDeploymentResponseAsync(ITool tool, IToolDeployment deployment, GetToolAction action, CancellationToken cancellationToken)
+		{
 			if (action == GetToolAction.Info)
 			{
-				return Ok(new GetToolDeploymentResponse(tool.Deployments[idx]));
+				GetToolDeploymentResponse response = await GetDeploymentInfoResponseAsync(deployment, cancellationToken);
+				return Ok(response);
 			}
 			else
 			{
-				return await GetDeploymentPayloadAsync(tool, tool.Deployments[idx], cancellationToken);
+				Stream stream = await _toolCollection.GetDeploymentZipAsync(tool, deployment, cancellationToken);
+				return new FileStreamResult(stream, "application/zip");
 			}
 		}
 
-		async Task<ActionResult> GetDeploymentPayloadAsync(ITool tool, IToolDeployment deployment, CancellationToken cancellationToken)
+		private async Task<GetToolDeploymentResponse> GetDeploymentInfoResponseAsync(IToolDeployment deployment, CancellationToken cancellationToken)
+		{
+			IStorageClient client = await _storageService.GetClientAsync(Namespace.Tools, cancellationToken);
+			NodeHandle rootHandle = await client.ReadRefTargetAsync(deployment.RefName, cancellationToken: cancellationToken);
+
+			return new GetToolDeploymentResponse(deployment, rootHandle);
+		}
+
+		/// <summary>
+		/// Retrieves blobs for a particular artifact
+		/// </summary>
+		/// <param name="id">Identifier of the tool to retrieve</param>
+		/// <param name="locator">The blob locator</param>
+		/// <param name="length">Length of data to return</param>
+		/// <param name="offset">Offset of the data to return</param>
+		/// <param name="cancellationToken">Cancellation token for the operation</param>
+		/// <returns>Information about all the artifacts</returns>
+		[HttpGet]
+		[Route("/api/v1/tools/{id}/blobs/{*locator}")]
+		public async Task<ActionResult<object>> ReadToolBlobAsync(ToolId id, BlobLocator locator, [FromQuery] int? offset = null, [FromQuery] int? length = null, CancellationToken cancellationToken = default)
 		{
-			Stream stream = await _toolCollection.GetDeploymentZipAsync(tool, deployment, cancellationToken);
-			return new FileStreamResult(stream, deployment.MimeType);
+			ITool? tool = await _toolCollection.GetAsync(id, _globalConfig.Value);
+			if (tool == null)
+			{
+				return NotFound(id);
+			}
+			if (!AuthorizeDownload(tool))
+			{
+				return Forbid(AclAction.DownloadTool, id);
+			}
+
+			if (!locator.BlobId.WithinFolder(tool.Id.Id.Text))
+			{
+				return BadRequest("Invalid blob id for tool");
+			}
+
+			return StorageController.ReadBlobInternalAsync(_storageService, Namespace.Tools, locator, offset, length, cancellationToken);
 		}
 
 		bool AuthorizeDownload(ITool tool)
