Horde:  Associate issues with created workflow threads (server side)
#rnx

diff --git a/Engine/Source/Programs/Horde/Horde.Build/Issues/IIssue.cs b/Engine/Source/Programs/Horde/Horde.Build/Issues/IIssue.cs
--- a/Engine/Source/Programs/Horde/Horde.Build/Issues/IIssue.cs
+++ b/Engine/Source/Programs/Horde/Horde.Build/Issues/IIssue.cs
@@ -171,6 +171,11 @@ namespace Horde.Build.Issues
 		/// </summary>
 		public UserId? ForceClosedByUserId { get; set; }
 
+		/// <summary>
+		/// The workflow thread url created for this issue
+		/// </summary>
+		public Uri? WorkflowThreadUrl { get; set; }
+
 	}
 
 	/// <summary>
diff --git a/Engine/Source/Programs/Horde/Horde.Build/Issues/IIssueCollection.cs b/Engine/Source/Programs/Horde/Horde.Build/Issues/IIssueCollection.cs
--- a/Engine/Source/Programs/Horde/Horde.Build/Issues/IIssueCollection.cs
+++ b/Engine/Source/Programs/Horde/Horde.Build/Issues/IIssueCollection.cs
@@ -461,8 +461,9 @@ namespace Horde.Build.Issues
 		/// <param name="newExternalIssueKey">Issue key for external issue tracking</param>
 		/// <param name="newQuarantinedById">The user that quarantined the issue</param>
 		/// <param name="newForceClosedById">The user that force closed the issue</param>
+		/// <param name="newWorkflowThreadUrl">The workflow thread url associated with the issue</param>
 		/// <returns>True if the issue was updated</returns>
-		Task<IIssue?> TryUpdateIssueAsync(IIssue issue, UserId? initiatedByUserId, IssueSeverity? newSeverity = null, string? newSummary = null, string? newUserSummary = null, string? newDescription = null, bool? newPromoted = null, UserId? newOwnerId = null, UserId? newNominatedById = null, bool? newAcknowledged = null, UserId? newDeclinedById = null, int? newFixChange = null, UserId? newResolvedById = null, List<ObjectId>? newExcludeSpanIds = null, DateTime? newLastSeenAt = null, string? newExternalIssueKey = null, UserId? newQuarantinedById = null, UserId? newForceClosedById = null);
+		Task<IIssue?> TryUpdateIssueAsync(IIssue issue, UserId? initiatedByUserId, IssueSeverity? newSeverity = null, string? newSummary = null, string? newUserSummary = null, string? newDescription = null, bool? newPromoted = null, UserId? newOwnerId = null, UserId? newNominatedById = null, bool? newAcknowledged = null, UserId? newDeclinedById = null, int? newFixChange = null, UserId? newResolvedById = null, List<ObjectId>? newExcludeSpanIds = null, DateTime? newLastSeenAt = null, string? newExternalIssueKey = null, UserId? newQuarantinedById = null, UserId? newForceClosedById = null, Uri? newWorkflowThreadUrl = null);
 
 		/// <summary>
 		/// Updates derived data for an issue (ie. data computed from the spans attached to it). Also clears the issue's 'modified' state.
diff --git a/Engine/Source/Programs/Horde/Horde.Build/Issues/IssueCollection.cs b/Engine/Source/Programs/Horde/Horde.Build/Issues/IssueCollection.cs
--- a/Engine/Source/Programs/Horde/Horde.Build/Issues/IssueCollection.cs
+++ b/Engine/Source/Programs/Horde/Horde.Build/Issues/IssueCollection.cs
@@ -120,6 +120,9 @@ namespace Horde.Build.Issues
 			[BsonIgnoreIfNull]
 			public UserId? ForceClosedByUserId { get; set; }
 
+			[BsonIgnoreIfNull]
+			public Uri? WorkflowThreadUrl { get; set; }
+
 			[BsonConstructor]
 			private Issue()
 			{
@@ -918,7 +921,7 @@ namespace Horde.Build.Issues
 		}
 
 		/// <inheritdoc/>
-		public async Task<IIssue?> TryUpdateIssueAsync(IIssue issue, UserId? initiatedByUserId, IssueSeverity? newSeverity = null, string? newSummary = null, string? newUserSummary = null, string? newDescription = null, bool? newManuallyPromoted = null, UserId? newOwnerId = null, UserId? newNominatedById = null, bool? newAcknowledged = null, UserId? newDeclinedById = null, int? newFixChange = null, UserId? newResolvedById = null, List<ObjectId>? newExcludeSpanIds = null, DateTime? newLastSeenAt = null, string? newExternaIssueKey = null, UserId? newQuarantinedById = null, UserId? newForceClosedById = null)
+		public async Task<IIssue?> TryUpdateIssueAsync(IIssue issue, UserId? initiatedByUserId, IssueSeverity? newSeverity = null, string? newSummary = null, string? newUserSummary = null, string? newDescription = null, bool? newManuallyPromoted = null, UserId? newOwnerId = null, UserId? newNominatedById = null, bool? newAcknowledged = null, UserId? newDeclinedById = null, int? newFixChange = null, UserId? newResolvedById = null, List<ObjectId>? newExcludeSpanIds = null, DateTime? newLastSeenAt = null, string? newExternaIssueKey = null, UserId? newQuarantinedById = null, UserId? newForceClosedById = null, Uri? newWorkflowThreadUrl = null)
 		{
 			Issue issueDocument = (Issue)issue;
 
@@ -1104,6 +1107,12 @@ namespace Horde.Build.Issues
 				updates.Add(Builders<Issue>.Update.Set(x => x.ExternalIssueKey, newExternaIssueKey.Length == 0 ? null : newExternaIssueKey));
 			}
 
+			if (newWorkflowThreadUrl != null)
+			{
+				updates.Add(Builders<Issue>.Update.Set(x => x.WorkflowThreadUrl, newWorkflowThreadUrl.ToString().Length == 0 ? null : newWorkflowThreadUrl));
+			}
+
+
 			if (updates.Count == 0)
 			{
 				return issueDocument;
diff --git a/Engine/Source/Programs/Horde/Horde.Build/Issues/IssueMessages.cs b/Engine/Source/Programs/Horde/Horde.Build/Issues/IssueMessages.cs
--- a/Engine/Source/Programs/Horde/Horde.Build/Issues/IssueMessages.cs
+++ b/Engine/Source/Programs/Horde/Horde.Build/Issues/IssueMessages.cs
@@ -530,6 +530,11 @@ namespace Horde.Build.Issues
 		/// </summary>
 		public GetThinUserInfoResponse? ForceClosedByUserInfo { get; set; }
 
+		/// <summary>
+		/// The workflow thread url for this issue
+		/// </summary>
+		public Uri? WorkflowThreadUrl { get; set; }
+
 		/// <summary>
 		/// Constructs a new issue
 		/// </summary>
@@ -591,6 +596,7 @@ namespace Horde.Build.Issues
 			QuarantinedByUserInfo = details.QuarantinedBy != null ? new GetThinUserInfoResponse(details.QuarantinedBy) : null;
 			QuarantineTimeUtc = details.QuarantineTimeUtc;
 			ForceClosedByUserInfo = details.ForceClosedBy != null ? new GetThinUserInfoResponse(details.ForceClosedBy) : null;
+			WorkflowThreadUrl = issue.WorkflowThreadUrl;
 		}
 	}
 
diff --git a/Engine/Source/Programs/Horde/Horde.Build/Issues/IssueService.cs b/Engine/Source/Programs/Horde/Horde.Build/Issues/IssueService.cs
--- a/Engine/Source/Programs/Horde/Horde.Build/Issues/IssueService.cs
+++ b/Engine/Source/Programs/Horde/Horde.Build/Issues/IssueService.cs
@@ -406,8 +406,9 @@ namespace Horde.Build.Issues
 		/// <param name="quarantinedById">User who has quarantined the issue</param>
 		/// <param name="forceClosedById">User who has force closed the issue</param>
 		/// <param name="initiatedById">User initiating the changes, for auditing purposes</param>
+		/// <param name="workflowThreadUrl">The workfloe thread created for the issue</param>
 		/// <returns>True if the issue was updated</returns>
-		public async Task<bool> UpdateIssueAsync(int id, string? summary = null, string? description = null, bool? promoted = null, UserId? ownerId = null, UserId? nominatedById = null, bool? acknowledged = null, UserId? declinedById = null, int? fixChange = null, UserId? resolvedById = null, List<ObjectId>? addSpanIds = null, List<ObjectId>? removeSpanIds = null, string? externalIssueKey = null, UserId? quarantinedById = null, UserId? forceClosedById = null, UserId? initiatedById = null)
+		public async Task<bool> UpdateIssueAsync(int id, string? summary = null, string? description = null, bool? promoted = null, UserId? ownerId = null, UserId? nominatedById = null, bool? acknowledged = null, UserId? declinedById = null, int? fixChange = null, UserId? resolvedById = null, List<ObjectId>? addSpanIds = null, List<ObjectId>? removeSpanIds = null, string? externalIssueKey = null, UserId? quarantinedById = null, UserId? forceClosedById = null, UserId? initiatedById = null, Uri? workflowThreadUrl = null)
 		{
 			IIssue? issue;
 			for (; ; )
@@ -418,7 +419,7 @@ namespace Horde.Build.Issues
 					return false;
 				}
 
-				issue = await _issueCollection.TryUpdateIssueAsync(issue, initiatedById, newUserSummary: summary, newDescription: description, newPromoted: promoted, newOwnerId: ownerId ?? resolvedById, newNominatedById: nominatedById, newDeclinedById: declinedById, newAcknowledged: acknowledged, newFixChange: fixChange, newResolvedById: resolvedById, newExcludeSpanIds: removeSpanIds, newExternalIssueKey: externalIssueKey, newQuarantinedById: quarantinedById, newForceClosedById: forceClosedById);
+				issue = await _issueCollection.TryUpdateIssueAsync(issue, initiatedById, newUserSummary: summary, newDescription: description, newPromoted: promoted, newOwnerId: ownerId ?? resolvedById, newNominatedById: nominatedById, newDeclinedById: declinedById, newAcknowledged: acknowledged, newFixChange: fixChange, newResolvedById: resolvedById, newExcludeSpanIds: removeSpanIds, newExternalIssueKey: externalIssueKey, newQuarantinedById: quarantinedById, newForceClosedById: forceClosedById, newWorkflowThreadUrl: workflowThreadUrl);
 				if (issue != null)
 				{
 					break;
diff --git a/Engine/Source/Programs/Horde/Horde.Build/Notifications/Sinks/SlackNotificationSink.cs b/Engine/Source/Programs/Horde/Horde.Build/Notifications/Sinks/SlackNotificationSink.cs
--- a/Engine/Source/Programs/Horde/Horde.Build/Notifications/Sinks/SlackNotificationSink.cs
+++ b/Engine/Source/Programs/Horde/Horde.Build/Notifications/Sinks/SlackNotificationSink.cs
@@ -1017,6 +1017,17 @@ namespace Horde.Build.Notifications.Sinks
 				await UpdateMessageStateAsync(state.Id, state.MessageId, permalink);
 
 				_issueService.Collection.GetLogger(issue.Id).LogInformation("Created Slack thread: {SlackLink}", permalink);
+				if (permalink != null)
+				{
+					try
+					{
+						await _issueService.UpdateIssueAsync(issue.Id, workflowThreadUrl: new Uri(permalink));
+					}
+					catch (Exception Ex)
+					{
+						_issueService.Collection.GetLogger(issue.Id).LogInformation("Error associating workflow thread with issue, bad URI format? {ErrorMessage}", Ex.Message);
+					}					
+				}				
 			}
 
 			// Post a message containing the controls and status
