Horde: Fix exception when querying tool deployments.

diff --git a/Engine/Source/Programs/Horde/Horde.Build/Tools/ToolsController.cs b/Engine/Source/Programs/Horde/Horde.Build/Tools/ToolsController.cs
--- a/Engine/Source/Programs/Horde/Horde.Build/Tools/ToolsController.cs
+++ b/Engine/Source/Programs/Horde/Horde.Build/Tools/ToolsController.cs
@@ -82,18 +82,24 @@ namespace Horde.Build.Tools
 		/// <inheritdoc cref="IToolDeployment.RefName"/>
 		public RefName RefName => _deployment.RefName;
 
+		/// <summary>
+		/// Hash of the root node
+		/// </summary>
+		public IoHash Hash { get; }
+
 		/// <summary>
 		/// Node for downloading this deployment
 		/// </summary>
-		public NodeHandle RootNode { get; }
+		public NodeLocator Locator { get; }
 
 		/// <summary>
 		/// Constructor
 		/// </summary>
-		public GetToolDeploymentResponse(IToolDeployment deployment, NodeHandle rootNode)
+		public GetToolDeploymentResponse(IToolDeployment deployment, NodeHandle handle)
 		{
 			_deployment = deployment;
-			RootNode = rootNode;
+			Hash = handle.Hash;
+			Locator = handle.Locator;
 		}
 	}
 
