Horde: Download agent software via the tools collection.

diff --git a/Engine/Source/Programs/Horde/Horde.Build/Agents/IAgent.cs b/Engine/Source/Programs/Horde/Horde.Build/Agents/IAgent.cs
--- a/Engine/Source/Programs/Horde/Horde.Build/Agents/IAgent.cs
+++ b/Engine/Source/Programs/Horde/Horde.Build/Agents/IAgent.cs
@@ -19,6 +19,7 @@ using Horde.Build.Logs;
 using Horde.Build.Perforce;
 using Horde.Build.Server;
 using Horde.Build.Streams;
+using Horde.Build.Tools;
 using Horde.Build.Utilities;
 using HordeCommon;
 using HordeCommon.Rpc.Tasks;
@@ -495,6 +496,28 @@ namespace Horde.Build.Agents
 	/// </summary>
 	public static class AgentExtensions
 	{
+		static readonly ToolId s_defaultToolId = new ToolId("horde-agent");
+
+		/// <summary>
+		/// Gets the tool id for the software the given agent should be running
+		/// </summary>
+		/// <param name="agent">Agent to check</param>
+		/// <param name="globalConfig">Current global config</param>
+		/// <returns>Identifier for the tool that the agent should be using</returns>
+		public static ToolId GetSoftwareToolId(this IAgent agent, GlobalConfig globalConfig)
+		{
+			ToolId toolId = s_defaultToolId;
+			foreach (AgentSoftwareConfig softwareConfig in globalConfig.Software)
+			{
+				if (softwareConfig.Condition != null && agent.SatisfiesCondition(softwareConfig.Condition))
+				{
+					toolId = softwareConfig.ToolId;
+					break;
+				}
+			}
+			return toolId;
+		}
+
 		/// <summary>
 		/// Determines whether this agent is online
 		/// </summary>
diff --git a/Engine/Source/Programs/Horde/Horde.Build/Server/GlobalConfig.cs b/Engine/Source/Programs/Horde/Horde.Build/Server/GlobalConfig.cs
--- a/Engine/Source/Programs/Horde/Horde.Build/Server/GlobalConfig.cs
+++ b/Engine/Source/Programs/Horde/Horde.Build/Server/GlobalConfig.cs
@@ -437,13 +437,13 @@ namespace Horde.Build.Server
 	/// <summary>
 	/// Selects different agent software versions by evaluating a condition
 	/// </summary>
-	[DebuggerDisplay("{Channel}")]
+	[DebuggerDisplay("{ToolId}")]
 	public class AgentSoftwareConfig
 	{
 		/// <summary>
-		/// Channel name
+		/// Tool identifier
 		/// </summary>
-		public AgentSoftwareChannelName Channel { get; set; }
+		public ToolId ToolId { get; set; }
 
 		/// <summary>
 		/// Condition for using this channel
diff --git a/Engine/Source/Programs/Horde/Horde.Build/Server/RpcService.cs b/Engine/Source/Programs/Horde/Horde.Build/Server/RpcService.cs
--- a/Engine/Source/Programs/Horde/Horde.Build/Server/RpcService.cs
+++ b/Engine/Source/Programs/Horde/Horde.Build/Server/RpcService.cs
@@ -1,11 +1,14 @@
 // Copyright Epic Games, Inc. All Rights Reserved.
 
 using System;
+using System.Buffers;
 using System.Collections.Generic;
+using System.IO;
 using System.Linq;
 using System.Net;
 using System.Threading;
 using System.Threading.Tasks;
+using Amazon.EC2.Model;
 using Google.Protobuf.WellKnownTypes;
 using Grpc.Core;
 using Horde.Build.Acls;
@@ -15,6 +18,7 @@ using Horde.Build.Agents.Sessions;
 using Horde.Build.Agents.Software;
 using Horde.Build.Jobs;
 using Horde.Build.Tasks;
+using Horde.Build.Tools;
 using Horde.Build.Utilities;
 using HordeCommon.Rpc;
 using Microsoft.AspNetCore.Authorization;
@@ -43,25 +47,25 @@ namespace Horde.Build.Server
 		internal TimeSpan _longPollTimeout = TimeSpan.FromMinutes(9);
 
 		readonly AgentService _agentService;
-		readonly AgentSoftwareService _agentSoftwareService;
 		readonly PoolService _poolService;
 		readonly LifetimeService _lifetimeService;
 		readonly ConformTaskSource _conformTaskSource;
 		readonly JobRpcCommon _jobRpcCommon;
+		readonly ToolCollection _toolCollection;
 		readonly IOptionsSnapshot<GlobalConfig> _globalConfig;
 		readonly ILogger _logger;
 
 		/// <summary>
 		/// Constructor
 		/// </summary>
-		public RpcService(AgentService agentService, AgentSoftwareService agentSoftwareService, PoolService poolService, LifetimeService lifetimeService, ConformTaskSource conformTaskSource, JobRpcCommon jobRpcCommon, IOptionsSnapshot<GlobalConfig> globalConfig, ILogger<RpcService> logger)
+		public RpcService(AgentService agentService, PoolService poolService, LifetimeService lifetimeService, ConformTaskSource conformTaskSource, JobRpcCommon jobRpcCommon, ToolCollection toolCollection, IOptionsSnapshot<GlobalConfig> globalConfig, ILogger<RpcService> logger)
 		{
 			_agentService = agentService;
-			_agentSoftwareService = agentSoftwareService;
 			_poolService = poolService;
 			_lifetimeService = lifetimeService;
 			_conformTaskSource = conformTaskSource;
 			_jobRpcCommon = jobRpcCommon;
+			_toolCollection = toolCollection;
 			_globalConfig = globalConfig;
 			_logger = logger;
 		}
@@ -403,31 +407,6 @@ namespace Horde.Build.Server
 		/// <inheritdoc/>
 		public override Task<Empty> CreateEvents(CreateEventsRequest request, ServerCallContext context) => _jobRpcCommon.CreateEvents(request, context);
 
-
-		/// <summary>
-		/// Uploads a new agent archive
-		/// </summary>
-		/// <param name="request">Request arguments</param>
-		/// <param name="context">Context for the RPC call</param>
-		/// <returns>Information about the new agent</returns>
-		public override async Task<UploadSoftwareResponse> UploadSoftware(UploadSoftwareRequest request, ServerCallContext context)
-		{
-			if (!_globalConfig.Value.Authorize(AclAction.UploadSoftware, context.GetHttpContext().User))
-			{
-				throw new StructuredRpcException(StatusCode.PermissionDenied, "Access to software is forbidden");
-			}
-
-			string version = await _agentSoftwareService.UploadArchiveAsync(request.Data.ToArray());
-			foreach (string channel in request.Channel.Split('+', ';'))
-			{
-				await _agentSoftwareService.UpdateChannelAsync(new AgentSoftwareChannelName(channel), null, version);
-			}
-
-			UploadSoftwareResponse response = new UploadSoftwareResponse();
-			response.Version = version;
-			return response;
-		}
-
 		/// <summary>
 		/// Downloads a new agent archive
 		/// </summary>
@@ -442,22 +421,37 @@ namespace Horde.Build.Server
 				throw new StructuredRpcException(StatusCode.NotFound, "Access to software is forbidden");
 			}
 
-			byte[]? data = await _agentSoftwareService.GetArchiveAsync(request.Version);
-			if (data == null)
+			int colonIdx = request.Version.IndexOf(':', StringComparison.Ordinal);
+			ToolId toolId = new ToolId(request.Version.Substring(0, colonIdx));
+			string version = request.Version.Substring(colonIdx + 1);
+
+			ITool? tool = await _toolCollection.GetAsync(toolId, _globalConfig.Value);
+			if (tool == null)
 			{
-				throw new StructuredRpcException(StatusCode.NotFound, "Missing version {Version}");
+				throw new StructuredRpcException(StatusCode.NotFound, $"Missing tool {toolId}");
 			}
 
-			for (int offset = 0; offset < data.Length;)
+			IToolDeployment? deployment = tool.Deployments.LastOrDefault(x => x.Version.Equals(version, StringComparison.Ordinal));
+			if (deployment == null)
 			{
-				int nextOffset = Math.Min(offset + 128 * 1024, data.Length);
-
-				DownloadSoftwareResponse response = new DownloadSoftwareResponse();
-				response.Data = Google.Protobuf.ByteString.CopyFrom(data.AsSpan(offset, nextOffset - offset));
+				throw new StructuredRpcException(StatusCode.NotFound, $"Missing tool version {version}");
+			}
 
-				await responseStream.WriteAsync(response);
+			using Stream stream = await _toolCollection.GetDeploymentZipAsync(tool, deployment, context.CancellationToken);
+			using (IMemoryOwner<byte> buffer = MemoryPool<byte>.Shared.Rent(128 * 1024))
+			{
+				for(; ;)
+				{
+					int read = await stream.ReadAsync(buffer.Memory, context.CancellationToken);
+					if (read == 0)
+					{
+						break;
+					}
 
-				offset = nextOffset;
+					DownloadSoftwareResponse response = new DownloadSoftwareResponse();
+					response.Data = Google.Protobuf.ByteString.CopyFrom(buffer.Memory.Slice(0, read).Span);
+					await responseStream.WriteAsync(response);
+				}
 			}
 		}
 
diff --git a/Engine/Source/Programs/Horde/Horde.Build/Tasks/UpgradeTaskSource.cs b/Engine/Source/Programs/Horde/Horde.Build/Tasks/UpgradeTaskSource.cs
--- a/Engine/Source/Programs/Horde/Horde.Build/Tasks/UpgradeTaskSource.cs
+++ b/Engine/Source/Programs/Horde/Horde.Build/Tasks/UpgradeTaskSource.cs
@@ -1,8 +1,10 @@
 // Copyright Epic Games, Inc. All Rights Reserved.
 
 using System;
+using System.Text;
 using System.Threading;
 using System.Threading.Tasks;
+using EpicGames.Core;
 using Google.Protobuf;
 using Google.Protobuf.WellKnownTypes;
 using Horde.Build.Agents;
@@ -11,7 +13,7 @@ using Horde.Build.Agents.Software;
 using Horde.Build.Jobs;
 using Horde.Build.Logs;
 using Horde.Build.Server;
-using Horde.Build.Utilities;
+using Horde.Build.Tools;
 using HordeCommon;
 using HordeCommon.Rpc.Tasks;
 using Microsoft.Extensions.Options;
@@ -24,15 +26,15 @@ namespace Horde.Build.Tasks
 
 		public override TaskSourceFlags Flags => TaskSourceFlags.AllowWhenDisabled | TaskSourceFlags.AllowDuringDowntime;
 
-		readonly AgentSoftwareService _agentSoftwareService;
+		readonly ToolCollection _toolCollection;
 		readonly ILogFileService _logService;
 		readonly IOptionsMonitor<GlobalConfig> _globalConfig;
 		readonly IOptions<ServerSettings> _serverSettings;
 		readonly IClock _clock;
 
-		public UpgradeTaskSource(AgentSoftwareService agentSoftwareService, ILogFileService logService, IOptionsMonitor<GlobalConfig> globalConfig, IOptions<ServerSettings> serverSettings, IClock clock)
+		public UpgradeTaskSource(ToolCollection toolCollection, ILogFileService logService, IOptionsMonitor<GlobalConfig> globalConfig, IOptions<ServerSettings> serverSettings, IClock clock)
 		{
-			_agentSoftwareService = agentSoftwareService;
+			_toolCollection = toolCollection;
 			_logService = logService;
 			_globalConfig = globalConfig;
 			_serverSettings = serverSettings;
@@ -48,8 +50,14 @@ namespace Horde.Build.Tasks
 				return Skip(cancellationToken);
 			}
 
-			string? requiredVersion = await GetRequiredSoftwareVersion(agent);
-			if (requiredVersion == null || agent.Version == requiredVersion)
+			(ITool, IToolDeployment)? required = await GetRequiredSoftwareVersion(agent);
+			if (required == null)
+			{
+				return Skip(cancellationToken);
+			}
+
+			(ITool tool, IToolDeployment deployment) = required.Value; 
+			if (agent.Version == deployment.Version)
 			{
 				return Skip(cancellationToken);
 			}
@@ -62,23 +70,11 @@ namespace Horde.Build.Tasks
 			ILogFile logFile = await _logService.CreateLogFileAsync(JobId.Empty, leaseId, agent.SessionId, LogType.Json, useNewStorageBackend: false, cancellationToken: cancellationToken);
 
 			UpgradeTask task = new UpgradeTask();
-			task.SoftwareId = requiredVersion.ToString();
+			task.SoftwareId = $"{tool.Id}:{deployment.Version}";
 			task.LogId = logFile.Id.ToString();
 
-			byte[] payload;
-			if (agent.Version == "5.0.0-17425336" || agent.Version == "5.0.0-17448746")
-			{
-				Any any = new Any();
-				any.TypeUrl = "type.googleapis.com/Horde.UpgradeTask";
-				any.Value = task.ToByteString();
-				payload = any.ToByteArray();
-			}
-			else
-			{
-				payload = Any.Pack(task).ToByteArray();
-			}
-
-			return Lease(new AgentLease(leaseId, $"Upgrade to {requiredVersion}", null, null, logFile.Id, LeaseState.Pending, null, true, payload));
+			byte[] payload = Any.Pack(task).ToByteArray();
+			return Lease(new AgentLease(leaseId, $"Upgrade to {tool.Id} {deployment.Version}", null, null, logFile.Id, LeaseState.Pending, null, true, payload));
 		}
 
 		/// <summary>
@@ -86,22 +82,28 @@ namespace Horde.Build.Tasks
 		/// </summary>
 		/// <param name="agent">The agent instance</param>
 		/// <returns>Unique id of the client version this agent should be running</returns>
-		public async Task<string?> GetRequiredSoftwareVersion(IAgent agent)
+		public async Task<(ITool, IToolDeployment)?> GetRequiredSoftwareVersion(IAgent agent)
 		{
 			GlobalConfig globalConfig = _globalConfig.CurrentValue;
 
-			AgentSoftwareChannelName channelName = AgentSoftwareService.DefaultChannelName;
-			foreach (AgentSoftwareConfig softwareConfig in globalConfig.Software)
+			ToolId toolId = agent.GetSoftwareToolId(globalConfig);
+
+			ITool? tool = await _toolCollection.GetAsync(toolId, globalConfig);
+			if (tool == null)
+			{
+				return null;
+			}
+
+			uint value = BuzHash.Add(0, Encoding.UTF8.GetBytes(agent.Id.ToString()));
+			double phase = (value % 10000) / 10000.0;
+
+			IToolDeployment? deployment = tool.GetCurrentDeployment(phase, _clock.UtcNow);
+			if (deployment == null)
 			{
-				if (softwareConfig.Condition != null && agent.SatisfiesCondition(softwareConfig.Condition))
-				{
-					channelName = softwareConfig.Channel;
-					break;
-				}
+				return null;
 			}
 
-			IAgentSoftwareChannel? channel = await _agentSoftwareService.GetCachedChannelAsync(channelName);
-			return channel?.Version;
+			return (tool, deployment);
 		}
 	}
 }
diff --git a/Engine/Source/Programs/Horde/Horde.Build/Tools/ITool.cs b/Engine/Source/Programs/Horde/Horde.Build/Tools/ITool.cs
--- a/Engine/Source/Programs/Horde/Horde.Build/Tools/ITool.cs
+++ b/Engine/Source/Programs/Horde/Horde.Build/Tools/ITool.cs
@@ -104,6 +104,26 @@ namespace Horde.Build.Tools
 	/// </summary>
 	public static class ToolExtensions
 	{
+		/// <summary>
+		/// Gets the current deployment
+		/// </summary>
+		/// <param name="tool">Tool to query</param>
+		/// <param name="phase">Adoption phase for the caller. 0 which the </param>
+		/// <param name="utcNow">Current time</param>
+		/// <returns></returns>
+		public static IToolDeployment? GetCurrentDeployment(this ITool tool, double phase, DateTime utcNow)
+		{
+			int idx = tool.Deployments.Count - 1;
+			for (; idx >= 0; idx--)
+			{
+				if (phase < tool.Deployments[idx].GetProgressValue(utcNow) || idx == 0)
+				{
+					return tool.Deployments[idx];
+				}
+			}
+			return null;
+		}
+
 		/// <summary>
 		/// Get the progress fraction for a particular deployment and time
 		/// </summary>
diff --git a/Engine/Source/Programs/Horde/Horde.Build/Tools/ToolsController.cs b/Engine/Source/Programs/Horde/Horde.Build/Tools/ToolsController.cs
--- a/Engine/Source/Programs/Horde/Horde.Build/Tools/ToolsController.cs
+++ b/Engine/Source/Programs/Horde/Horde.Build/Tools/ToolsController.cs
@@ -303,23 +303,14 @@ namespace Horde.Build.Tools
 			{
 				return Forbid(AclAction.DownloadTool, id);
 			}
-			if (tool.Deployments.Count == 0)
-			{
-				return NotFound(LogEvent.Create(LogLevel.Error, "Tool {ToolId} does not currently have any deployments", id));
-			}
 
-			DateTime utcNow = _clock.UtcNow;
-
-			int idx = tool.Deployments.Count - 1;
-			for (; idx > 0; idx--)
+			IToolDeployment? deployment = tool.GetCurrentDeployment(phase, _clock.UtcNow);
+			if (deployment == null)
 			{
-				if (phase < tool.Deployments[idx].GetProgressValue(utcNow))
-				{
-					break;
-				}
+				return NotFound(LogEvent.Create(LogLevel.Error, "Tool {ToolId} does not currently have any deployments", id));
 			}
 
-			return await GetDeploymentResponseAsync(tool, tool.Deployments[idx], action, cancellationToken);
+			return await GetDeploymentResponseAsync(tool, deployment, action, cancellationToken);
 		}
 
 		/// <summary>
