Horde: Build fix.

diff --git a/Engine/Source/Programs/Horde/Horde.Build/Server/RpcService.cs b/Engine/Source/Programs/Horde/Horde.Build/Server/RpcService.cs
--- a/Engine/Source/Programs/Horde/Horde.Build/Server/RpcService.cs
+++ b/Engine/Source/Programs/Horde/Horde.Build/Server/RpcService.cs
@@ -9,6 +9,7 @@ using System.Net;
 using System.Threading;
 using System.Threading.Tasks;
 using Amazon.EC2.Model;
+using Google.Protobuf;
 using Google.Protobuf.WellKnownTypes;
 using Grpc.Core;
 using Horde.Build.Acls;
@@ -449,7 +450,7 @@ namespace Horde.Build.Server
 					}
 
 					DownloadSoftwareResponse response = new DownloadSoftwareResponse();
-					response.Data = Google.Protobuf.ByteString.CopyFrom(buffer.Memory.Slice(0, read).Span);
+					response.Data = UnsafeByteOperations.UnsafeWrap(buffer.Memory.Slice(0, read));
 					await responseStream.WriteAsync(response);
 				}
 			}
diff --git a/Engine/Source/Programs/Horde/Horde.Build/Tasks/UpgradeTaskSource.cs b/Engine/Source/Programs/Horde/Horde.Build/Tasks/UpgradeTaskSource.cs
--- a/Engine/Source/Programs/Horde/Horde.Build/Tasks/UpgradeTaskSource.cs
+++ b/Engine/Source/Programs/Horde/Horde.Build/Tasks/UpgradeTaskSource.cs
@@ -61,7 +61,7 @@ namespace Horde.Build.Tasks
 			{
 				return Skip(cancellationToken);
 			}
-			if (agent.Leases.Count > 0 || !(agent.LastUpgradeTime == null || agent.LastUpgradeTime.Value + TimeSpan.FromMinutes(5.0) < _clock.UtcNow || agent.LastUpgradeVersion != requiredVersion.ToString()))
+			if (agent.Leases.Count > 0 || !(agent.LastUpgradeTime == null || agent.LastUpgradeTime.Value + TimeSpan.FromMinutes(5.0) < _clock.UtcNow || agent.LastUpgradeVersion != deployment.Version))
 			{
 				return await DrainAsync(cancellationToken);
 			}
