Horde: Remove agent software collection/service. Agent software is now stored in the tools controller.

diff --git a/Engine/Source/Programs/Horde/Horde.Agent/Commands/Utilities/PublishCommand.cs b/Engine/Source/Programs/Horde/Horde.Agent/Commands/Utilities/PublishCommand.cs
deleted file mode 100644
diff --git a/Engine/Source/Programs/Horde/Horde.Build.Tests/TestSetup.cs b/Engine/Source/Programs/Horde/Horde.Build.Tests/TestSetup.cs
--- a/Engine/Source/Programs/Horde/Horde.Build.Tests/TestSetup.cs
+++ b/Engine/Source/Programs/Horde/Horde.Build.Tests/TestSetup.cs
@@ -90,7 +90,6 @@ namespace Horde.Build.Tests
 		public IPoolCollection PoolCollection => ServiceProvider.GetRequiredService <IPoolCollection>();
 		public ILeaseCollection LeaseCollection => ServiceProvider.GetRequiredService <ILeaseCollection>();
 		public ISessionCollection SessionCollection => ServiceProvider.GetRequiredService <ISessionCollection>();
-		public IAgentSoftwareCollection AgentSoftwareCollection => ServiceProvider.GetRequiredService <IAgentSoftwareCollection>();
 		public ITestDataCollection TestDataCollection => ServiceProvider.GetRequiredService<ITestDataCollection>();
 		public IUserCollection UserCollection => ServiceProvider.GetRequiredService<IUserCollection>();
 		public IDeviceCollection DeviceCollection => ServiceProvider.GetRequiredService<IDeviceCollection>();
@@ -98,7 +97,6 @@ namespace Horde.Build.Tests
 
 		public AclService AclService => ServiceProvider.GetRequiredService<AclService>();
 		public FleetService FleetService => ServiceProvider.GetRequiredService<FleetService>();
-		public AgentSoftwareService AgentSoftwareService => ServiceProvider.GetRequiredService<AgentSoftwareService>();
 		public AgentService AgentService => ServiceProvider.GetRequiredService<AgentService>();
 		public ICommitService CommitService => ServiceProvider.GetRequiredService<ICommitService>();
 		public GlobalsService GlobalsService => ServiceProvider.GetRequiredService<GlobalsService>();
@@ -200,7 +198,6 @@ namespace Horde.Build.Tests
 			services.AddSingleton<IConfigSource, FileConfigSource>();
 
 			services.AddSingleton<IAgentCollection, AgentCollection>();
-			services.AddSingleton<IAgentSoftwareCollection, AgentSoftwareCollection>();
 			services.AddSingleton<IArtifactCollection, ArtifactCollection>();
 			services.AddSingleton<IArtifactCollectionV1, ArtifactCollectionV1>();
 			services.AddSingleton<ICommitService, CommitService>();
@@ -239,7 +236,6 @@ namespace Horde.Build.Tests
 
 			services.AddSingleton<AclService>();
 			services.AddSingleton<AgentService>();
-			services.AddSingleton<AgentSoftwareService>();
 			services.AddSingleton<FleetService>();
 			services.AddSingleton<ConsistencyService>();
 			services.AddSingleton<RequestTrackerService>();
@@ -277,7 +273,6 @@ namespace Horde.Build.Tests
 
 			services.AddSingleton<StorageService>();
 
-			services.AddSingleton<ISingletonDocument<AgentSoftwareChannels>>(new SingletonDocumentStub<AgentSoftwareChannels>());
 			services.AddSingleton<ISingletonDocument<DevicePlatformMapV1>>(new SingletonDocumentStub<DevicePlatformMapV1>());
 		}
 
diff --git a/Engine/Source/Programs/Horde/Horde.Build/Agents/IAgent.cs b/Engine/Source/Programs/Horde/Horde.Build/Agents/IAgent.cs
--- a/Engine/Source/Programs/Horde/Horde.Build/Agents/IAgent.cs
+++ b/Engine/Source/Programs/Horde/Horde.Build/Agents/IAgent.cs
@@ -496,7 +496,10 @@ namespace Horde.Build.Agents
 	/// </summary>
 	public static class AgentExtensions
 	{
-		static readonly ToolId s_defaultToolId = new ToolId("horde-agent");
+		/// <summary>
+		/// Default tool id for agent software
+		/// </summary>
+		public static ToolId DefaultAgentSoftwareToolId { get; } = new ToolId("horde-agent");
 
 		/// <summary>
 		/// Gets the tool id for the software the given agent should be running
@@ -506,7 +509,7 @@ namespace Horde.Build.Agents
 		/// <returns>Identifier for the tool that the agent should be using</returns>
 		public static ToolId GetSoftwareToolId(this IAgent agent, GlobalConfig globalConfig)
 		{
-			ToolId toolId = s_defaultToolId;
+			ToolId toolId = DefaultAgentSoftwareToolId;
 			foreach (AgentSoftwareConfig softwareConfig in globalConfig.Software)
 			{
 				if (softwareConfig.Condition != null && agent.SatisfiesCondition(softwareConfig.Condition))
diff --git a/Engine/Source/Programs/Horde/Horde.Build/Agents/Software/AgentSoftwareChannelName.cs b/Engine/Source/Programs/Horde/Horde.Build/Agents/Software/AgentSoftwareChannelName.cs
deleted file mode 100644
diff --git a/Engine/Source/Programs/Horde/Horde.Build/Agents/Software/AgentSoftwareCollection.cs b/Engine/Source/Programs/Horde/Horde.Build/Agents/Software/AgentSoftwareCollection.cs
deleted file mode 100644
diff --git a/Engine/Source/Programs/Horde/Horde.Build/Agents/Software/AgentSoftwareController.Messages.cs b/Engine/Source/Programs/Horde/Horde.Build/Agents/Software/AgentSoftwareController.Messages.cs
deleted file mode 100644
diff --git a/Engine/Source/Programs/Horde/Horde.Build/Agents/Software/AgentSoftwareController.cs b/Engine/Source/Programs/Horde/Horde.Build/Agents/Software/AgentSoftwareController.cs
--- a/Engine/Source/Programs/Horde/Horde.Build/Agents/Software/AgentSoftwareController.cs
+++ b/Engine/Source/Programs/Horde/Horde.Build/Agents/Software/AgentSoftwareController.cs
@@ -1,13 +1,14 @@
 // Copyright Epic Games, Inc. All Rights Reserved.
 
-using System.Collections.Generic;
+using System;
 using System.IO;
+using System.Threading;
 using System.Threading.Tasks;
 using Horde.Build.Acls;
 using Horde.Build.Server;
-using Horde.Build.Utilities;
+using Horde.Build.Tools;
+using HordeCommon;
 using Microsoft.AspNetCore.Authorization;
-using Microsoft.AspNetCore.Http;
 using Microsoft.AspNetCore.Mvc;
 using Microsoft.Extensions.Options;
 using Microsoft.Net.Http.Headers;
@@ -15,147 +16,56 @@ using Microsoft.Net.Http.Headers;
 namespace Horde.Build.Agents.Software
 {
 	/// <summary>
-	/// Controller for the /api/v1/software endpoint
+	/// Controller for the /api/v1/agentsoftware endpoint
 	/// </summary>
 	[ApiController]
 	[Authorize]
 	[Route("[controller]")]
 	public class AgentSoftwareController : ControllerBase
 	{
-		private readonly AgentSoftwareService _agentSoftwareService;
+		private readonly ToolCollection _toolCollection;
+		private readonly IClock _clock;
 		private readonly IOptionsSnapshot<GlobalConfig> _globalConfig;
 
 		/// <summary>
 		/// Constructor
 		/// </summary>
-		public AgentSoftwareController(AgentSoftwareService agentSoftwareService, IOptionsSnapshot<GlobalConfig> globalConfig)
+		public AgentSoftwareController(ToolCollection toolCollection, IClock clock, IOptionsSnapshot<GlobalConfig> globalConfig)
 		{
-			_agentSoftwareService = agentSoftwareService;
+			_toolCollection = toolCollection;
+			_clock = clock;
 			_globalConfig = globalConfig;
 		}
 
 		/// <summary>
-		/// Finds all uploaded software matching the given criteria
-		/// </summary>
-		/// <param name="filter">Filter for the properties to return</param>
-		/// <returns>Http response</returns>
-		[HttpGet]
-		[Route("/api/v1/agentsoftware")]
-		[ProducesResponseType(typeof(List<GetAgentSoftwareChannelResponse>), 200)]
-		public async Task<ActionResult<List<object>>> FindSoftwareAsync([FromQuery] PropertyFilter? filter = null)
-		{
-			if (!_globalConfig.Value.Authorize(AclAction.DownloadSoftware, User))
-			{
-				return Forbid();
-			}
-
-			List<IAgentSoftwareChannel> results = await _agentSoftwareService.FindChannelsAsync();
-
-			List<object> responses = new List<object>();
-			foreach (IAgentSoftwareChannel result in results)
-			{
-				responses.Add(new GetAgentSoftwareChannelResponse(result).ApplyFilter(filter));
-			}
-			return responses;
-		}
-
-		/// <summary>
-		/// Finds all uploaded software matching the given criteria
+		/// Gets the zip file for a specific channel
 		/// </summary>
-		/// <param name="name">Name of the channel to get</param>
-		/// <param name="filter">Filter for the properties to return</param>
+		/// <param name="cancellationToken">Cancellation token</param>
 		/// <returns>Http response</returns>
 		[HttpGet]
-		[Route("/api/v1/agentsoftware/{name}")]
-		[ProducesResponseType(typeof(GetAgentSoftwareChannelResponse), 200)]
-		public async Task<ActionResult<object>> FindSoftwareAsync(string name, [FromQuery] PropertyFilter? filter = null)
+		[Obsolete("Agent software is now stored as a tool. This endpoint exists for backwards compatibility, but will be removed in the future.")]
+		[Route("/api/v1/agentsoftware/default/zip")]
+		public async Task<ActionResult> GetArchiveAsync(CancellationToken cancellationToken)
 		{
 			if (!_globalConfig.Value.Authorize(AclAction.DownloadSoftware, User))
 			{
 				return Forbid();
 			}
 
-			IAgentSoftwareChannel? channel = await _agentSoftwareService.GetChannelAsync(new AgentSoftwareChannelName(name));
-			if(channel == null)
-			{
-				return NotFound();
-			}
-
-			return new GetAgentSoftwareChannelResponse(channel).ApplyFilter(filter);
-		}
-
-		/// <summary>
-		/// Finds all uploaded software matching the given criteria
-		/// </summary>
-		/// <param name="name">Name of the channel to get</param>
-		/// <param name="request">Request body</param>
-		/// <param name="filter">Filter for the properties to return</param>
-		/// <returns>Http response</returns>
-		[HttpPut]
-		[Route("/api/v1/agentsoftware/{name}")]
-		public async Task<ActionResult<object>> SetSoftwareAsync(string name, [FromBody] UpdateAgentSoftwareChannelRequest request, [FromQuery] PropertyFilter? filter = null)
-		{
-			if (!_globalConfig.Value.Authorize(AclAction.UploadSoftware, User))
-			{
-				return Forbid();
-			}
-
-			await _agentSoftwareService.UpdateChannelAsync(new AgentSoftwareChannelName(name), User.GetUserName(), request.Version);
-			return Ok();
-		}
-
-		/// <summary>
-		/// Uploads a new agent zip file
-		/// </summary>
-		/// <param name="name">Name of the channel to post to</param>
-		/// <param name="file">Zip archive containing the new client software</param>
-		/// <returns>Http result code</returns>
-		[HttpPost]
-		[Route("/api/v1/agentsoftware/{name}/zip")]
-		public async Task<ActionResult<UploadAgentSoftwareResponse>> UploadArchiveAsync(string name, [FromForm] IFormFile file)
-		{
-			if (!_globalConfig.Value.Authorize(AclAction.UploadSoftware, User))
-			{
-				return Forbid();
-			}
-
-			byte[] data;
-			using (MemoryStream memoryStream = new MemoryStream())
-			{
-				using (System.IO.Stream stream = file.OpenReadStream())
-				{
-					await stream.CopyToAsync(memoryStream);
-				}
-				data = memoryStream.ToArray();
-			}
-
-			string version = await _agentSoftwareService.UploadArchiveAsync(data);
-			await _agentSoftwareService.UpdateChannelAsync(new AgentSoftwareChannelName(name), User.Identity?.Name, version);
-
-			return new UploadAgentSoftwareResponse { Version = version };
-		}
-
-		/// <summary>
-		/// Gets the zip file for a specific channel
-		/// </summary>
-		/// <param name="name">Name of the channel</param>
-		/// <returns>Http response</returns>
-		[HttpGet]
-		[Route("/api/v1/agentsoftware/{name}/zip")]
-		public async Task<ActionResult> GetArchiveAsync(string name)
-		{
-			if (!_globalConfig.Value.Authorize(AclAction.DownloadSoftware, User))
+			ITool? tool = await _toolCollection.GetAsync(AgentExtensions.DefaultAgentSoftwareToolId, _globalConfig.Value);
+			if (tool == null)
 			{
-				return Forbid();
+				return NotFound("No agent software tool is currently registered");
 			}
 
-			byte[]? data = await _agentSoftwareService.GetArchiveAsync(new AgentSoftwareChannelName(name));
-			if (data == null)
+			IToolDeployment? deployment = tool.GetCurrentDeployment(1.0, _clock.UtcNow);
+			if (deployment == null)
 			{
-				return NotFound();
+				return NotFound("No deployment currently set for agent software");
 			}
 
-			return new FileStreamResult(new MemoryStream(data), new MediaTypeHeaderValue("application/octet-stream")) { FileDownloadName = $"HordeAgent.zip" };
+			Stream stream = await _toolCollection.GetDeploymentZipAsync(tool, deployment, cancellationToken);
+			return new FileStreamResult(stream, new MediaTypeHeaderValue("application/octet-stream")) { FileDownloadName = $"HordeAgent.zip" };
 		}
 	}
 }
diff --git a/Engine/Source/Programs/Horde/Horde.Build/Agents/Software/AgentSoftwareService.cs b/Engine/Source/Programs/Horde/Horde.Build/Agents/Software/AgentSoftwareService.cs
deleted file mode 100644
diff --git a/Engine/Source/Programs/Horde/Horde.Build/Agents/Software/IAgentSoftwareChannel.cs b/Engine/Source/Programs/Horde/Horde.Build/Agents/Software/IAgentSoftwareChannel.cs
deleted file mode 100644
diff --git a/Engine/Source/Programs/Horde/Horde.Build/Agents/Software/IAgentSoftwareCollection.cs b/Engine/Source/Programs/Horde/Horde.Build/Agents/Software/IAgentSoftwareCollection.cs
deleted file mode 100644
diff --git a/Engine/Source/Programs/Horde/Horde.Build/Server/RpcService.cs b/Engine/Source/Programs/Horde/Horde.Build/Server/RpcService.cs
--- a/Engine/Source/Programs/Horde/Horde.Build/Server/RpcService.cs
+++ b/Engine/Source/Programs/Horde/Horde.Build/Server/RpcService.cs
@@ -417,15 +417,20 @@ namespace Horde.Build.Server
 		/// <returns>Information about the new agent</returns>
 		public override async Task DownloadSoftware(DownloadSoftwareRequest request, IServerStreamWriter<DownloadSoftwareResponse> responseStream, ServerCallContext context)
 		{
-			if (!_globalConfig.Value.Authorize(AclAction.DownloadSoftware, context.GetHttpContext().User))
-			{
-				throw new StructuredRpcException(StatusCode.NotFound, "Access to software is forbidden");
-			}
-
 			int colonIdx = request.Version.IndexOf(':', StringComparison.Ordinal);
 			ToolId toolId = new ToolId(request.Version.Substring(0, colonIdx));
 			string version = request.Version.Substring(colonIdx + 1);
 
+			ToolConfig? toolConfig;
+			if (!_globalConfig.Value.TryGetTool(toolId, out toolConfig))
+			{
+				throw new StructuredRpcException(StatusCode.NotFound, $"Missing tool {toolId}");
+			}
+			if (!toolConfig.Authorize(AclAction.DownloadTool, context.GetHttpContext().User))
+			{
+				throw new StructuredRpcException(StatusCode.NotFound, "Access to software is forbidden");
+			}
+
 			ITool? tool = await _toolCollection.GetAsync(toolId, _globalConfig.Value);
 			if (tool == null)
 			{
diff --git a/Engine/Source/Programs/Horde/Horde.Build/Startup.cs b/Engine/Source/Programs/Horde/Horde.Build/Startup.cs
--- a/Engine/Source/Programs/Horde/Horde.Build/Startup.cs
+++ b/Engine/Source/Programs/Horde/Horde.Build/Startup.cs
@@ -336,7 +336,6 @@ namespace Horde.Build
 			services.AddScoped<JobRpcCommon>();
 
 			services.AddSingleton<IAgentCollection, AgentCollection>();
-			services.AddSingleton<IAgentSoftwareCollection, AgentSoftwareCollection>();
 			services.AddSingleton<IArtifactCollection, ArtifactCollection>();
 			services.AddSingleton<IArtifactCollectionV1, ArtifactCollectionV1>();
 			services.AddSingleton<IGraphCollection, GraphCollection>();
@@ -394,8 +393,7 @@ namespace Horde.Build
 			services.AddHostedService<TunnelService>(sp => sp.GetRequiredService<TunnelService>());
 
 			services.AddSingleton<AclService>();
-			services.AddSingleton<AgentService>();			
-			services.AddSingleton<AgentSoftwareService>();
+			services.AddSingleton<AgentService>();
 			services.AddSingleton<ConsistencyService>();
 			services.AddSingleton<RequestTrackerService>();
 			services.AddSingleton<CredentialService>();
