Horde: Restore the agent software endpoint used to determine version for the dashboard.

diff --git a/Engine/Source/Programs/Horde/Horde.Build/Agents/Software/AgentSoftwareController.cs b/Engine/Source/Programs/Horde/Horde.Build/Agents/Software/AgentSoftwareController.cs
--- a/Engine/Source/Programs/Horde/Horde.Build/Agents/Software/AgentSoftwareController.cs
+++ b/Engine/Source/Programs/Horde/Horde.Build/Agents/Software/AgentSoftwareController.cs
@@ -7,6 +7,7 @@ using System.Threading.Tasks;
 using Horde.Build.Acls;
 using Horde.Build.Server;
 using Horde.Build.Tools;
+using Horde.Build.Utilities;
 using HordeCommon;
 using Microsoft.AspNetCore.Authorization;
 using Microsoft.AspNetCore.Mvc;
@@ -15,6 +16,17 @@ using Microsoft.Net.Http.Headers;
 
 namespace Horde.Build.Agents.Software
 {
+	/// <summary>
+	/// Information about an agent software channel
+	/// </summary>
+	public class GetAgentSoftwareChannelResponse
+	{
+		/// <summary>
+		/// Version number of this software
+		/// </summary>
+		public string? Version { get; set; }
+	}
+
 	/// <summary>
 	/// Controller for the /api/v1/agentsoftware endpoint
 	/// </summary>
@@ -37,6 +49,36 @@ namespace Horde.Build.Agents.Software
 			_globalConfig = globalConfig;
 		}
 
+		/// <summary>
+		/// Finds all uploaded software matching the given criteria
+		/// </summary>
+		/// <returns>Http response</returns>
+		[HttpGet]
+		[Obsolete("Agent software is now stored as a tool. This endpoint exists for backwards compatibility, but will be removed in the future.")]
+		[Route("/api/v1/agentsoftware/default")]
+		[ProducesResponseType(typeof(GetAgentSoftwareChannelResponse), 200)]
+		public async Task<ActionResult<object>> FindSoftwareAsync()
+		{
+			if (!_globalConfig.Value.Authorize(AclAction.DownloadSoftware, User))
+			{
+				return Forbid();
+			}
+
+			ITool? tool = await _toolCollection.GetAsync(AgentExtensions.DefaultAgentSoftwareToolId, _globalConfig.Value);
+			if (tool == null)
+			{
+				return NotFound("No agent software tool is currently registered");
+			}
+
+			IToolDeployment? deployment = tool.GetCurrentDeployment(1.0, _clock.UtcNow);
+			if (deployment == null)
+			{
+				return NotFound("No deployment currently set for agent software");
+			}
+
+			return new GetAgentSoftwareChannelResponse { Version = deployment.Version };
+		}
+
 		/// <summary>
 		/// Gets the zip file for a specific channel
 		/// </summary>
