Horde: Fix delay between software upgrade attempts.

diff --git a/Engine/Source/Programs/Horde/Horde.Build/Tasks/UpgradeTaskSource.cs b/Engine/Source/Programs/Horde/Horde.Build/Tasks/UpgradeTaskSource.cs
--- a/Engine/Source/Programs/Horde/Horde.Build/Tasks/UpgradeTaskSource.cs
+++ b/Engine/Source/Programs/Horde/Horde.Build/Tasks/UpgradeTaskSource.cs
@@ -61,7 +61,9 @@ namespace Horde.Build.Tasks
 			{
 				return Skip(cancellationToken);
 			}
-			if (agent.Leases.Count > 0 || !(agent.LastUpgradeTime == null || agent.LastUpgradeTime.Value + TimeSpan.FromMinutes(5.0) < _clock.UtcNow || agent.LastUpgradeVersion != deployment.Version))
+
+			string softwareId = $"{tool.Id}:{deployment.Version}";
+			if (agent.Leases.Count > 0 || !(agent.LastUpgradeTime == null || agent.LastUpgradeTime.Value + TimeSpan.FromMinutes(5.0) < _clock.UtcNow || agent.LastUpgradeVersion != softwareId))
 			{
 				return await DrainAsync(cancellationToken);
 			}
@@ -70,7 +72,7 @@ namespace Horde.Build.Tasks
 			ILogFile logFile = await _logService.CreateLogFileAsync(JobId.Empty, leaseId, agent.SessionId, LogType.Json, useNewStorageBackend: false, cancellationToken: cancellationToken);
 
 			UpgradeTask task = new UpgradeTask();
-			task.SoftwareId = $"{tool.Id}:{deployment.Version}";
+			task.SoftwareId = softwareId;
 			task.LogId = logFile.Id.ToString();
 
 			byte[] payload = Any.Pack(task).ToByteArray();
