Horde: Support richer Perforce triggers

Adds cluster back to path.

diff --git a/Engine/Source/Programs/Horde/Horde.Build/Perforce/PerforceController.cs b/Engine/Source/Programs/Horde/Horde.Build/Perforce/PerforceController.cs
--- a/Engine/Source/Programs/Horde/Horde.Build/Perforce/PerforceController.cs
+++ b/Engine/Source/Programs/Horde/Horde.Build/Perforce/PerforceController.cs
@@ -3,6 +3,7 @@
 using System;
 using System.Collections.Generic;
 using System.Net;
+using System.Text.Json.Serialization;
 using System.Threading.Tasks;
 using Horde.Build.Acls;
 using Horde.Build.Server;
@@ -71,58 +72,177 @@ namespace Horde.Build.Perforce
 	}
 
 	/// <summary>
-	/// Perforce trigger received as JSON
+	/// Body of Perforce trigger request (deserialized from JSON)
 	/// </summary>
-	public class PerforceTriggerPayload
+	public class PerforceTriggerRequest
 	{
 		/// <summary>
-		/// Type of trigger (change-commit etc)
+		/// Type of trigger (change-commit, form-save etc)
 		/// </summary>
 		public string TriggerType { get; }
 		
 		/// <summary>
-		/// Author of changelist (Perforce user) 
+		/// Triggering user’s client workspace name.
+		/// </summary>
+		public string Client { get; }
+		
+		/// <summary>
+		/// Hostname of the user’s workstation (even if connected through a proxy, broker, replica, or an edge server.)
+		/// </summary>
+		public string ClientHost { get; }
+		
+		/// <summary>
+		/// The IP address of the user’s workstation (even if connected through a proxy, broker, replica, or an edge server.)
+		/// </summary>
+		public string ClientIp { get; }
+		
+		/// <summary>
+		/// The name of the user’s client application. For example, P4V, P4Win
+		/// </summary>
+		public string ClientProg { get;}
+		
+		/// <summary>
+		/// The version of the user’s client application.
+		/// </summary>
+		public string ClientVersion { get; }
+		
+		/// <summary>
+		/// If the command was sent through a proxy, broker, replica, or edge server, the hostname of the proxy, broker, replica, or edge server.
+		/// (If the command was sent directly, %peerhost% matches %clienthost%)
+		/// </summary>
+		public string PeerHost { get; }
+		
+		/// <summary>
+		/// If the command was sent through a proxy, broker, replica, or edge server, the IP address of the proxy, broker, replica, or edge server.
+		/// (If the command was sent directly, %peerip% matches %clientip%)
+		/// </summary>
+		public string PeerIp { get; }
+		
+		/// <summary>
+		/// Hostname of the Helix Core Server.
+		/// </summary>
+		public string ServerHost { get; }
+		
+		/// <summary>
+		/// The value of the Helix Core Server’s server.id. See p4 serverid in the Helix Core Command-Line (P4) Reference.
+		/// </summary>
+		public string ServerId { get; }
+		
+		/// <summary>
+		/// The IP address of the server.
+		/// </summary>
+		public string ServerIp { get; }
+		
+		/// <summary>
+		/// The value of the Helix Core Server’s P4NAME.
+		/// </summary>
+		public string ServerName { get; }
+		
+		/// <summary>
+		/// The transport, IP address, and port of the Helix Core Server, in the format prefix:ip_address:port.
+		/// </summary>
+		public string ServerPort { get; }
+		
+		/// <summary>
+		/// In a distributed installation, for any change trigger:
+		///     If the submit was run on the commit server, %submitserverid% equals %serverid%.
+		///     If the submit was run on the edge server, %submitserverid% does not equal %serverid%. In this case, %submitserverid% holds the edge server’s server id.
+		/// If this is not a distributed installation, %submitserverid% is always empty.
+		/// </summary>
+		public string? SubmitServerId { get; }
+		
+		/// <summary>
+		/// Helix Server username of the triggering user.
 		/// </summary>
 		public string User { get; }
 		
 		/// <summary>
-		/// Change number
+		/// Name of form (for instance, a branch name or a changelist number).
 		/// </summary>
-		public string Changelist { get; }
+		public string? FormName { get; }
 		
 		/// <summary>
-		/// Change number
+		/// Type of form (for instance, branch, change, and so on).
+		/// </summary>
+		public string? FormType { get; }
+		
+		/// <summary>
+		/// The number of the changelist being submitted. Not set for form-save.
 		/// </summary>
-		public int ChangelistNumber { get; }
+		[JsonPropertyName("ChangeNumber")]
+		public string? ChangeNumberString { get; }
 		
 		/// <summary>
-		/// Change root of changelist
+		/// The root path of files submitted.
 		/// </summary>
-		public string ChangeRoot { get; }
+		public string? ChangeRoot { get; }
+		
+		/// <summary>
+		/// Change number
+		/// Normalized as form-save stores the change number in 'formname'
+		/// </summary>
+		[JsonIgnore]
+		public int ChangeNumber { get; }
 
 		/// <summary>
 		/// Constructor
 		/// </summary>
 		/// <param name="triggerType"></param>
+		/// <param name="client"></param>
+		/// <param name="clientHost"></param>
+		/// <param name="clientIp"></param>
+		/// <param name="clientProg"></param>
+		/// <param name="clientVersion"></param>
+		/// <param name="peerHost"></param>
+		/// <param name="peerIp"></param>
+		/// <param name="serverHost"></param>
+		/// <param name="serverId"></param>
+		/// <param name="serverIp"></param>
+		/// <param name="serverName"></param>
+		/// <param name="serverPort"></param>
+		/// <param name="submitServerId"></param>
 		/// <param name="user"></param>
-		/// <param name="changelist"></param>
+		/// <param name="formName"></param>
+		/// <param name="formType"></param>
+		/// <param name="changeNumberString"></param>
 		/// <param name="changeRoot"></param>
-		public PerforceTriggerPayload(string triggerType, string user, string changelist, string changeRoot)
+		public PerforceTriggerRequest(string triggerType, string client, string clientHost, string clientIp, string clientProg, string clientVersion, string peerHost, string peerIp, string serverHost, string serverId, string serverIp, string serverName, string serverPort, string submitServerId, string user, string formName, string formType, string changeNumberString, string changeRoot)
 		{
 			TriggerType = triggerType;
+			Client = client;
+			ClientHost = clientHost;
+			ClientIp = clientIp;
+			ClientProg = clientProg;
+			ClientVersion = clientVersion;
+			PeerHost = peerHost;
+			PeerIp = peerIp;
+			ServerHost = serverHost;
+			ServerId = serverId;
+			ServerIp = serverIp;
+			ServerName = serverName;
+			ServerPort = serverPort;
+			SubmitServerId = submitServerId;
 			User = user;
-			Changelist = changelist;
-			ChangelistNumber = Convert.ToInt32(changelist);
+			FormName = formName;
+			FormType = formType;
+			ChangeNumberString = changeNumberString;
 			ChangeRoot = changeRoot;
+
+			try
+			{
+				ChangeNumber = TriggerType == "form-save" ? Convert.ToInt32(FormName) : Convert.ToInt32(ChangeNumberString);
+			}
+			catch (Exception)
+			{
+				// form-save can contain the change number "default"
+				ChangeNumber = -1;
+			}
 		}
 
-		/// <summary>
-		/// Format as a string for debugging purposes
-		/// </summary>
-		/// <returns></returns>
+		/// <inheritdoc/>
 		public override string ToString()
 		{
-			return $"TriggerType={TriggerType} User={User} CL={Changelist} Root={ChangeRoot}";
+			return $"TriggerType={TriggerType} User={User} ChangeNumber={ChangeNumber}";
 		}
 	}
 	
@@ -150,12 +270,19 @@ namespace Horde.Build.Perforce
 		/// </summary>
 		/// <returns>200 OK on success</returns>
 		[HttpPost]
-		[Route("/api/v1/perforce/trigger")]
-		public ActionResult TriggerCallback([FromBody]PerforceTriggerPayload payload)
+		[Route("/api/v1/perforce/{cluster}/trigger")]
+		public async Task<ActionResult> TriggerCallback(string cluster, [FromBody]PerforceTriggerRequest trigger)
 		{
-			// Currently just a placeholder until correct triggers are in place.
 			_logger.LogDebug("Received Perforce trigger callback. Type={Type} CL={Changelist} User={User} Root={Root}", 
-				payload.TriggerType, payload.Changelist, payload.User, payload.ChangeRoot);
+				trigger.TriggerType, trigger.ChangeNumber, trigger.User, trigger.ChangeRoot);
+			
+			// For "form-save" triggers, change number can be -1 due to variable "formname" is set to "default" (non-submitted changelist)
+			if (trigger.TriggerType == "change-commit" && trigger.ChangeNumber != -1)
+			{
+				// Not implemented yet
+				await Task.Delay(0); // Avoid await warnings
+				// await _perforceService.RefreshCachedCommitAsync(cluster, trigger.ChangeNumber);
+			}
 
 			string content = "{\"message\": \"Trigger received\"}";
 			return new ContentResult { ContentType = "application/json", StatusCode = (int)HttpStatusCode.OK, Content = content };
