Horde: Track timestamps for status changes to agents

diff --git a/Engine/Source/Programs/Horde/Horde.Build.Tests/AgentServiceTest.cs b/Engine/Source/Programs/Horde/Horde.Build.Tests/AgentServiceTest.cs
--- a/Engine/Source/Programs/Horde/Horde.Build.Tests/AgentServiceTest.cs
+++ b/Engine/Source/Programs/Horde/Horde.Build.Tests/AgentServiceTest.cs
@@ -13,6 +13,7 @@ using Horde.Build.Jobs;
 using Horde.Build.Server;
 using Horde.Build.Utilities;
 using HordeCommon;
+using HordeCommon.Rpc.Messages;
 using Microsoft.AspNetCore.Mvc;
 using Microsoft.Extensions.DependencyInjection;
 using Microsoft.VisualStudio.TestTools.UnitTesting;
@@ -58,6 +59,76 @@ public class AgentServiceTest : TestSetup
 		Assert.IsFalse(AgentService.AuthorizeSession(agent, GetUser(agent)));
 	}
 	
+	private static long ToUnixTime(DateTime dateTime)
+	{
+		return new DateTimeOffset(dateTime).ToUnixTimeSeconds();
+	}
+	
+	[TestMethod]
+	public async Task LastStatusChangeDuringSessionCreate()
+	{
+		// No session created yet, status change timestamp is empty
+		IAgent agent = await AgentService.CreateAgentAsync("agent1", true, null);
+		Assert.AreEqual(AgentStatus.Unspecified, agent.Status);
+		Assert.IsFalse(agent.LastStatusChange.HasValue);
+		
+		// A session has been created, status change timestamp is current time
+		agent = await AgentService.CreateSessionAsync(agent, AgentStatus.Ok, new List<string>(), new Dictionary<string, int>(), "v1");
+		Assert.AreEqual(AgentStatus.Ok, agent.Status);
+		Assert.AreEqual(ToUnixTime(Clock.UtcNow), ToUnixTime(agent.LastStatusChange!.Value));
+		DateTime lastStatusChange = agent.LastStatusChange!.Value;
+		
+		await Clock.AdvanceAsync(TimeSpan.FromMinutes(1));
+		
+		// The session is re-created, status change timestamp is same as when it first got created
+		agent = await AgentService.CreateSessionAsync(agent, AgentStatus.Ok, new List<string>(), new Dictionary<string, int>(), "v1");
+		Assert.AreEqual(AgentStatus.Ok, agent.Status);
+		Assert.AreEqual(ToUnixTime(lastStatusChange), ToUnixTime(agent.LastStatusChange!.Value));
+	}
+
+	private static int s_agentId = 1;
+	private async Task<IAgent> CreateAgentSession(bool enabled = true)
+	{
+		IAgent agent = await AgentService.CreateAgentAsync("agentServiceTest-" + s_agentId++, true, null);
+		agent = await AgentService.CreateSessionAsync(agent, AgentStatus.Ok, new List<string>(), new Dictionary<string, int>(), "v1");
+		return agent;
+	}
+
+	private static void AssertEqual(DateTime expected, DateTime? actual)
+	{
+		// When saved as MongoDB documents, some precision is lost. This compares only Unix seconds.
+		Assert.AreEqual(ToUnixTime(expected), ToUnixTime(actual!.Value));
+	}
+	
+	private static void AssertNotEqual(DateTime expected, DateTime? actual)
+	{
+		// When saved as MongoDB documents, some precision is lost. This compares only Unix seconds.
+		Assert.AreNotEqual(ToUnixTime(expected), ToUnixTime(actual!.Value));
+	}
+
+	[TestMethod]
+	[DataRow(AgentStatus.Ok, false)]
+	[DataRow(AgentStatus.Unhealthy, true)]
+	[DataRow(AgentStatus.Stopping, true)]
+	[DataRow(AgentStatus.Stopped, true)]
+	[DataRow(AgentStatus.Unspecified, true)]
+	public async Task LastStatusChange(AgentStatus status, bool expectTimestampUpdate)
+	{
+		IAgent agent = await CreateAgentSession();
+		DateTime lastStatusChange = agent.LastStatusChange!.Value;
+		await Clock.AdvanceAsync(TimeSpan.FromMinutes(1));
+
+		agent = (await AgentService.UpdateSessionAsync(agent, agent.SessionId!.Value, status, null, null, new List<Lease>()))!;
+		if (expectTimestampUpdate)
+		{
+			AssertNotEqual(lastStatusChange, agent.LastStatusChange);
+		}
+		else
+		{
+			AssertEqual(lastStatusChange, agent.LastStatusChange);
+		}
+	}
+	
 	[TestMethod]
 	public async Task AuditLogAwsInstanceTypeChanges()
 	{
diff --git a/Engine/Source/Programs/Horde/Horde.Build/Agents/AgentCollection.cs b/Engine/Source/Programs/Horde/Horde.Build/Agents/AgentCollection.cs
--- a/Engine/Source/Programs/Horde/Horde.Build/Agents/AgentCollection.cs
+++ b/Engine/Source/Programs/Horde/Horde.Build/Agents/AgentCollection.cs
@@ -62,6 +62,7 @@ namespace Horde.Build.Agents
 			public DateTime? SessionExpiresAt { get; set; }
 
 			public AgentStatus Status { get; set; }
+			public DateTime? LastStatusChange { get; set; }
 
 			[BsonRequired]
 			public bool Enabled { get; set; }
@@ -136,18 +137,20 @@ namespace Horde.Build.Agents
 		readonly IMongoCollection<AgentDocument> _agents;
 		readonly IAuditLog<AgentId> _auditLog;
 		readonly RedisService _redisService;
+		readonly IClock _clock;
 		readonly RedisChannel<AgentId> _updateEventChannel;
 
 		/// <summary>
 		/// Constructor
 		/// </summary>
-		public AgentCollection(MongoService mongoService, RedisService redisService, IAuditLog<AgentId> auditLog)
+		public AgentCollection(MongoService mongoService, RedisService redisService, IClock clock, IAuditLog<AgentId> auditLog)
 		{
 			List<MongoIndex<AgentDocument>> indexes = new List<MongoIndex<AgentDocument>>();
 			indexes.Add(keys => keys.Ascending(x => x.Deleted).Ascending(x => x.Id).Ascending(x => x.Pools));
 
 			_agents = mongoService.GetCollection<AgentDocument>("Agents", indexes);
 			_redisService = redisService;
+			_clock = clock;
 			_updateEventChannel = new RedisChannel<AgentId>("agents/notify");
 			_auditLog = auditLog;
 		}
@@ -360,6 +363,7 @@ namespace Horde.Build.Agents
 			if (status != null && agent.Status != status.Value)
 			{
 				updates.Add(updateBuilder.Set(x => x.Status, status.Value));
+				updates.Add(updateBuilder.Set(x => x.LastStatusChange, _clock.UtcNow));
 			}
 			if (sessionExpiresAt != null)
 			{
@@ -455,7 +459,7 @@ namespace Horde.Build.Agents
 		}
 
 		/// <inheritdoc/>
-		public async Task<IAgent?> TryStartSessionAsync(IAgent agentInterface, SessionId sessionId, DateTime sessionExpiresAt, AgentStatus status, IReadOnlyList<string> properties, IReadOnlyDictionary<string, int> resources, IReadOnlyList<PoolId> pools, IReadOnlyList<PoolId> dynamicPools, string? version)
+		public async Task<IAgent?> TryStartSessionAsync(IAgent agentInterface, SessionId sessionId, DateTime sessionExpiresAt, AgentStatus status, IReadOnlyList<string> properties, IReadOnlyDictionary<string, int> resources, IReadOnlyList<PoolId> pools, IReadOnlyList<PoolId> dynamicPools, DateTime lastStatusChange, string? version)
 		{
 			AgentDocument agent = (AgentDocument)agentInterface;
 			List<string> newProperties = properties.OrderBy(x => x, StringComparer.OrdinalIgnoreCase).ToList();
@@ -481,6 +485,11 @@ namespace Horde.Build.Agents
 			updates.Add(updateBuilder.Unset(x => x.RequestShutdown));
 			updates.Add(updateBuilder.Set(x => x.LastShutdownReason, "Unexpected"));
 
+			if (agent.Status != status)
+			{
+				updates.Add(updateBuilder.Set(x => x.LastStatusChange, lastStatusChange));
+			}
+
 			// Apply the update
 			return await TryUpdateAsyncWithRetries(agent, updateBuilder.Combine(updates));
 		}
@@ -495,6 +504,7 @@ namespace Horde.Build.Agents
 			update = update.Unset(x => x.SessionExpiresAt);
 			update = update.Unset(x => x.Leases);
 			update = update.Set(x => x.Status, AgentStatus.Stopped);
+			update = update.Set(x => x.LastStatusChange, _clock.UtcNow);
 
 			bool deleted = agent.Deleted || agent.Ephemeral;
 			if (deleted != agent.Deleted)
diff --git a/Engine/Source/Programs/Horde/Horde.Build/Agents/AgentMessages.cs b/Engine/Source/Programs/Horde/Horde.Build/Agents/AgentMessages.cs
--- a/Engine/Source/Programs/Horde/Horde.Build/Agents/AgentMessages.cs
+++ b/Engine/Source/Programs/Horde/Horde.Build/Agents/AgentMessages.cs
@@ -409,6 +409,11 @@ namespace Horde.Build.Agents
 		/// Last update time of this agent
 		/// </summary>
 		public DateTime? UpdateTime { get; set; }
+		
+		/// <summary>
+		/// Last time agent's status was changed
+		/// </summary>
+		public DateTime? LastStatusChange { get; set; }
 
 		/// <summary>
 		/// Pools for this agent
@@ -463,6 +468,7 @@ namespace Horde.Build.Agents
 			Version = agent.Version?.ToString() ?? "Unknown";
 			Version = agent.Version?.ToString();
 			UpdateTime = agent.UpdateTime;
+			LastStatusChange = agent.LastStatusChange;
 			Pools = agent.GetPools().Select(x => x.ToString()).ToList();
 			Workspaces = agent.Workspaces.ConvertAll(x => new GetAgentWorkspaceResponse(x));
 			Capabilities = new { Devices = new[] { new { agent.Properties, agent.Resources } } };
diff --git a/Engine/Source/Programs/Horde/Horde.Build/Agents/AgentService.cs b/Engine/Source/Programs/Horde/Horde.Build/Agents/AgentService.cs
--- a/Engine/Source/Programs/Horde/Horde.Build/Agents/AgentService.cs
+++ b/Engine/Source/Programs/Horde/Horde.Build/Agents/AgentService.cs
@@ -308,6 +308,7 @@ namespace Horde.Build.Agents
 		/// <returns>New agent state</returns>
 		public async Task<IAgent> CreateSessionAsync(IAgent agent, AgentStatus status, IReadOnlyList<string> properties, IReadOnlyDictionary<string, int> resources, string? version)
 		{
+			DateTime? lastStatusChange = null;
 			for (; ; )
 			{
 				IAuditLogChannel<AgentId> agentLogger = Agents.GetLogger(agent.Id);
@@ -316,6 +317,9 @@ namespace Horde.Build.Agents
 				IAgent? newAgent;
 				if (agent.SessionId != null)
 				{
+					// Save last status change timestamp to avoid registering the change to "stopped" when it's re-created immediately after to "ok".
+					lastStatusChange = agent.LastStatusChange;
+					
 					// Try to terminate the current session
 					await TryTerminateSessionAsync(agent);
 				}
@@ -339,7 +343,7 @@ namespace Horde.Build.Agents
 					List<PoolId> pools = CombineCurrentAndRequestedPools(agent.ExplicitPools, properties);
 
 					// Reset the agent to use the new session
-					newAgent = await Agents.TryStartSessionAsync(agent, newSession.Id, sessionExpiresAt, status, properties, resources, pools, dynamicPools, version);
+					newAgent = await Agents.TryStartSessionAsync(agent, newSession.Id, sessionExpiresAt, status, properties, resources, pools, dynamicPools, lastStatusChange ?? utcNow, version);
 					if(newAgent != null)
 					{
 						LogPropertyChanges(agentLogger, agent.Properties, newAgent.Properties);
@@ -637,11 +641,14 @@ namespace Horde.Build.Agents
 					// If the session is valid, we can terminate once the agent leases are also empty
 					if (agent.Leases.Count == 0)
 					{
-						if (!await TryTerminateSessionAsync(agent))
+						IAgent? terminatedAgent = await TryTerminateSessionAsync(agent);
+						if (terminatedAgent == null)
 						{
 							agent = await GetAgentAsync(agent.Id);
 							continue;
 						}
+
+						agent = terminatedAgent;
 						break;
 					}
 				}
@@ -739,13 +746,13 @@ namespace Horde.Build.Agents
 		/// Terminates an existing session. Does not update the agent itself, if it's currently 
 		/// </summary>
 		/// <param name="agent">The agent whose current session should be terminated</param>
-		/// <returns>True if the session was terminated</returns>
-		private async Task<bool> TryTerminateSessionAsync(IAgent agent)
+		/// <returns>An up-to-date IAgent if the session was terminated</returns>
+		private async Task<IAgent?> TryTerminateSessionAsync(IAgent agent)
 		{
 			// Make sure the agent has a valid session id
 			if (agent.SessionId == null)
 			{
-				return true;
+				return agent;
 			}
 
 			// Get the time that the session finishes at
@@ -775,9 +782,9 @@ namespace Horde.Build.Agents
 				// Update the session document
 				Agents.GetLogger(agent.Id).LogInformation("Terminated session {SessionId}", sessionId);
 				await _sessions.UpdateAsync(sessionId, finishTime, agent.Properties, agent.Resources);
-				return true;
+				return agent;
 			}
-			return false;
+			return null;
 		}
 
 		/// <summary>
diff --git a/Engine/Source/Programs/Horde/Horde.Build/Agents/IAgent.cs b/Engine/Source/Programs/Horde/Horde.Build/Agents/IAgent.cs
--- a/Engine/Source/Programs/Horde/Horde.Build/Agents/IAgent.cs
+++ b/Engine/Source/Programs/Horde/Horde.Build/Agents/IAgent.cs
@@ -384,6 +384,11 @@ namespace Horde.Build.Agents
 		/// Current status of this agent
 		/// </summary>
 		public AgentStatus Status { get; }
+		
+		/// <summary>
+		/// Time at which last status change took place.
+		/// </summary>
+		public DateTime? LastStatusChange { get; }
 
 		/// <summary>
 		/// Whether the agent is enabled
diff --git a/Engine/Source/Programs/Horde/Horde.Build/Agents/IAgentCollection.cs b/Engine/Source/Programs/Horde/Horde.Build/Agents/IAgentCollection.cs
--- a/Engine/Source/Programs/Horde/Horde.Build/Agents/IAgentCollection.cs
+++ b/Engine/Source/Programs/Horde/Horde.Build/Agents/IAgentCollection.cs
@@ -101,9 +101,10 @@ namespace Horde.Build.Agents
 		/// <param name="resources">Resources for the agent</param>
 		/// <param name="pools">New list of pools for the agent</param>
 		/// <param name="dynamicPools">New list of dynamic pools for the agent</param>
+		/// <param name="lastStatusChange">Time to force status change timestamp to</param>
 		/// <param name="version">Current version of the agent software</param>
 		/// <returns>New agent state</returns>
-		Task<IAgent?> TryStartSessionAsync(IAgent agent, SessionId sessionId, DateTime sessionExpiresAt, AgentStatus status, IReadOnlyList<string> properties, IReadOnlyDictionary<string, int> resources, IReadOnlyList<PoolId> pools, IReadOnlyList<PoolId> dynamicPools, string? version);
+		Task<IAgent?> TryStartSessionAsync(IAgent agent, SessionId sessionId, DateTime sessionExpiresAt, AgentStatus status, IReadOnlyList<string> properties, IReadOnlyDictionary<string, int> resources, IReadOnlyList<PoolId> pools, IReadOnlyList<PoolId> dynamicPools, DateTime lastStatusChange, string? version);
 
 		/// <summary>
 		/// Attempt to update the agent state
