Horde: Fix issue not being able to dismiss agent modal on agents view, when it is initially populated
#rnx
#skipci

diff --git a/Engine/Source/Programs/Horde/HordeDashboard/src/components/AgentView.tsx b/Engine/Source/Programs/Horde/HordeDashboard/src/components/AgentView.tsx
--- a/Engine/Source/Programs/Horde/HordeDashboard/src/components/AgentView.tsx
+++ b/Engine/Source/Programs/Horde/HordeDashboard/src/components/AgentView.tsx
@@ -592,6 +592,7 @@ class LocalState {
       state.exactSearch = exact?.trim() === "true" ? true : undefined;
       state.agentId = agentId?.trim() ? agentId : undefined;
 
+      this.search = search;
       this.searchState = state;
 
       if (state.agentSearch) {
