Horde: Add an endpoint to query summary information for all tools.

diff --git a/Engine/Source/Programs/Horde/Horde.Build/Tools/ToolsController.cs b/Engine/Source/Programs/Horde/Horde.Build/Tools/ToolsController.cs
--- a/Engine/Source/Programs/Horde/Horde.Build/Tools/ToolsController.cs
+++ b/Engine/Source/Programs/Horde/Horde.Build/Tools/ToolsController.cs
@@ -54,6 +54,42 @@ namespace Horde.Build.Tools
 		}
 	}
 
+	/// <summary>
+	/// Summary for a particular tool.
+	/// </summary>
+	public class GetToolSummaryResponse
+	{
+		readonly ITool _tool;
+
+		/// <inheritdoc cref="VersionedDocument{TId, TLatest}.Id"/>
+		public ToolId Id => _tool.Id;
+
+		/// <inheritdoc cref="ToolConfig.Name"/>
+		public string Name => _tool.Config.Name;
+
+		/// <inheritdoc cref="ToolConfig.Description"/>
+		public string Description => _tool.Config.Description;
+
+		/// <inheritdoc cref="IToolDeployment.Version"/>
+		public string? Version => (_tool.Deployments.Count > 0) ? _tool.Deployments[^1].Version : null;
+
+		/// <summary>
+		/// Constructor
+		/// </summary>
+		internal GetToolSummaryResponse(ITool tool) => _tool = tool;
+	}
+
+	/// <summary>
+	/// Response when querying all tools
+	/// </summary>
+	public class GetToolsSummaryResponse
+	{
+		/// <summary>
+		/// List of tools.
+		/// </summary>
+		public List<GetToolSummaryResponse> Tools { get; } = new List<GetToolSummaryResponse>();
+	}
+
 	/// <summary>
 	/// Response object describing the deployment of a tool
 	/// </summary>
@@ -243,6 +279,32 @@ namespace Horde.Build.Tools
 			_globalConfig = globalConfig;
 		}
 
+		/// <summary>
+		/// Create a new deployment of the given tool.
+		/// </summary>
+		/// <returns>Information about the registered agent</returns>
+		[HttpGet]
+		[Route("/api/v1/tools")]
+		public async Task<ActionResult<GetToolsSummaryResponse>> GetToolsAsync()
+		{
+			GlobalConfig globalConfig = _globalConfig.Value;
+
+			GetToolsSummaryResponse response = new GetToolsSummaryResponse();
+			foreach (ToolConfig toolConfig in globalConfig.Tools.OrderBy(x => x.Name, StringComparer.Ordinal))
+			{
+				if(AuthorizeDownload(toolConfig))
+				{
+					ITool? tool = await _toolCollection.GetAsync(toolConfig.Id, _globalConfig.Value);
+					if (tool != null)
+					{
+						response.Tools.Add(new GetToolSummaryResponse(tool));
+					}
+				}
+			}
+
+			return response;
+		}
+
 		/// <summary>
 		/// Gets information about a particular tool
 		/// </summary>
@@ -256,7 +318,7 @@ namespace Horde.Build.Tools
 			{
 				return NotFound(id);
 			}
-			if (!AuthorizeDownload(tool))
+			if (!AuthorizeDownload(tool.Config))
 			{
 				return Forbid(AclAction.DownloadTool, id);
 			}
@@ -299,7 +361,7 @@ namespace Horde.Build.Tools
 			{
 				return NotFound(id);
 			}
-			if (!AuthorizeDownload(tool))
+			if (!AuthorizeDownload(tool.Config))
 			{
 				return Forbid(AclAction.DownloadTool, id);
 			}
@@ -326,7 +388,7 @@ namespace Horde.Build.Tools
 			{
 				return NotFound(id);
 			}
-			if (!AuthorizeDownload(tool))
+			if (!AuthorizeDownload(tool.Config))
 			{
 				return Forbid(AclAction.DownloadTool, id);
 			}
@@ -380,7 +442,7 @@ namespace Horde.Build.Tools
 			{
 				return NotFound(id);
 			}
-			if (!AuthorizeDownload(tool))
+			if (!AuthorizeDownload(tool.Config))
 			{
 				return Forbid(AclAction.DownloadTool, id);
 			}
@@ -393,15 +455,15 @@ namespace Horde.Build.Tools
 			return StorageController.ReadBlobInternalAsync(_storageService, Namespace.Tools, locator, offset, length, cancellationToken);
 		}
 
-		bool AuthorizeDownload(ITool tool)
+		bool AuthorizeDownload(ToolConfig toolConfig)
 		{
-			if (!tool.Config.Public)
+			if (!toolConfig.Public)
 			{
 				if (User.Identity == null || !User.Identity.IsAuthenticated)
 				{
 					return false;
 				}
-				if (!tool.Config.Authorize(AclAction.DownloadTool, User))
+				if (!toolConfig.Authorize(AclAction.DownloadTool, User))
 				{
 					return false;
 				}
