Horde: Replace history on filter changes in agents view
#skipci

diff --git a/Engine/Source/Programs/Horde/HordeDashboard/src/components/AgentView.tsx b/Engine/Source/Programs/Horde/HordeDashboard/src/components/AgentView.tsx
--- a/Engine/Source/Programs/Horde/HordeDashboard/src/components/AgentView.tsx
+++ b/Engine/Source/Programs/Horde/HordeDashboard/src/components/AgentView.tsx
@@ -1717,7 +1717,7 @@ export const SearchUpdate: React.FC = observer(() => {
    const csearch = localState.search.toString();
 
    if (state.search !== csearch) {
-      setSearchParams(csearch);
+      setSearchParams(csearch, {replace: true});
       setState({ search: csearch });
    }
 
