Horde: Add an interface for ToolCollection, for consistency with other services.

diff --git a/Engine/Source/Programs/Horde/Horde.Build.Tests/TestSetup.cs b/Engine/Source/Programs/Horde/Horde.Build.Tests/TestSetup.cs
--- a/Engine/Source/Programs/Horde/Horde.Build.Tests/TestSetup.cs
+++ b/Engine/Source/Programs/Horde/Horde.Build.Tests/TestSetup.cs
@@ -218,13 +218,12 @@ namespace Horde.Build.Tests
 			services.AddSingleton<ITestDataCollection, TestDataCollection>();
 			services.AddSingleton<ITelemetryCollection, TelemetryCollection>();
 			services.AddSingleton<ITemplateCollection, TemplateCollection>();
+			services.AddSingleton<IToolCollection, ToolCollection>();
 			services.AddSingleton<IUgsMetadataCollection, UgsMetadataCollection>();
 			services.AddSingleton<IUserCollection, UserCollectionV1>();
 			services.AddSingleton<IDeviceCollection, DeviceCollection>();
 			services.AddSingleton<IDashboardPreviewCollection, DashboardPreviewCollection>();
 
-			services.AddSingleton<ToolCollection>();
-
 			services.AddSingleton<FakeClock>();
 			services.AddSingleton<IClock>(sp => sp.GetRequiredService<FakeClock>());
 			services.AddSingleton<IHostApplicationLifetime, AppLifetimeStub>();
diff --git a/Engine/Source/Programs/Horde/Horde.Build.Tests/ToolTests.cs b/Engine/Source/Programs/Horde/Horde.Build.Tests/ToolTests.cs
--- a/Engine/Source/Programs/Horde/Horde.Build.Tests/ToolTests.cs
+++ b/Engine/Source/Programs/Horde/Horde.Build.Tests/ToolTests.cs
@@ -36,7 +36,7 @@ namespace Horde.Build.Tests
         [TestMethod]
         public async Task AddTool()
         {
-			ToolCollection collection = ServiceProvider.GetRequiredService<ToolCollection>();
+			IToolCollection collection = ServiceProvider.GetRequiredService<IToolCollection>();
 
 			ITool tool = Deref(await collection.GetAsync(_toolId, GlobalConfig.CurrentValue));
 			Assert.AreEqual(tool.Id, new ToolId("ugs"));
@@ -52,7 +52,7 @@ namespace Horde.Build.Tests
 		[TestMethod]
 		public async Task AddDeployment()
 		{
-			ToolCollection collection = ServiceProvider.GetRequiredService<ToolCollection>();
+			IToolCollection collection = ServiceProvider.GetRequiredService<IToolCollection>();
 
 			ITool tool = Deref(await collection.GetAsync(_toolId, GlobalConfig.CurrentValue));
 			Assert.AreEqual(new ToolId("ugs"), tool.Id);
diff --git a/Engine/Source/Programs/Horde/Horde.Build/Agents/Software/AgentSoftwareController.cs b/Engine/Source/Programs/Horde/Horde.Build/Agents/Software/AgentSoftwareController.cs
--- a/Engine/Source/Programs/Horde/Horde.Build/Agents/Software/AgentSoftwareController.cs
+++ b/Engine/Source/Programs/Horde/Horde.Build/Agents/Software/AgentSoftwareController.cs
@@ -35,14 +35,14 @@ namespace Horde.Build.Agents.Software
 	[Route("[controller]")]
 	public class AgentSoftwareController : ControllerBase
 	{
-		private readonly ToolCollection _toolCollection;
+		private readonly IToolCollection _toolCollection;
 		private readonly IClock _clock;
 		private readonly IOptionsSnapshot<GlobalConfig> _globalConfig;
 
 		/// <summary>
 		/// Constructor
 		/// </summary>
-		public AgentSoftwareController(ToolCollection toolCollection, IClock clock, IOptionsSnapshot<GlobalConfig> globalConfig)
+		public AgentSoftwareController(IToolCollection toolCollection, IClock clock, IOptionsSnapshot<GlobalConfig> globalConfig)
 		{
 			_toolCollection = toolCollection;
 			_clock = clock;
diff --git a/Engine/Source/Programs/Horde/Horde.Build/Server/RpcService.cs b/Engine/Source/Programs/Horde/Horde.Build/Server/RpcService.cs
--- a/Engine/Source/Programs/Horde/Horde.Build/Server/RpcService.cs
+++ b/Engine/Source/Programs/Horde/Horde.Build/Server/RpcService.cs
@@ -52,14 +52,14 @@ namespace Horde.Build.Server
 		readonly LifetimeService _lifetimeService;
 		readonly ConformTaskSource _conformTaskSource;
 		readonly JobRpcCommon _jobRpcCommon;
-		readonly ToolCollection _toolCollection;
+		readonly IToolCollection _toolCollection;
 		readonly IOptionsSnapshot<GlobalConfig> _globalConfig;
 		readonly ILogger _logger;
 
 		/// <summary>
 		/// Constructor
 		/// </summary>
-		public RpcService(AgentService agentService, PoolService poolService, LifetimeService lifetimeService, ConformTaskSource conformTaskSource, JobRpcCommon jobRpcCommon, ToolCollection toolCollection, IOptionsSnapshot<GlobalConfig> globalConfig, ILogger<RpcService> logger)
+		public RpcService(AgentService agentService, PoolService poolService, LifetimeService lifetimeService, ConformTaskSource conformTaskSource, JobRpcCommon jobRpcCommon, IToolCollection toolCollection, IOptionsSnapshot<GlobalConfig> globalConfig, ILogger<RpcService> logger)
 		{
 			_agentService = agentService;
 			_poolService = poolService;
diff --git a/Engine/Source/Programs/Horde/Horde.Build/Startup.cs b/Engine/Source/Programs/Horde/Horde.Build/Startup.cs
--- a/Engine/Source/Programs/Horde/Horde.Build/Startup.cs
+++ b/Engine/Source/Programs/Horde/Horde.Build/Startup.cs
@@ -362,7 +362,7 @@ namespace Horde.Build
 			services.AddSingleton<INoticeCollection, NoticeCollection>();
 			services.AddSingleton<IDashboardPreviewCollection, DashboardPreviewCollection>();
 
-			services.AddSingleton<ToolCollection>();
+			services.AddSingleton<IToolCollection, ToolCollection>();
 
 			services.AddSingleton<IConfigSource, InMemoryConfigSource>();
 			services.AddSingleton<IConfigSource, FileConfigSource>();
diff --git a/Engine/Source/Programs/Horde/Horde.Build/Tasks/UpgradeTaskSource.cs b/Engine/Source/Programs/Horde/Horde.Build/Tasks/UpgradeTaskSource.cs
--- a/Engine/Source/Programs/Horde/Horde.Build/Tasks/UpgradeTaskSource.cs
+++ b/Engine/Source/Programs/Horde/Horde.Build/Tasks/UpgradeTaskSource.cs
@@ -9,7 +9,6 @@ using Google.Protobuf;
 using Google.Protobuf.WellKnownTypes;
 using Horde.Build.Agents;
 using Horde.Build.Agents.Leases;
-using Horde.Build.Agents.Software;
 using Horde.Build.Jobs;
 using Horde.Build.Logs;
 using Horde.Build.Server;
@@ -26,13 +25,13 @@ namespace Horde.Build.Tasks
 
 		public override TaskSourceFlags Flags => TaskSourceFlags.AllowWhenDisabled | TaskSourceFlags.AllowDuringDowntime;
 
-		readonly ToolCollection _toolCollection;
+		readonly IToolCollection _toolCollection;
 		readonly ILogFileService _logService;
 		readonly IOptionsMonitor<GlobalConfig> _globalConfig;
 		readonly IOptions<ServerSettings> _serverSettings;
 		readonly IClock _clock;
 
-		public UpgradeTaskSource(ToolCollection toolCollection, ILogFileService logService, IOptionsMonitor<GlobalConfig> globalConfig, IOptions<ServerSettings> serverSettings, IClock clock)
+		public UpgradeTaskSource(IToolCollection toolCollection, ILogFileService logService, IOptionsMonitor<GlobalConfig> globalConfig, IOptions<ServerSettings> serverSettings, IClock clock)
 		{
 			_toolCollection = toolCollection;
 			_logService = logService;
diff --git a/Engine/Source/Programs/Horde/Horde.Build/Tools/IToolCollection.cs b/Engine/Source/Programs/Horde/Horde.Build/Tools/IToolCollection.cs
new file mode 100644
--- /dev/null
+++ b/Engine/Source/Programs/Horde/Horde.Build/Tools/IToolCollection.cs
@@ -0,0 +1,52 @@
+// Copyright Epic Games, Inc. All Rights Reserved.
+
+using System.IO;
+using System.Threading;
+using System.Threading.Tasks;
+using Horde.Build.Server;
+
+namespace Horde.Build.Tools
+{
+	/// <summary>
+	/// Collection of tool documents
+	/// </summary>
+	public interface IToolCollection
+	{
+		/// <summary>
+		/// Gets a tool with the given identifier
+		/// </summary>
+		/// <param name="id">The tool identifier</param>
+		/// <param name="globalConfig">The current global configuration</param>
+		/// <returns></returns>
+		Task<ITool?> GetAsync(ToolId id, GlobalConfig globalConfig);
+
+		/// <summary>
+		/// Adds a new deployment to the given tool. The new deployment will replace the current active deployment.
+		/// </summary>
+		/// <param name="tool">The tool to update</param>
+		/// <param name="options">Options for the new deployment</param>
+		/// <param name="stream">Stream containing the tool data</param>
+		/// <param name="globalConfig">The current configuration</param>
+		/// <param name="cancellationToken">Cancellation token for the operation</param>
+		/// <returns>Updated tool document, or null if it does not exist</returns>
+		Task<ITool?> CreateDeploymentAsync(ITool tool, ToolDeploymentConfig options, Stream stream, GlobalConfig globalConfig, CancellationToken cancellationToken);
+
+		/// <summary>
+		/// Updates the state of the current deployment
+		/// </summary>
+		/// <param name="tool">Tool to be updated</param>
+		/// <param name="deploymentId">Identifier for the deployment to modify</param>
+		/// <param name="action">New state of the deployment</param>
+		/// <returns></returns>
+		Task<ITool?> UpdateDeploymentAsync(ITool tool, ToolDeploymentId deploymentId, ToolDeploymentState action);
+
+		/// <summary>
+		/// Opens a stream to the data for a particular deployment
+		/// </summary>
+		/// <param name="tool">Identifier for the tool</param>
+		/// <param name="deployment">The deployment</param>
+		/// <param name="cancellationToken">Cancellation token for the operation</param>
+		/// <returns>Stream for the data</returns>
+		Task<Stream> GetDeploymentZipAsync(ITool tool, IToolDeployment deployment, CancellationToken cancellationToken);
+	}
+}
diff --git a/Engine/Source/Programs/Horde/Horde.Build/Tools/ToolCollection.cs b/Engine/Source/Programs/Horde/Horde.Build/Tools/ToolCollection.cs
--- a/Engine/Source/Programs/Horde/Horde.Build/Tools/ToolCollection.cs
+++ b/Engine/Source/Programs/Horde/Horde.Build/Tools/ToolCollection.cs
@@ -26,7 +26,7 @@ namespace Horde.Build.Tools
 	/// <summary>
 	/// Collection of tool documents
 	/// </summary>
-	public class ToolCollection
+	public class ToolCollection : IToolCollection
 	{
 		class Tool : VersionedDocument<ToolId, Tool>, ITool
 		{
diff --git a/Engine/Source/Programs/Horde/Horde.Build/Tools/ToolsController.cs b/Engine/Source/Programs/Horde/Horde.Build/Tools/ToolsController.cs
--- a/Engine/Source/Programs/Horde/Horde.Build/Tools/ToolsController.cs
+++ b/Engine/Source/Programs/Horde/Horde.Build/Tools/ToolsController.cs
@@ -185,13 +185,13 @@ namespace Horde.Build.Tools
 	[Route("[controller]")]
 	public class ToolsController : HordeControllerBase
 	{
-		readonly ToolCollection _toolCollection;
+		readonly IToolCollection _toolCollection;
 		readonly IOptionsSnapshot<GlobalConfig> _globalConfig;
 
 		/// <summary>
 		/// Constructor
 		/// </summary>
-		public ToolsController(ToolCollection toolCollection, IOptionsSnapshot<GlobalConfig> globalConfig)
+		public ToolsController(IToolCollection toolCollection, IOptionsSnapshot<GlobalConfig> globalConfig)
 		{
 			_toolCollection = toolCollection;
 			_globalConfig = globalConfig;
@@ -263,7 +263,7 @@ namespace Horde.Build.Tools
 	[ApiController]
 	public class PublicToolsController : HordeControllerBase
 	{
-		readonly ToolCollection _toolCollection;
+		readonly IToolCollection _toolCollection;
 		readonly StorageService _storageService;
 		readonly IOptionsSnapshot<GlobalConfig> _globalConfig;
 		readonly IClock _clock;
@@ -271,7 +271,7 @@ namespace Horde.Build.Tools
 		/// <summary>
 		/// Constructor
 		/// </summary>
-		public PublicToolsController(ToolCollection toolCollection, StorageService storageService, IClock clock, IOptionsSnapshot<GlobalConfig> globalConfig)
+		public PublicToolsController(IToolCollection toolCollection, StorageService storageService, IClock clock, IOptionsSnapshot<GlobalConfig> globalConfig)
 		{
 			_toolCollection = toolCollection;
 			_storageService = storageService;
