Horde: Add endpoints to the storage controller that allow decoding bundles to json. Makes it much simpler to browse uploaded data without having to locally decode it.

diff --git a/Engine/Source/Programs/Horde/Horde.Build/Commands/Bundles/CreateCommand.cs b/Engine/Source/Programs/Horde/Horde.Build/Commands/Bundles/CreateCommand.cs
--- a/Engine/Source/Programs/Horde/Horde.Build/Commands/Bundles/CreateCommand.cs
+++ b/Engine/Source/Programs/Horde/Horde.Build/Commands/Bundles/CreateCommand.cs
@@ -6,6 +6,7 @@ using EpicGames.Core;
 using EpicGames.Horde.Storage;
 using EpicGames.Horde.Storage.Nodes;
 using Horde.Build.Perforce;
+using Horde.Build.Storage;
 using Microsoft.Extensions.Configuration;
 using Microsoft.Extensions.DependencyInjection;
 using Microsoft.Extensions.Logging;
@@ -18,6 +19,9 @@ namespace Horde.Build.Commands.Bundles
 		[CommandLine("-Ref=")]
 		public RefName RefName { get; set; } = new RefName("default-ref");
 
+		[CommandLine("-Namespace=")]
+		public NamespaceId NamespaceId { get; set; } = Namespace.Artifacts;
+
 		[CommandLine("-InputDir=", Required = true)]
 		public DirectoryReference InputDir { get; set; } = null!;
 
@@ -34,8 +38,10 @@ namespace Horde.Build.Commands.Bundles
 		{
 			using ServiceProvider serviceProvider = Startup.CreateServiceProvider(_configuration, _loggerProvider);
 
-			IStorageClient store = serviceProvider.GetRequiredService<IStorageClient<ReplicationService>>();
-			using TreeWriter writer = new TreeWriter(store, prefix: RefName.Text);
+			StorageService storageService = serviceProvider.GetRequiredService<StorageService>();
+			IStorageClient storageClient = await storageService.GetClientAsync(NamespaceId, CancellationToken.None);			
+			
+			using TreeWriter writer = new TreeWriter(storageClient, prefix: RefName.Text);
 
 			DirectoryNode node = new DirectoryNode(DirectoryFlags.None);
 			await node.CopyFromDirectoryAsync(InputDir.ToDirectoryInfo(), new ChunkingOptions(), writer, new CopyStatsLogger(logger), CancellationToken.None);
diff --git a/Engine/Source/Programs/Horde/Horde.Build/Storage/StorageController.cs b/Engine/Source/Programs/Horde/Horde.Build/Storage/StorageController.cs
--- a/Engine/Source/Programs/Horde/Horde.Build/Storage/StorageController.cs
+++ b/Engine/Source/Programs/Horde/Horde.Build/Storage/StorageController.cs
@@ -4,6 +4,7 @@ using System;
 using System.Collections.Generic;
 using System.IO;
 using System.Linq;
+using System.Reflection;
 using System.Security.Claims;
 using System.Text.RegularExpressions;
 using System.Threading;
@@ -11,12 +12,17 @@ using System.Threading.Tasks;
 using Amazon.EC2.Model;
 using EpicGames.Core;
 using EpicGames.Horde.Storage;
+using EpicGames.Horde.Storage.Nodes;
+using EpicGames.Redis;
 using Horde.Build.Acls;
 using Horde.Build.Server;
 using Horde.Build.Utilities;
 using Microsoft.AspNetCore.Authorization;
 using Microsoft.AspNetCore.Http;
 using Microsoft.AspNetCore.Mvc;
+using Microsoft.CodeAnalysis.CSharp.Syntax;
+using Microsoft.Extensions.Caching.Memory;
+using Microsoft.Extensions.Logging;
 using Microsoft.Extensions.Options;
 using Microsoft.Extensions.Primitives;
 
@@ -130,14 +136,20 @@ namespace Horde.Build.Storage
 		/// </summary>
 		public int ExportIdx { get; set; }
 
+		/// <summary>
+		/// Link to information about the target node
+		/// </summary>
+		public string Link { get; set; }
+
 		/// <summary>
 		/// Constructor
 		/// </summary>
-		public ReadRefResponse(NodeHandle target)
+		public ReadRefResponse(NodeHandle target, string link)
 		{
 			Hash = target.Hash;
 			Blob = target.Locator.Blob;
 			ExportIdx = target.Locator.ExportIdx;
+			Link = link;
 		}
 	}
 
@@ -150,15 +162,19 @@ namespace Horde.Build.Storage
 	public class StorageController : HordeControllerBase
 	{
 		readonly StorageService _storageService;
+		readonly IMemoryCache _memoryCache;
 		readonly IOptionsSnapshot<GlobalConfig> _globalConfig;
+		readonly ILogger _logger;
 
 		/// <summary>
 		/// Constructor
 		/// </summary>
-		public StorageController(StorageService storageService, IOptionsSnapshot<GlobalConfig> globalConfig)
+		public StorageController(StorageService storageService, IMemoryCache memoryCache, IOptionsSnapshot<GlobalConfig> globalConfig, ILogger<StorageController> logger)
 		{
 			_storageService = storageService;
+			_memoryCache = memoryCache;
 			_globalConfig = globalConfig;
+			_logger = logger;
 		}
 
 		/// <summary>
@@ -200,7 +216,7 @@ namespace Horde.Build.Storage
 				using (Stream stream = file.OpenReadStream())
 				{
 					BlobLocator locator = await client.WriteBlobAsync(stream, prefix: (prefix == null) ? Utf8String.Empty : new Utf8String(prefix), cancellationToken: cancellationToken);
-					return new WriteBlobResponse { Blob = locator, SupportsRedirects = client.SupportsRedirects? (bool?)true : null };
+					return new WriteBlobResponse { Blob = locator, SupportsRedirects = client.SupportsRedirects ? (bool?)true : null };
 				}
 			}
 		}
@@ -396,7 +412,8 @@ namespace Horde.Build.Storage
 				return new NotFoundResult();
 			}
 
-			return new ReadRefResponse(target);
+			string link = $"/api/v1/storage/{namespaceId}/nodes/{target.Locator.Blob}?export={target.Locator.ExportIdx}";
+			return new ReadRefResponse(target, link);
 		}
 
 		/// <summary>
@@ -437,5 +454,171 @@ namespace Horde.Build.Storage
 		{
 			return name.Length > prefix.Length && name[prefix.Length] == '/' && name.AsSpan(0, prefix.Length).SequenceEqual(prefix);
 		}
+
+		/// <summary>
+		/// Gets information about a particular bundle in storage
+		/// </summary>
+		/// <param name="namespaceId">Namespace containing the blob</param>
+		/// <param name="locator">Blob locator</param>
+		/// <param name="includeImports">Whether to include imports for the bundle</param>
+		/// <param name="includeExports">Whether to include exports for the bundle</param>
+		/// <param name="includePackets">Whether to include packets for the bundle</param>
+		/// <param name="cancellationToken">Cancellation token for the operation</param>
+		/// <returns></returns>
+		[Route("/api/v1/storage/{namespaceId}/bundles/{*locator}")]
+		public async Task<ActionResult<object>> GetBundleAsync(NamespaceId namespaceId, BlobLocator locator, [FromQuery(Name = "imports")] bool includeImports = false, [FromQuery(Name = "exports")] bool includeExports = true, [FromQuery(Name = "packets")] bool includePackets = false, CancellationToken cancellationToken = default)
+		{
+			NamespaceConfig? namespaceConfig;
+			if (!_globalConfig.Value.Storage.TryGetNamespace(namespaceId, out namespaceConfig))
+			{
+				return NotFound(namespaceId);
+			}
+			if (!namespaceConfig.Authorize(AclAction.ReadBlobs, User))
+			{
+				return Forbid(AclAction.ReadBlobs, namespaceId);
+			}
+
+			IStorageClient storageClient = await _storageService.GetClientAsync(namespaceId, cancellationToken);
+			TreeReader reader = new TreeReader(storageClient, _memoryCache, _logger);
+
+			BundleHeader header = await reader.ReadBundleHeaderAsync(locator, cancellationToken);
+
+			string linkBase = $"/api/v1/storage/{namespaceId}";
+
+			List<object>? responseImports = null;
+			if (includeImports)
+			{
+				responseImports = new List<object>();
+				foreach (BundleImport import in header.Imports)
+				{
+					string link = $"{linkBase}/bundles/{import.Locator}";
+					List<string> exports = import.Exports.ConvertAll(x => $"{linkBase}/nodes/{import.Locator}?export={x}");
+					responseImports.Add(new { bundle = link, exports });
+				}
+			}
+
+			List<object>? responseExports = null;
+			if (includeExports)
+			{
+				responseExports = new List<object>();
+				for (int exportIdx = 0; exportIdx < header.Exports.Count; exportIdx++)
+				{
+					BundleExport export = header.Exports[exportIdx];
+
+					string details = $"{linkBase}/nodes/{locator}?export={exportIdx}";
+					BundleType type = header.Types[export.TypeIdx];
+					string typeName = GetNodeType(type.Guid)?.Name ?? type.Guid.ToString();
+
+					responseExports.Add(new { export.Hash, export.Length, details, type = typeName });
+				}
+			}
+
+			List<object>? responsePackets = null;
+			if (includePackets)
+			{
+				responsePackets = new List<object>();
+				for (int packetIdx = 0, exportIdx = 0; packetIdx < header.Packets.Count; packetIdx++)
+				{
+					BundlePacket packet = header.Packets[packetIdx];
+
+					List<string> packetExports = new List<string>();
+
+					int length = 0;
+					for (; exportIdx < header.Exports.Count && length + header.Exports[exportIdx].Length <= packet.DecodedLength; exportIdx++)
+					{
+						BundleExport export = header.Exports[exportIdx];
+						packetExports.Add($"{linkBase}/{locator}?export={exportIdx}");
+						length += export.Length;
+					}
+
+					responsePackets.Add(new { packetIdx, packet.EncodedLength, packet.DecodedLength, exports = packetExports });
+				}
+			}
+
+			return new { header.CompressionFormat, imports = responseImports, exports = responseExports, packets = responsePackets };
+		}
+
+		/// <summary>
+		/// Gets information about a particular bundle in storage
+		/// </summary>
+		/// <param name="namespaceId">Namespace containing the blob</param>
+		/// <param name="locator">Blob locator</param>
+		/// <param name="exportIdx">Index of the export</param>
+		/// <param name="cancellationToken">Cancellation token for the operation</param>
+		/// <returns></returns>
+		[Route("/api/v1/storage/{namespaceId}/nodes/{*locator}")]
+		public async Task<ActionResult<object>> GetNodeAsync(NamespaceId namespaceId, BlobLocator locator, [FromQuery(Name = "export")] int exportIdx, CancellationToken cancellationToken = default)
+		{
+			NamespaceConfig? namespaceConfig;
+			if (!_globalConfig.Value.Storage.TryGetNamespace(namespaceId, out namespaceConfig))
+			{
+				return NotFound(namespaceId);
+			}
+			if (!namespaceConfig.Authorize(AclAction.ReadBlobs, User))
+			{
+				return Forbid(AclAction.ReadBlobs, namespaceId);
+			}
+
+			IStorageClient storageClient = await _storageService.GetClientAsync(namespaceId, cancellationToken);
+			TreeReader reader = new TreeReader(storageClient, _memoryCache, _logger);
+
+			BundleHeader header = await reader.ReadBundleHeaderAsync(locator, cancellationToken);
+			BundleExport export = header.Exports[exportIdx];
+
+			string linkBase = $"/api/v1/storage/{namespaceId}";
+
+			object content;
+
+			TreeNode node = await reader.ReadNodeAsync(new NodeLocator(locator, exportIdx), cancellationToken);
+			switch (node)
+			{
+				case DirectoryNode directoryNode:
+					{
+						List<object> directories = new List<object>();
+						foreach ((Utf8String name, DirectoryEntry entry) in directoryNode.NameToDirectory)
+						{
+							directories.Add(new { name = name.ToString(), link = GetNodeLink(linkBase, entry) });
+						}
+
+						List<object> files = new List<object>();
+						foreach ((Utf8String name, FileEntry entry) in directoryNode.NameToFile)
+						{
+							files.Add(new { name = name.ToString(), link = GetNodeLink(linkBase, entry) });
+						}
+
+						content = new { directoryNode.Length, directories, files };
+					}
+					break;
+				default:
+					content = new { references = node.EnumerateRefs().Select(x => GetNodeLink(linkBase, x)) };
+					break;
+			}
+
+			return new { bundle = $"{linkBase}/bundles/{locator}", export.Hash, export.Length, guid = header.Types[export.TypeIdx].Guid, type = node.GetType().Name, content = content };
+		}
+
+		static string GetNodeLink(string linkBase, TreeNodeRef treeNodeRef) => $"{linkBase}/nodes/{treeNodeRef.Handle!.Locator.Blob}?export={treeNodeRef.Handle!.Locator.ExportIdx}";
+
+		static readonly Dictionary<Guid, Type> s_guidToType = GetGuidToTypeMap();
+
+		static Dictionary<Guid, Type> GetGuidToTypeMap()
+		{
+			Dictionary<Guid, Type> guidToType = new Dictionary<Guid, Type>();
+			foreach (Type nodeType in new TreeReaderOptions().Types)
+			{
+				TreeNodeAttribute? attribute = nodeType.GetCustomAttribute<TreeNodeAttribute>();
+				if (attribute != null)
+				{
+					guidToType.Add(Guid.Parse(attribute.Guid), nodeType);
+				}
+			}
+			return guidToType;
+		}
+
+		static Type? GetNodeType(Guid typeGuid)
+		{
+			s_guidToType.TryGetValue(typeGuid, out Type? type);
+			return type;
+		}
 	}
 }
diff --git a/Engine/Source/Programs/Horde/Horde.Build/Utilities/HordeControllerBase.cs b/Engine/Source/Programs/Horde/Horde.Build/Utilities/HordeControllerBase.cs
--- a/Engine/Source/Programs/Horde/Horde.Build/Utilities/HordeControllerBase.cs
+++ b/Engine/Source/Programs/Horde/Horde.Build/Utilities/HordeControllerBase.cs
@@ -234,6 +234,15 @@ namespace Horde.Build.Utilities
 			return NotFound("Step {StepId} not found on job {JobId} batch {BatchId}", stepId, jobId, batchId);
 		}
 
+		/// <summary>
+		/// Returns a 404 response for the given object
+		/// </summary>
+		[NonAction]
+		protected ActionResult NotFound(NamespaceId namespaceId)
+		{
+			return NotFound("Namespace {NamespaceId} not found", namespaceId);
+		}
+
 		/// <summary>
 		/// Returns a 404 response for the given object
 		/// </summary>
diff --git a/Engine/Source/Programs/Shared/EpicGames.Horde/Storage/NodeLocator.cs b/Engine/Source/Programs/Shared/EpicGames.Horde/Storage/NodeLocator.cs
--- a/Engine/Source/Programs/Shared/EpicGames.Horde/Storage/NodeLocator.cs
+++ b/Engine/Source/Programs/Shared/EpicGames.Horde/Storage/NodeLocator.cs
@@ -51,6 +51,11 @@ namespace EpicGames.Horde.Storage
 		public static NodeLocator Parse(ReadOnlySpan<char> text)
 		{
 			int hashIdx = text.IndexOf('#');
+			if (hashIdx == -1)
+			{
+				throw new ArgumentException("Invalid node locator", nameof(text));
+			}
+
 			int exportIdx = Int32.Parse(text.Slice(hashIdx + 1), NumberStyles.None, CultureInfo.InvariantCulture);
 			BlobLocator blobLocator = new BlobLocator(new Utf8String(text.Slice(0, hashIdx)));
 			return new NodeLocator(blobLocator, exportIdx);
@@ -76,7 +81,7 @@ namespace EpicGames.Horde.Storage
 	}
 
 	/// <summary>
-	/// Type converter for BlobId to and from JSON
+	/// Type converter for <see cref="NodeLocator"/> to and from JSON
 	/// </summary>
 	sealed class NodeLocatorJsonConverter : JsonConverter<NodeLocator>
 	{
@@ -88,7 +93,7 @@ namespace EpicGames.Horde.Storage
 	}
 
 	/// <summary>
-	/// Type converter from strings to BlobId objects
+	/// Type converter from strings to <see cref="NodeLocator"/> objects
 	/// </summary>
 	sealed class NodeLocatorTypeConverter : TypeConverter
 	{
@@ -101,7 +106,7 @@ namespace EpicGames.Horde.Storage
 		/// <inheritdoc/>
 		public override object ConvertFrom(ITypeDescriptorContext? context, CultureInfo? culture, object? value)
 		{
-			return new BlobLocator((string)value!);
+			return NodeLocator.Parse((string)value!);
 		}
 	}
 
diff --git a/Engine/Source/Programs/Shared/EpicGames.Horde/Storage/Nodes/DirectoryNode.cs b/Engine/Source/Programs/Shared/EpicGames.Horde/Storage/Nodes/DirectoryNode.cs
--- a/Engine/Source/Programs/Shared/EpicGames.Horde/Storage/Nodes/DirectoryNode.cs
+++ b/Engine/Source/Programs/Shared/EpicGames.Horde/Storage/Nodes/DirectoryNode.cs
@@ -826,7 +826,13 @@ namespace EpicGames.Horde.Storage.Nodes
 							directory = await FindOrAddDirectoryAsync(null!, fullName.Substring(0, slashIdx), cancellationToken);
 						}
 
-						FileEntry fileEntry = new FileEntry(entry.Name, FileEntryFlags.None, fileWriter.Length, node);
+						FileEntryFlags flags = FileEntryFlags.None;
+						if ((entry.ExternalAttributes & (0b_001_001_001 << 16)) != 0)
+						{
+							flags |= FileEntryFlags.Executable;
+						}
+
+						FileEntry fileEntry = new FileEntry(entry.Name, flags, fileWriter.Length, node);
 						directory.AddFile(fileEntry);
 					}
 				}
diff --git a/Engine/Source/Programs/Shared/EpicGames.Horde/Storage/Nodes/FileNode.cs b/Engine/Source/Programs/Shared/EpicGames.Horde/Storage/Nodes/FileNode.cs
--- a/Engine/Source/Programs/Shared/EpicGames.Horde/Storage/Nodes/FileNode.cs
+++ b/Engine/Source/Programs/Shared/EpicGames.Horde/Storage/Nodes/FileNode.cs
@@ -390,8 +390,8 @@ namespace EpicGames.Horde.Storage.Nodes
 			}
 		}
 
-		static readonly BundleType s_leafNodeType = TreeNodeExtensions.GetBundleType(typeof(LeafFileNode));
-		static readonly BundleType s_interiorNodeType = TreeNodeExtensions.GetBundleType(typeof(InteriorFileNode));
+		static readonly BundleType s_leafNodeType = TreeNode.GetBundleType<LeafFileNode>();
+		static readonly BundleType s_interiorNodeType = TreeNode.GetBundleType<InteriorFileNode>();
 
 		readonly TreeWriter _writer;
 		readonly ChunkingOptions _options;
diff --git a/Engine/Source/Programs/Shared/EpicGames.Horde/Storage/TreeNode.cs b/Engine/Source/Programs/Shared/EpicGames.Horde/Storage/TreeNode.cs
--- a/Engine/Source/Programs/Shared/EpicGames.Horde/Storage/TreeNode.cs
+++ b/Engine/Source/Programs/Shared/EpicGames.Horde/Storage/TreeNode.cs
@@ -3,6 +3,7 @@
 using System;
 using System.Collections.Concurrent;
 using System.Collections.Generic;
+using System.Linq;
 using System.Reflection;
 using System.Threading;
 using System.Threading.Tasks;
@@ -26,6 +27,11 @@ namespace EpicGames.Horde.Storage
 		/// </summary>
 		public IoHash Hash { get; internal set; }
 
+		/// <summary>
+		/// Accessor for the bundle type definition associated with this node
+		/// </summary>
+		public BundleType BundleType => GetBundleType(GetType());
+
 		/// <summary>
 		/// Default constructor
 		/// </summary>
@@ -62,6 +68,41 @@ namespace EpicGames.Horde.Storage
 		/// </summary>
 		/// <returns>References to other nodes</returns>
 		public abstract IEnumerable<TreeNodeRef> EnumerateRefs();
+
+		#region Static Methods
+
+		/// <summary>
+		/// Cache of constructed <see cref="BundleType"/> instances.
+		/// </summary>
+		static readonly ConcurrentDictionary<Type, BundleType> s_typeToBundleType = new ConcurrentDictionary<Type, BundleType>();
+
+		/// <summary>
+		/// Gets the <see cref="BundleType"/> instance for a particular node type
+		/// </summary>
+		/// <param name="type"></param>
+		/// <returns>Bundle</returns>
+		public static BundleType GetBundleType(Type type)
+		{
+			BundleType? bundleType;
+			if (!s_typeToBundleType.TryGetValue(type, out bundleType))
+			{
+				TreeNodeAttribute? attribute = type.GetCustomAttribute<TreeNodeAttribute>();
+				if (attribute == null)
+				{
+					throw new InvalidOperationException($"Missing {nameof(TreeNodeAttribute)} from type {type.Name}");
+				}
+				bundleType = s_typeToBundleType.GetOrAdd(type, new BundleType(Guid.Parse(attribute.Guid), attribute.Version));
+			}
+			return bundleType;
+		}
+
+		/// <summary>
+		/// Gets the <see cref="BundleType"/> instance for a particular node type
+		/// </summary>
+		/// <typeparam name="T">Type of the node</typeparam>
+		public static BundleType GetBundleType<T>() where T : TreeNode => GetBundleType(typeof(T));
+
+		#endregion
 	}
 
 	/// <summary>
@@ -206,37 +247,5 @@ namespace EpicGames.Horde.Storage
 			using TreeWriter writer = new TreeWriter(store, options, prefix.IsEmpty ? name.Text : prefix);
 			return await writer.WriteAsync(name, node, refOptions, cancellationToken);
 		}
-
-		/// <summary>
-		/// Cache of constructed <see cref="BundleType"/> instances.
-		/// </summary>
-		static readonly ConcurrentDictionary<Type, BundleType> s_typeToBundleType = new ConcurrentDictionary<Type, BundleType>();
-
-		/// <summary>
-		/// Gets the bundle type object for a particular node
-		/// </summary>
-		/// <param name="node"></param>
-		/// <returns></returns>
-		public static BundleType GetBundleType(this TreeNode node) => GetBundleType(node.GetType());
-
-		/// <summary>
-		/// Gets the <see cref="BundleType"/> instance for a particular node type
-		/// </summary>
-		/// <param name="type"></param>
-		/// <returns>Bundle</returns>
-		public static BundleType GetBundleType(Type type)
-		{
-			BundleType? bundleType;
-			if (!s_typeToBundleType.TryGetValue(type, out bundleType))
-			{
-				TreeNodeAttribute? attribute = type.GetCustomAttribute<TreeNodeAttribute>();
-				if (attribute == null)
-				{
-					throw new InvalidOperationException($"Missing {nameof(TreeNodeAttribute)} from type {type.Name}");
-				}
-				bundleType = s_typeToBundleType.GetOrAdd(type, new BundleType(Guid.Parse(attribute.Guid), attribute.Version));
-			}
-			return bundleType;
-		}
 	}
 }
diff --git a/Engine/Source/Programs/Shared/EpicGames.Horde/Storage/TreeNodeRef.cs b/Engine/Source/Programs/Shared/EpicGames.Horde/Storage/TreeNodeRef.cs
--- a/Engine/Source/Programs/Shared/EpicGames.Horde/Storage/TreeNodeRef.cs
+++ b/Engine/Source/Programs/Shared/EpicGames.Horde/Storage/TreeNodeRef.cs
@@ -350,7 +350,7 @@ namespace EpicGames.Horde.Storage
 			target.Serialize(nodeWriter);
 
 			// Write the final data
-			NodeHandle handle = await writer.WriteNodeAsync(nodeWriter.Length, nextRefHandles, target.GetBundleType(), cancellationToken);
+			NodeHandle handle = await writer.WriteNodeAsync(nodeWriter.Length, nextRefHandles, target.BundleType, cancellationToken);
 			target.Hash = handle.Hash;
 
 			NodeWriteCallback writeState = new NodeWriteCallback(nodeRef, handle);
diff --git a/Engine/Source/Programs/Shared/EpicGames.Horde/Storage/TreeReader.cs b/Engine/Source/Programs/Shared/EpicGames.Horde/Storage/TreeReader.cs
--- a/Engine/Source/Programs/Shared/EpicGames.Horde/Storage/TreeReader.cs
+++ b/Engine/Source/Programs/Shared/EpicGames.Horde/Storage/TreeReader.cs
@@ -113,7 +113,7 @@ namespace EpicGames.Horde.Storage
 				TypeInfo? typeInfo;
 				if (!s_cachedTypeInfo.TryGetValue(type, out typeInfo))
 				{
-					BundleType bundleType = TreeNodeExtensions.GetBundleType(type);
+					BundleType bundleType = TreeNode.GetBundleType(type);
 
 					Type[] signature = new[] { typeof(ITreeNodeReader) };
 
