Horde: Add support for bundling tools with the server distribution.

* Bundled tools are specified in the server config file (appsettings.json) and are public by default
* Bundled tools use the file storage client to return blobs, allowing most of the code to be shared with regular tools.
* Any tool specified in the global config file with the same id will replace a bundled tool.
* Creating new deployments of bundled tools is not permitted.

diff --git a/Engine/Source/Programs/Horde/Horde.Build/Commands/Bundles/CreateCommand.cs b/Engine/Source/Programs/Horde/Horde.Build/Commands/Bundles/CreateCommand.cs
--- a/Engine/Source/Programs/Horde/Horde.Build/Commands/Bundles/CreateCommand.cs
+++ b/Engine/Source/Programs/Horde/Horde.Build/Commands/Bundles/CreateCommand.cs
@@ -4,6 +4,7 @@ using System.Threading;
 using System.Threading.Tasks;
 using EpicGames.Core;
 using EpicGames.Horde.Storage;
+using EpicGames.Horde.Storage.Backends;
 using EpicGames.Horde.Storage.Nodes;
 using Horde.Build.Perforce;
 using Horde.Build.Storage;
@@ -25,6 +26,9 @@ namespace Horde.Build.Commands.Bundles
 		[CommandLine("-InputDir=", Required = true)]
 		public DirectoryReference InputDir { get; set; } = null!;
 
+		[CommandLine("-OutputDir=")]
+		public DirectoryReference? OutputDir { get; set; }
+
 		readonly IConfiguration _configuration;
 		readonly ILoggerProvider _loggerProvider;
 
@@ -38,9 +42,18 @@ namespace Horde.Build.Commands.Bundles
 		{
 			using ServiceProvider serviceProvider = Startup.CreateServiceProvider(_configuration, _loggerProvider);
 
-			StorageService storageService = serviceProvider.GetRequiredService<StorageService>();
-			IStorageClient storageClient = await storageService.GetClientAsync(NamespaceId, CancellationToken.None);			
-			
+			IStorageClient storageClient;
+			if (OutputDir == null)
+			{
+				StorageService storageService = serviceProvider.GetRequiredService<StorageService>();
+				storageClient = await storageService.GetClientAsync(NamespaceId, CancellationToken.None);
+			}
+			else
+			{
+				FileStorageClient fileStorageClient = new FileStorageClient(OutputDir, logger);
+				storageClient = fileStorageClient;
+			}
+		
 			using TreeWriter writer = new TreeWriter(storageClient, prefix: RefName.Text);
 
 			DirectoryNode node = new DirectoryNode(DirectoryFlags.None);
diff --git a/Engine/Source/Programs/Horde/Horde.Build/ServerSettings.cs b/Engine/Source/Programs/Horde/Horde.Build/ServerSettings.cs
--- a/Engine/Source/Programs/Horde/Horde.Build/ServerSettings.cs
+++ b/Engine/Source/Programs/Horde/Horde.Build/ServerSettings.cs
@@ -2,6 +2,7 @@
 
 using System;
 using System.Collections.Generic;
+using System.Diagnostics.CodeAnalysis;
 using System.Linq;
 using System.Security.Claims;
 using System.Text.Json.Serialization;
@@ -15,6 +16,7 @@ using Horde.Build.Agents.Fleet;
 using Horde.Build.Server;
 using Horde.Build.Storage.Backends;
 using Horde.Build.Telemetry;
+using Horde.Build.Tools;
 using Horde.Build.Utilities;
 using Serilog.Events;
 
@@ -240,6 +242,35 @@ namespace Horde.Build
 		public ChunkingOptions Chunking { get; set; } = new ChunkingOptions();
 	}
 
+	/// <summary>
+	/// Configuration for a tool bundled alongsize the server
+	/// </summary>
+	public class BundledToolConfig : ToolConfig
+	{
+		/// <summary>
+		/// Version string for the current tool data
+		/// </summary>
+		public string Version { get; set; } = "1.0";
+
+		/// <summary>
+		/// Ref name in the tools directory
+		/// </summary>
+		public RefName RefName { get; set; } = new RefName("default-ref");
+
+		/// <summary>
+		/// Directory containing blob data for this tool. If empty, the tools/{id} folder next to the server will be used.
+		/// </summary>
+		public string? DataDir { get; set; }
+
+		/// <summary>
+		/// Constructor
+		/// </summary>
+		public BundledToolConfig()
+		{
+			Public = true;
+		}
+	}
+
 	/// <summary>
 	/// Global settings for the application
 	/// </summary>
@@ -671,6 +702,11 @@ namespace Horde.Build
 		/// </summary>
 		public TelemetryConfig Telemetry { get; set; } = new TelemetryConfig();
 
+		/// <summary>
+		/// Tools bundled along with the server. Data for each tool can be produced using the 'bundle create' command, and should be stored in the /tools/{id} directory.
+		/// </summary>
+		public List<BundledToolConfig> BundledTools { get; set; } = new List<BundledToolConfig>();
+
 		/// <summary>
 		/// Default pre-baked ACL for authentication of well-known roles
 		/// </summary>
@@ -736,6 +772,18 @@ namespace Horde.Build
 				throw new ArgumentException($"Settings key '{nameof(RunModes)}' contains one or more invalid entries");
 			}
 		}
+
+		/// <summary>
+		/// Attempts to get a bundled tool with the given id
+		/// </summary>
+		/// <param name="toolId">The tool id</param>
+		/// <param name="bundledToolConfig">Configuration for the bundled tool</param>
+		/// <returns>True if the tool was found</returns>
+		public bool TryGetBundledTool(ToolId toolId, [NotNullWhen(true)] out BundledToolConfig? bundledToolConfig)
+		{
+			bundledToolConfig = BundledTools.FirstOrDefault(x => x.Id == toolId);
+			return bundledToolConfig != null;
+		}
 	}
 
 	/// <summary>
diff --git a/Engine/Source/Programs/Horde/Horde.Build/Storage/StorageController.cs b/Engine/Source/Programs/Horde/Horde.Build/Storage/StorageController.cs
--- a/Engine/Source/Programs/Horde/Horde.Build/Storage/StorageController.cs
+++ b/Engine/Source/Programs/Horde/Horde.Build/Storage/StorageController.cs
@@ -286,11 +286,21 @@ namespace Horde.Build.Storage
 		internal static async Task<ActionResult> ReadBlobInternalAsync(StorageService storageService, NamespaceId namespaceId, BlobLocator locator, int? offset, int? length, CancellationToken cancellationToken)
 		{
 			IStorageClientImpl client = await storageService.GetClientAsync(namespaceId, cancellationToken);
+			return await ReadBlobInternalAsync(client, locator, offset, length, cancellationToken);
+		}
 
-			Uri? redirectUrl = await client.GetReadRedirectAsync(locator, cancellationToken);
-			if (redirectUrl != null)
+		/// <summary>
+		/// Reads a blob from storage, without performing namespace access checks.
+		/// </summary>
+		internal static async Task<ActionResult> ReadBlobInternalAsync(IStorageClient storageClient, BlobLocator locator, int? offset, int? length, CancellationToken cancellationToken)
+		{
+			if (storageClient is IStorageClientImpl storageClientImpl)
 			{
-				return new RedirectResult(redirectUrl.ToString());
+				Uri? redirectUrl = await storageClientImpl.GetReadRedirectAsync(locator, cancellationToken);
+				if (redirectUrl != null)
+				{
+					return new RedirectResult(redirectUrl.ToString());
+				}
 			}
 
 #pragma warning disable CA2000 // Dispose objects before losing scope
@@ -298,11 +308,11 @@ namespace Horde.Build.Storage
 			Stream stream;
 			if (offset == null && length == null)
 			{
-				stream = await client.ReadBlobAsync(locator, cancellationToken);
+				stream = await storageClient.ReadBlobAsync(locator, cancellationToken);
 			}
 			else if (offset != null && length != null)
 			{
-				stream = await client.ReadBlobRangeAsync(locator, offset.Value, length.Value, cancellationToken);
+				stream = await storageClient.ReadBlobRangeAsync(locator, offset.Value, length.Value, cancellationToken);
 			}
 			else
 			{
diff --git a/Engine/Source/Programs/Horde/Horde.Build/Tools/IToolCollection.cs b/Engine/Source/Programs/Horde/Horde.Build/Tools/IToolCollection.cs
--- a/Engine/Source/Programs/Horde/Horde.Build/Tools/IToolCollection.cs
+++ b/Engine/Source/Programs/Horde/Horde.Build/Tools/IToolCollection.cs
@@ -3,6 +3,7 @@
 using System.IO;
 using System.Threading;
 using System.Threading.Tasks;
+using EpicGames.Horde.Storage;
 using Horde.Build.Server;
 
 namespace Horde.Build.Tools
@@ -40,6 +41,14 @@ namespace Horde.Build.Tools
 		/// <returns></returns>
 		Task<ITool?> UpdateDeploymentAsync(ITool tool, ToolDeploymentId deploymentId, ToolDeploymentState action);
 
+		/// <summary>
+		/// Gets the storage client for a particular tool
+		/// </summary>
+		/// <param name="tool">The tool to get a storage client for</param>
+		/// <param name="cancellationToken">Cancellation token for the operation</param>
+		/// <returns>Instance of the storage client</returns>
+		Task<IStorageClient> GetStorageClientAsync(ITool tool, CancellationToken cancellationToken);
+
 		/// <summary>
 		/// Opens a stream to the data for a particular deployment
 		/// </summary>
diff --git a/Engine/Source/Programs/Horde/Horde.Build/Tools/ToolCollection.cs b/Engine/Source/Programs/Horde/Horde.Build/Tools/ToolCollection.cs
--- a/Engine/Source/Programs/Horde/Horde.Build/Tools/ToolCollection.cs
+++ b/Engine/Source/Programs/Horde/Horde.Build/Tools/ToolCollection.cs
@@ -1,6 +1,8 @@
 // Copyright Epic Games, Inc. All Rights Reserved.
 
+using EpicGames.Core;
 using EpicGames.Horde.Storage;
+using EpicGames.Horde.Storage.Backends;
 using EpicGames.Horde.Storage.Nodes;
 using EpicGames.Serialization;
 using Horde.Build.Server;
@@ -209,15 +211,29 @@ namespace Horde.Build.Tools
 		async Task<Tool?> GetInternalAsync(ToolId toolId, GlobalConfig globalConfig)
 		{
 			ToolConfig? toolConfig;
-			if (!globalConfig.TryGetTool(toolId, out toolConfig))
+			if (globalConfig.TryGetTool(toolId, out toolConfig))
 			{
-				return null;
+				Tool tool = await _tools.FindOrAddAsync(toolId, () => new Tool(toolId));
+				tool.Config = toolConfig;
+				tool.UpdateTemporalState(_clock.UtcNow);
+				return tool;
+			}
+
+			BundledToolConfig? bundledToolConfig;
+			if (globalConfig.ServerSettings.TryGetBundledTool(toolId, out bundledToolConfig))
+			{
+				Tool tool = new Tool(bundledToolConfig);
+
+				ToolDeployment deployment = new ToolDeployment(default);
+				deployment.Version = bundledToolConfig.Version;
+				deployment.State = ToolDeploymentState.Complete;
+				deployment.RefName = bundledToolConfig.RefName;
+				tool.Deployments.Add(deployment);
+
+				return tool;
 			}
 
-			Tool tool = await _tools.FindOrAddAsync(toolId, () => new Tool(toolId));
-			tool.Config = toolConfig;
-			tool.UpdateTemporalState(_clock.UtcNow);
-			return tool;
+			return null;
 		}
 
 		/// <summary>
@@ -231,6 +247,11 @@ namespace Horde.Build.Tools
 		/// <returns>Updated tool document, or null if it does not exist</returns>
 		public async Task<ITool?> CreateDeploymentAsync(ITool tool, ToolDeploymentConfig options, Stream stream, GlobalConfig globalConfig, CancellationToken cancellationToken)
 		{
+			if (tool.Config is BundledToolConfig)
+			{
+				throw new InvalidOperationException("Cannot update the state of bundled tools.");
+			}
+
 			ToolDeploymentId deploymentId = ToolDeploymentId.GenerateNewId();
 			RefName refName = new RefName($"{tool.Id}/{deploymentId}");
 
@@ -308,6 +329,11 @@ namespace Horde.Build.Tools
 		/// <returns></returns>
 		public async Task<ITool?> UpdateDeploymentAsync(ITool tool, ToolDeploymentId deploymentId, ToolDeploymentState action)
 		{
+			if (tool.Config is BundledToolConfig)
+			{
+				throw new InvalidOperationException("Cannot update the state of bundled tools.");
+			}
+
 			return await UpdateDeploymentInternalAsync((Tool)tool, deploymentId, action);
 		}
 
@@ -354,6 +380,24 @@ namespace Horde.Build.Tools
 			}
 		}
 
+		/// <summary>
+		/// Gets the storage client containing data for a particular tool
+		/// </summary>
+		/// <param name="tool">Identifier for the tool</param>
+		/// <param name="cancellationToken">Cancellation token for the operation</param>
+		/// <returns>Storage client for the data</returns>
+		public async Task<IStorageClient> GetStorageClientAsync(ITool tool, CancellationToken cancellationToken)
+		{
+			if (tool.Config is BundledToolConfig bundledConfig)
+			{
+				return new FileStorageClient(DirectoryReference.Combine(Program.AppDir, bundledConfig.DataDir ?? $"tools/{tool.Id}"), _logger);
+			}
+			else
+			{
+				return await _storageService.GetClientAsync(Namespace.Tools, cancellationToken);
+			}
+		}
+
 		/// <summary>
 		/// Opens a stream to the data for a particular deployment
 		/// </summary>
@@ -363,9 +407,8 @@ namespace Horde.Build.Tools
 		/// <returns>Stream for the data</returns>
 		public async Task<Stream> GetDeploymentZipAsync(ITool tool, IToolDeployment deployment, CancellationToken cancellationToken)
 		{
-			_ = tool;
+			IStorageClient client = await GetStorageClientAsync(tool, cancellationToken);
 
-			IStorageClient client = await _storageService.GetClientAsync(Namespace.Tools, cancellationToken);
 			TreeReader reader = new TreeReader(client, _cache, _logger);
 			DirectoryNode node = await reader.ReadNodeAsync<DirectoryNode>(deployment.RefName, DateTime.UtcNow - TimeSpan.FromDays(2.0), cancellationToken);
 
diff --git a/Engine/Source/Programs/Horde/Horde.Build/Tools/ToolsController.cs b/Engine/Source/Programs/Horde/Horde.Build/Tools/ToolsController.cs
--- a/Engine/Source/Programs/Horde/Horde.Build/Tools/ToolsController.cs
+++ b/Engine/Source/Programs/Horde/Horde.Build/Tools/ToolsController.cs
@@ -289,8 +289,18 @@ namespace Horde.Build.Tools
 		{
 			GlobalConfig globalConfig = _globalConfig.Value;
 
+			Dictionary<ToolId, ToolConfig> tools = new Dictionary<ToolId, ToolConfig>();
+			foreach (BundledToolConfig bundledToolConfig in globalConfig.ServerSettings.BundledTools)
+			{
+				tools[bundledToolConfig.Id] = bundledToolConfig;
+			}
+			foreach (ToolConfig toolConfig in globalConfig.Tools)
+			{
+				tools[toolConfig.Id] = toolConfig;
+			}
+
 			GetToolsSummaryResponse response = new GetToolsSummaryResponse();
-			foreach (ToolConfig toolConfig in globalConfig.Tools.OrderBy(x => x.Name, StringComparer.Ordinal))
+			foreach (ToolConfig toolConfig in tools.Values.OrderBy(x => x.Name, StringComparer.Ordinal))
 			{
 				if(AuthorizeDownload(toolConfig))
 				{
@@ -328,7 +338,7 @@ namespace Horde.Build.Tools
 				List<GetToolDeploymentResponse> deploymentResponses = new List<GetToolDeploymentResponse>();
 				foreach (IToolDeployment deployment in tool.Deployments)
 				{
-					GetToolDeploymentResponse deploymentResponse = await GetDeploymentInfoResponseAsync(deployment, cancellationToken);
+					GetToolDeploymentResponse deploymentResponse = await GetDeploymentInfoResponseAsync(tool, deployment, cancellationToken);
 					deploymentResponses.Add(deploymentResponse);
 				}
 				return Ok(new GetToolResponse(tool, deploymentResponses));
@@ -406,7 +416,7 @@ namespace Horde.Build.Tools
 		{
 			if (action == GetToolAction.Info)
 			{
-				GetToolDeploymentResponse response = await GetDeploymentInfoResponseAsync(deployment, cancellationToken);
+				GetToolDeploymentResponse response = await GetDeploymentInfoResponseAsync(tool, deployment, cancellationToken);
 				return Ok(response);
 			}
 			else
@@ -416,9 +426,9 @@ namespace Horde.Build.Tools
 			}
 		}
 
-		private async Task<GetToolDeploymentResponse> GetDeploymentInfoResponseAsync(IToolDeployment deployment, CancellationToken cancellationToken)
+		private async Task<GetToolDeploymentResponse> GetDeploymentInfoResponseAsync(ITool tool, IToolDeployment deployment, CancellationToken cancellationToken)
 		{
-			IStorageClient client = await _storageService.GetClientAsync(Namespace.Tools, cancellationToken);
+			IStorageClient client = await _toolCollection.GetStorageClientAsync(tool, cancellationToken);
 			NodeHandle rootHandle = await client.ReadRefTargetAsync(deployment.RefName, cancellationToken: cancellationToken);
 
 			return new GetToolDeploymentResponse(deployment, rootHandle);
@@ -452,7 +462,8 @@ namespace Horde.Build.Tools
 				return BadRequest("Invalid blob id for tool");
 			}
 
-			return StorageController.ReadBlobInternalAsync(_storageService, Namespace.Tools, locator, offset, length, cancellationToken);
+			IStorageClient storageClient = await _toolCollection.GetStorageClientAsync(tool, cancellationToken);
+			return StorageController.ReadBlobInternalAsync(storageClient, locator, offset, length, cancellationToken);
 		}
 
 		bool AuthorizeDownload(ToolConfig toolConfig)
