Horde: Updating dashboard layout, adding tools page, preview support
#rnx
#skipci

diff --git a/Engine/Source/Programs/Horde/HordeDashboard/src/App.tsx b/Engine/Source/Programs/Horde/HordeDashboard/src/App.tsx
--- a/Engine/Source/Programs/Horde/HordeDashboard/src/App.tsx
+++ b/Engine/Source/Programs/Horde/HordeDashboard/src/App.tsx
@@ -25,6 +25,7 @@ import { PreflightRedirector } from './components/Preflight';
 import { ProjectHome } from './components/ProjectHome';
 import { StreamView } from './components/StreamView';
 import { TestReportView } from './components/TestReportView';
+import { ToolView } from './components/ToolView';
 import { UserHomeView } from './components/UserHome';
 import { UtilizationReportView } from './components/UtilizationReportView';
 import hordePlugins from './Plugins';
@@ -144,6 +145,7 @@ const Main: React.FC = () => {
                { path: "audit/agent/:agentId", element: <AuditLogView /> },
                { path: "audit/issue/:issueId", element: <AuditLogView /> },
                { path: "automation", element: <AutomationView /> },
+               { path: "tools", element: <ToolView /> },
                { path: "debug/lease/:leaseId", element: <DebugView /> }
             ]
          }
diff --git a/Engine/Source/Programs/Horde/HordeDashboard/src/backend/Api.ts b/Engine/Source/Programs/Horde/HordeDashboard/src/backend/Api.ts
--- a/Engine/Source/Programs/Horde/HordeDashboard/src/backend/Api.ts
+++ b/Engine/Source/Programs/Horde/HordeDashboard/src/backend/Api.ts
@@ -4458,3 +4458,19 @@ export type GetTestDataRefResponse = {
 	/// The number of suite tests swith errors
 	suiteErrorCount?: number;
 }
+
+/** Summary for a particular tool */
+export type GetToolSummaryResponse = {
+
+	/** Unique id of tool */
+	id: string;
+
+	/** Name of tool */
+	name: string;
+
+	/** Description of tool */
+	description: string;
+
+	/** Version of tool */
+	version?: string;
+}
\ No newline at end of file
diff --git a/Engine/Source/Programs/Horde/HordeDashboard/src/backend/Backend.ts b/Engine/Source/Programs/Horde/HordeDashboard/src/backend/Backend.ts
--- a/Engine/Source/Programs/Horde/HordeDashboard/src/backend/Backend.ts
+++ b/Engine/Source/Programs/Horde/HordeDashboard/src/backend/Backend.ts
@@ -1,7 +1,7 @@
 // Copyright Epic Games, Inc. All Rights Reserved.
 
 import templateCache from '../backend/TemplateCache';
-import { AgentData, AgentQuery, ArtifactData, AuditLogEntry, AuditLogQuery, BatchUpdatePoolRequest, ChangeSummaryData, CreateDeviceRequest, CreateDeviceResponse, CreateExternalIssueRequest, CreateExternalIssueResponse, CreateJobRequest, CreateJobResponse, CreateNoticeRequest, CreatePoolRequest, CreateSoftwareResponse, CreateSubscriptionRequest, CreateSubscriptionResponse, DashboardPreference, DevicePoolTelemetryQuery, DeviceTelemetryQuery, EventData, FindIssueResponse, FindJobTimingsResponse, GetAgentSoftwareChannelResponse, GetArtifactZipRequest, GetDashboardConfigResponse, GetDevicePlatformResponse, GetDevicePoolResponse, GetDevicePoolTelemetryResponse, GetDeviceReservationResponse, GetDeviceResponse, GetDeviceTelemetryResponse, GetExternalIssueProjectResponse, GetExternalIssueResponse, GetGraphResponse, GetIssueStreamResponse, GetJobsTabResponse, GetJobStepRefResponse, GetJobStepTraceResponse, GetJobTimingResponse, GetLogEventResponse, GetNoticeResponse, GetNotificationResponse, GetPerforceServerStatusResponse, GetPoolResponse, GetServerInfoResponse, GetServerSettingsResponse, GetSoftwareResponse, GetSubscriptionResponse, GetTestDataDetailsResponse, GetTestDataRefResponse, GetTestMetaResponse, GetTestResponse, GetTestsRequest, GetTestStreamResponse, GetUserResponse, GetUtilizationTelemetryResponse, GlobalConfig, IssueData, IssueQuery, IssueQueryV2, JobData, JobQuery, JobsTabColumnType, JobStepOutcome, JobStreamQuery, JobTimingsQuery, LeaseData, LogData, LogLineData, PoolData, ProjectData, ScheduleData, ScheduleQuery, SearchLogFileResponse, ServerUpdateResponse, SessionData, StreamData, TabType, TestData, UpdateAgentRequest, UpdateDeviceRequest, UpdateGlobalConfigRequest, UpdateIssueRequest, UpdateJobRequest, UpdateLeaseRequest, UpdateNoticeRequest, UpdateNotificationsRequest, UpdatePoolRequest, UpdateServerSettingsRequest, UpdateStepRequest, UpdateStepResponse, UpdateTemplateRefRequest, UpdateUserRequest, UsersQuery } from './Api';
+import { AgentData, AgentQuery, ArtifactData, AuditLogEntry, AuditLogQuery, BatchUpdatePoolRequest, ChangeSummaryData, CreateDeviceRequest, CreateDeviceResponse, CreateExternalIssueRequest, CreateExternalIssueResponse, CreateJobRequest, CreateJobResponse, CreateNoticeRequest, CreatePoolRequest, CreateSoftwareResponse, CreateSubscriptionRequest, CreateSubscriptionResponse, DashboardPreference, DevicePoolTelemetryQuery, DeviceTelemetryQuery, EventData, FindIssueResponse, FindJobTimingsResponse, GetAgentSoftwareChannelResponse, GetArtifactZipRequest, GetDashboardConfigResponse, GetDevicePlatformResponse, GetDevicePoolResponse, GetDevicePoolTelemetryResponse, GetDeviceReservationResponse, GetDeviceResponse, GetDeviceTelemetryResponse, GetExternalIssueProjectResponse, GetExternalIssueResponse, GetGraphResponse, GetIssueStreamResponse, GetJobsTabResponse, GetJobStepRefResponse, GetJobStepTraceResponse, GetJobTimingResponse, GetLogEventResponse, GetNoticeResponse, GetNotificationResponse, GetPerforceServerStatusResponse, GetPoolResponse, GetServerInfoResponse, GetServerSettingsResponse, GetSoftwareResponse, GetSubscriptionResponse, GetTestDataDetailsResponse, GetTestDataRefResponse, GetTestMetaResponse, GetTestResponse, GetTestsRequest, GetTestStreamResponse, GetToolSummaryResponse, GetUserResponse, GetUtilizationTelemetryResponse, GlobalConfig, IssueData, IssueQuery, IssueQueryV2, JobData, JobQuery, JobsTabColumnType, JobStepOutcome, JobStreamQuery, JobTimingsQuery, LeaseData, LogData, LogLineData, PoolData, ProjectData, ScheduleData, ScheduleQuery, SearchLogFileResponse, ServerUpdateResponse, SessionData, StreamData, TabType, TestData, UpdateAgentRequest, UpdateDeviceRequest, UpdateGlobalConfigRequest, UpdateIssueRequest, UpdateJobRequest, UpdateLeaseRequest, UpdateNoticeRequest, UpdateNotificationsRequest, UpdatePoolRequest, UpdateServerSettingsRequest, UpdateStepRequest, UpdateStepResponse, UpdateTemplateRefRequest, UpdateUserRequest, UsersQuery } from './Api';
 import dashboard from './Dashboard';
 import { ChallengeStatus, Fetch } from './Fetch';
 import graphCache, { GraphQuery } from './GraphCache';
@@ -1534,6 +1534,17 @@ export class Backend {
         });
     }
 
+    // get tools
+    getTools(): Promise<GetToolSummaryResponse[]> {
+        return new Promise<GetToolSummaryResponse[]>((resolve, reject) => {
+            this.backend.get(`/api/v1/tools`).then((value) => {
+                resolve(value.data?.tools as GetToolSummaryResponse[]);
+            }).catch(reason => {
+                reject(reason);
+            });
+        });
+    }
+    
     // update a notice
     updateNotice(request: UpdateNoticeRequest): Promise<void> {
         return new Promise<void>((resolve, reject) => {
diff --git a/Engine/Source/Programs/Horde/HordeDashboard/src/components/AgentView.tsx b/Engine/Source/Programs/Horde/HordeDashboard/src/components/AgentView.tsx
--- a/Engine/Source/Programs/Horde/HordeDashboard/src/components/AgentView.tsx
+++ b/Engine/Source/Programs/Horde/HordeDashboard/src/components/AgentView.tsx
@@ -74,8 +74,10 @@ const agentStyles = mergeStyleSets({
    },
    buttonFont: {
       height: '26px',
-      font: '8pt Horde Open Sans SemiBold !important',
+      font: '7pt Horde Open Sans SemiBold !important',
       flexShrink: '0 !important',
+      paddingLeft: 6,
+      paddingRight: 6,
       selectors: {
          '.ms-Icon': {
             width: 0,
@@ -619,7 +621,7 @@ class LocalState {
          {
             key: 'name',
             displayText: 'Name',
-            colSize: 175,
+            colSize: 145,
             isChecked: true,
             isCheckable: false,
             isSorted: true,
@@ -629,7 +631,7 @@ class LocalState {
          {
             key: 'pools',
             displayText: 'Pools',
-            colSize: 840,
+            colSize: 540,
             isChecked: true,
             isCheckable: false,
             isSorted: false,
@@ -1751,7 +1753,7 @@ export const AgentView: React.FC = observer(() => {
          <Breadcrumbs items={[{ text: 'Admin' }, { text: 'Agents' }]} />
          <Stack horizontal>
             <Stack grow styles={{ root: { backgroundColor: 'rgb(250, 249, 249)' } }} />
-            <Stack tokens={{ maxWidth: 1800, childrenGap: 4 }} styles={{ root: { width: 1800, height: '100vh', backgroundColor: 'rgb(250, 249, 249)', paddingTop: 18, paddingLeft: 12 } }}>
+            <Stack tokens={{ maxWidth: 1440, childrenGap: 4 }} styles={{ root: { width: 1440, height: '100vh', backgroundColor: 'rgb(250, 249, 249)', paddingTop: 18, paddingLeft: 12 } }}>
                <Stack className={hordeClasses.raised} styles={{ root: { paddingRight: '40px' } }}>
                   <Stack style={{ position: "relative", height: "calc(100vh - 240px)" }} horizontalAlign="center">
                      <Stack horizontal tokens={{ childrenGap: 24 }}>
@@ -2102,7 +2104,7 @@ export const AgentView: React.FC = observer(() => {
          <Breadcrumbs items={[{ text: 'Admin' }, { text: 'Agents' }]} />
          <Stack horizontal>
             <Stack grow styles={{ root: { backgroundColor: 'rgb(250, 249, 249)' } }} />
-            <Stack tokens={{ maxWidth: 1800, childrenGap: 4 }} styles={{ root: { width: 1800, height: '100vh', backgroundColor: 'rgb(250, 249, 249)', paddingTop: 18, paddingLeft: 12 } }}>
+            <Stack tokens={{ maxWidth: 1440, childrenGap: 4 }} styles={{ root: { width: 1440, height: '100vh', backgroundColor: 'rgb(250, 249, 249)', paddingTop: 18, paddingLeft: 12 } }}>
                <Stack className={hordeClasses.raised} styles={{ root: { paddingRight: '40px' } }}>
                   <SearchUpdate />
                   <Stack.Item>
@@ -2398,7 +2400,7 @@ export const AgentView: React.FC = observer(() => {
                let subtitle = "";
                if (!agent.online) {
                   title = `Offline - ${agent.lastShutdownReason}`;
-                  subtitle = `(Last online at ${getShortNiceTime(agent.updateTime, true, true, true)}`;
+                  subtitle = `Last online at ${getShortNiceTime(agent.updateTime, true, true, true)}`;
                }
                else if (agent.pendingShutdown) {
                   title = `Pending Shutdown`;
@@ -2428,7 +2430,7 @@ export const AgentView: React.FC = observer(() => {
                leaseSearchItems.push(title);
                if (subtitle !== "") {
                   leases.push(<Stack.Item key={"itemSubtitleStatus_" + agent.id} align={"center"} className={agentStyles.ellipsesStackItem}>
-                     <Text key={"statusSubtitleText_" + agent.id}>{subtitle}</Text>
+                     <Text variant="small" key={"statusSubtitleText_" + agent.id}>{subtitle}</Text>
                   </Stack.Item>);
                   leaseSearchItems.push(subtitle);
                }
diff --git a/Engine/Source/Programs/Horde/HordeDashboard/src/components/AuditLog.tsx b/Engine/Source/Programs/Horde/HordeDashboard/src/components/AuditLog.tsx
--- a/Engine/Source/Programs/Horde/HordeDashboard/src/components/AuditLog.tsx
+++ b/Engine/Source/Programs/Horde/HordeDashboard/src/components/AuditLog.tsx
@@ -460,9 +460,9 @@ export const AuditLogPanel: React.FC<{ agentId?: string, issueId?: string }> = o
       <Breadcrumbs items={crumbItems} />
       <Stack tokens={{ childrenGap: 12 }}>
          <Stack horizontal>
-            <div key={`windowsize_logview1_${windowSize.width}_${windowSize.height}`} style={{ width: vw / 2 - 880, flexShrink: 0, backgroundColor: '#FFFFFF' }} />
+            <div key={`windowsize_logview1_${windowSize.width}_${windowSize.height}`} style={{ width: vw / 2 - (1440/2) - 230, flexShrink: 0, backgroundColor: '#FFFFFF' }} />
             <Stack tokens={{ childrenGap: 0 }} styles={{ root: { backgroundColor: "#FFFFFF", margin: "auto", paddingTop: 12, paddingRight: 10 } }}>
-               <Stack horizontal styles={{ root: { paddingLeft: 0, paddingBottom: 4, paddingRight: 12, width: 1800 } }}>
+               <Stack horizontal styles={{ root: { paddingLeft: 0, paddingBottom: 4, paddingRight: 12, width: 1440 } }}>
                   <Stack horizontal tokens={{ childrenGap: 12 }}>
                      <Stack>
                         <SearchBox style={{ width: 400 }} onEscape={() => setSearch(undefined)} autoComplete="off" disableAnimation={true} spellCheck={false} onChange={(ev, value) => setSearch(value)} />
@@ -563,7 +563,7 @@ export const AuditLogPanel: React.FC<{ agentId?: string, issueId?: string }> = o
                <FocusZone direction={FocusZoneDirection.vertical} isInnerZoneKeystroke={() => { return true; }} defaultActiveElement="#LogList" style={{ padding: 0, margin: 0 }} >
                   <div className={auditStyleNormal.container} style={{ height: 'calc(100vh - 260px)', position: 'relative' }} data-is-scrollable={true}>
                      <Stack horizontal>
-                        <div key={`windowsize_logview2_${windowSize.width}_${windowSize.height}`} style={{ width: vw / 2 - 880 - 24, flexShrink: 0 }} />
+                        <div key={`windowsize_logview2_${windowSize.width}_${windowSize.height}`} style={{ width: vw / 2 - (1440/2) - 24, flexShrink: 0 }} />
                         <Stack styles={{ root: { backgroundColor: "#FFFFFF", paddingLeft: "0px", paddingRight: "0px" } }}>
                            {!handler.haveUpdated && <Spinner size={SpinnerSize.large} />}
                            {!!handler.haveUpdated && !logItems.length && <Stack style={{ paddingLeft: 0 }}><Text variant="mediumPlus">No audit entries found</Text></Stack>}
diff --git a/Engine/Source/Programs/Horde/HordeDashboard/src/components/AutomationView.tsx b/Engine/Source/Programs/Horde/HordeDashboard/src/components/AutomationView.tsx
--- a/Engine/Source/Programs/Horde/HordeDashboard/src/components/AutomationView.tsx
+++ b/Engine/Source/Programs/Horde/HordeDashboard/src/components/AutomationView.tsx
@@ -534,7 +534,7 @@ export const AutomationView: React.FC = observer(() => {
          <Breadcrumbs items={[{ text: 'Automation' }]} />
          <ErrorBoundary>
             <Stack horizontalAlign="center" grow styles={{ root: { width: "100%", padding: 12, backgroundColor: modeColors.background } }}>
-               <Stack styles={{ root: { width: 1774 } }}>
+               <Stack styles={{ root: { width: 1440 } }}>
                   {!handler.loaded && <Stack>
                   </Stack>}
                   {handler.loaded && <Stack style={{ paddingTop: 8 }}>
@@ -745,7 +745,7 @@ class AutomationGraph {
 
       const handler = this.handler;
       const refs = this.refs.sort((a, b) => handler.metaNames.get(a.metaId)!.localeCompare(handler.metaNames.get(b.metaId)!));
-      const width = 1328
+      const width = 1000
 
       const scolors = dashboard.getStatusColors();
       const colors: Record<string, string> = {
diff --git a/Engine/Source/Programs/Horde/HordeDashboard/src/components/Breadcrumbs.tsx b/Engine/Source/Programs/Horde/HordeDashboard/src/components/Breadcrumbs.tsx
--- a/Engine/Source/Programs/Horde/HordeDashboard/src/components/Breadcrumbs.tsx
+++ b/Engine/Source/Programs/Horde/HordeDashboard/src/components/Breadcrumbs.tsx
@@ -80,7 +80,7 @@ export const Breadcrumbs: React.FC<{ items: BreadcrumbItem[], title?: string, su
       <Stack>
          <Stack styles={{ root: { userInput: 'all' } }}>
             <Stack tokens={{ childrenGap: 4 }} styles={{ root: { height: 88, padding: 0, paddingLeft: 24, paddingBottom: 8, paddingTop: 8, backgroundColor: modeColors.crumbs, userSelect: 'text' } }}>
-               <Stack tokens={{ childrenGap: 0 }} disableShrink={true} styles={{ root: { margin: "auto", width: "100%", maxWidth: 1800 } }}>
+               <Stack tokens={{ childrenGap: 0 }} disableShrink={true} styles={{ root: { margin: "auto", width: "100%", maxWidth: 1464 } }}>
                   <Stack horizontal tokens={{ childrenGap: 8, padding: 0 }}>{topElements}</Stack>
                   {
                      <Stack horizontal tokens={{ childrenGap: 18 }}>
diff --git a/Engine/Source/Programs/Horde/HordeDashboard/src/components/DashboardView.tsx b/Engine/Source/Programs/Horde/HordeDashboard/src/components/DashboardView.tsx
--- a/Engine/Source/Programs/Horde/HordeDashboard/src/components/DashboardView.tsx
+++ b/Engine/Source/Programs/Horde/HordeDashboard/src/components/DashboardView.tsx
@@ -410,9 +410,9 @@ export const DashboardView: React.FC = () => {
       <TopNav />
       <Breadcrumbs items={[{ text: 'Preferences' }]} />
       <Stack horizontal>
-         <div key={`windowsize_streamview_${windowSize.width}_${windowSize.height}`} style={{ width: vw / 2 - 900, flexShrink: 0, backgroundColor: modeColors.background }} />
+         <div key={`windowsize_streamview_${windowSize.width}_${windowSize.height}`} style={{ width: vw / 2 - (1440/2), flexShrink: 0, backgroundColor: modeColors.background }} />
          <Stack tokens={{ childrenGap: 0 }} styles={{ root: { backgroundColor: modeColors.background, width: "100%" } }}>
-            <Stack style={{ maxWidth: 1800, paddingTop: 16, marginLeft: 4 }}>
+            <Stack style={{ maxWidth: 1440, paddingTop: 16, marginLeft: 4 }}>
                <Stack horizontal className={hordeClasses.raised}>
                   <Stack style={{ width: "100%" }} tokens={{ childrenGap: 18 }}>
                      {dashboard.available && <GeneralPanel />}
diff --git a/Engine/Source/Programs/Horde/HordeDashboard/src/components/DeviceView.tsx b/Engine/Source/Programs/Horde/HordeDashboard/src/components/DeviceView.tsx
--- a/Engine/Source/Programs/Horde/HordeDashboard/src/components/DeviceView.tsx
+++ b/Engine/Source/Programs/Horde/HordeDashboard/src/components/DeviceView.tsx
@@ -140,20 +140,21 @@ const DevicePanel: React.FC = observer(() => {
    const platforms = handler.platforms;
 
    const poolWidth = automationTab ? 80 : 140;
-   const detailsWidth = automationTab ? 180 : 120;
+   const detailsWidth = automationTab ? 160 : 100;
 
    const columns: IColumn[] = [
-      { key: 'column_platform', name: 'Platform', fieldName: 'platformId', minWidth: 80, maxWidth: 80, isResizable: false },
+      { key: 'column_platform', name: 'Platform', fieldName: 'platformId', minWidth: 72, maxWidth: 72, isResizable: false },
       { key: 'column_pool', name: 'Pool', fieldName: 'poolId', minWidth: poolWidth, maxWidth: poolWidth, isResizable: false },
       { key: 'column_model', name: 'Model', fieldName: 'modelId', minWidth: 80, maxWidth: 80, isResizable: false },
-      { key: 'column_name', name: 'Name', fieldName: 'name', minWidth: 160, maxWidth: 160, isResizable: false },
+      { key: 'column_name', name: 'Name', fieldName: 'name', minWidth: 120, maxWidth: 120, isResizable: false },
       { key: 'column_address', name: 'Address', fieldName: 'address', minWidth: 120, maxWidth: 120, isResizable: false },
       { key: 'column_status', name: 'Status', fieldName: 'status', minWidth: 120, maxWidth: 120, isResizable: false },
       { key: 'column_details', name: 'Details', fieldName: 'details', minWidth: detailsWidth, maxWidth: detailsWidth, isResizable: false },
-      { key: 'column_notes', name: 'Notes', fieldName: 'notes', minWidth: 312, maxWidth: 312, isResizable: false },
+      { key: 'column_notes', name: 'Notes', fieldName: 'notes', minWidth: 240, maxWidth: 240, isResizable: false },
    ];
 
    columns.forEach(c => {
+      c.isPadded = false;
       c.onRender = ((item: DeviceItem, index, column) => {
          return <Stack verticalFill={true} verticalAlign="center">
             <Text>{(item.device as any)[column?.fieldName!]}</Text>
@@ -543,7 +544,7 @@ const DevicePanel: React.FC = observer(() => {
                <div className={customStyles.details} style={{ height: "calc(100vh - 280px)", position: 'relative' }} data-is-scrollable>
                   <ScrollablePane scrollbarVisibility={ScrollbarVisibility.always} onScroll={() => { }}>
                      <DetailsList
-                        styles={{ root: { overflowX: "hidden", width: 1676 } }}
+                        styles={{ root: { overflowX: "hidden", width: 1312 } }}
                         items={devices}
                         groups={groups}
                         columns={columns}
@@ -665,9 +666,9 @@ export const DeviceView: React.FC = () => {
       <TopNav />
       <Breadcrumbs items={[{ text: 'Devices' }]} />
       <Stack horizontal>
-         <div key={`windowsize_streamview_${windowSize.width}_${windowSize.height}`} style={{ width: vw / 2 - 900, flexShrink: 0, backgroundColor: modeColors.background }} />
+         <div key={`windowsize_streamview_${windowSize.width}_${windowSize.height}`} style={{ width: vw / 2 - (1440/2), flexShrink: 0, backgroundColor: modeColors.background }} />
          <Stack tokens={{ childrenGap: 0 }} styles={{ root: { backgroundColor: modeColors.background, width: "100%" } }}>
-            <Stack style={{ maxWidth: 1800, paddingTop: 6, marginLeft: 4, height: 'calc(100vh - 8px)' }}>
+            <Stack style={{ maxWidth: 1440, paddingTop: 6, marginLeft: 4, height: 'calc(100vh - 8px)' }}>
                <Stack horizontal className={hordeClasses.raised}>
                   <Stack style={{ width: "100%", height: 'calc(100vh - 228px)' }} tokens={{ childrenGap: 18 }}>
                      <DevicePanel />
diff --git a/Engine/Source/Programs/Horde/HordeDashboard/src/components/IssueViewV2.tsx b/Engine/Source/Programs/Horde/HordeDashboard/src/components/IssueViewV2.tsx
--- a/Engine/Source/Programs/Horde/HordeDashboard/src/components/IssueViewV2.tsx
+++ b/Engine/Source/Programs/Horde/HordeDashboard/src/components/IssueViewV2.tsx
@@ -742,7 +742,7 @@ const StreamCanvas: React.FC = () => {
             const maxColumns = 5;
             const totalDays = Math.ceil(moment.duration(moment(maxTime * 1000).diff(minTime * 1000)).asDays());
 
-            const spanWidth = 1180;
+            const spanWidth = 980;
             const columnWidth = spanWidth / maxColumns;
 
             for (let i = 0; i < maxColumns; i++) {
@@ -807,7 +807,7 @@ const StreamCanvas: React.FC = () => {
          }
 
 
-         context.font = "14px Horde Open Sans SemiBold";
+         context.font = "12px Horde Open Sans SemiBold";
          context.fillStyle = color2;
          context.textAlign = 'left';
          context.textBaseline = 'middle';
@@ -833,7 +833,7 @@ const StreamCanvas: React.FC = () => {
                gotone = true;
 
                if (drawTemplateText) {
-                  context.font = "13px Horde Open Sans SemiBold";
+                  context.font = "11px Horde Open Sans SemiBold";
                   context.fillStyle = color0;
                   context.textAlign = 'left';
                   context.textBaseline = 'middle';
@@ -847,7 +847,7 @@ const StreamCanvas: React.FC = () => {
                   return timeA.getTime() - timeB.getTime();
                });
 
-               context.font = "13px Horde Open Sans Regular";
+               context.font = "11px Horde Open Sans Regular";
                context.fillStyle = color0;
                context.textAlign = 'left';
                context.textBaseline = 'middle';
@@ -861,8 +861,8 @@ const StreamCanvas: React.FC = () => {
 
                   const step = jobStep.step;
 
-                  const width = 1180;
-                  const startX = 300;
+                  const width = 980;
+                  const startX = 280;
 
                   let color = color1;
                   if (step.severity === IssueSeverity.Warning) {
@@ -912,12 +912,12 @@ const StreamCanvas: React.FC = () => {
 
                   }
 
-                  let size = 5;
+                  let size = 4;
                   if (chosenOne === 1) {
-                     size = 7;
+                     size = 6;
                   }
                   if (chosenOne === 2) {
-                     size = 10;
+                     size = 8;
                   }
 
 
@@ -1017,12 +1017,13 @@ const StreamCanvas: React.FC = () => {
             const metrics = context.measureText(change);
 
             context.fillStyle = "#F3F2F1";
-            context.fillRect(drawCL.minX - metrics.width / 2 - 6, drawCL.minY - 28, metrics.width + 24, 25);
+            const minx = Math.min(drawCL.minX, 1200);
+            context.fillRect(minx - metrics.width / 2 - 3, drawCL.minY - 28, metrics.width + 21, 25);
 
 
 
             context.fillStyle = "#2D3F5F";
-            context.fillText(change, drawCL.minX + 6, drawCL.minY - 21);
+            context.fillText(change, minx + 6, drawCL.minY - 21);
 
          }
       });
@@ -1111,7 +1112,7 @@ const StreamCanvas: React.FC = () => {
    }
 
    // Blurry text render fix
-   const canvasBaseWidth = 1560;
+   const canvasBaseWidth = 1280;
    const ratio = Math.ceil(window.devicePixelRatio);
    const canvasWidth = canvasBaseWidth * ratio;
    const canvasHeight = height * ratio;
@@ -1127,7 +1128,7 @@ const StreamCanvas: React.FC = () => {
          <Stack className="horde-no-darktheme" style={{ position: "relative", height: "100%" }}>
             <div ref={scrollRef} style={{ overflowX: "hidden", overflowY: "auto" }}>
                <Stack>
-                  <canvas ref={newRef => setCanvas(newRef)} width={canvasWidth} height={canvasHeight} onWheel={onChartScroll} style={{ width: canvasStyleWidth, height: canvasStyleHeight }} onClick={(ev) => onCanvasClick(ev)} onMouseLeave={(ev => onCanvasMouseLeave(ev))} onMouseMove={(ev) => onCanvasMouseMove(ev)}></canvas>
+                  <canvas ref={newRef => setCanvas(newRef)} width={canvasWidth} height={canvasHeight} onWheel={onChartScroll} style={{ width: canvasStyleWidth, height: canvasStyleHeight, overflowX: "visible" }} onClick={(ev) => onCanvasClick(ev)} onMouseLeave={(ev => onCanvasMouseLeave(ev))} onMouseMove={(ev) => onCanvasMouseMove(ev)}></canvas>
                </Stack>
             </div>
          </Stack>
@@ -1344,7 +1345,7 @@ export const ErrorPane: React.FC<{ events?: GetLogEventResponse[]; onClose?: ()
 
       const url = `/log/${item.logId}?lineindex=${item.lineIndex}`;
 
-      const lines = item.lines.filter(line => line.message?.trim().length).map(line => <Stack key={`errorpane_line_${item.lineIndex}_${lineKey++}`} styles={{ root: { paddingLeft: 8, paddingRight: 8, lineBreak: "anywhere", whiteSpace: "pre-wrap", lineHeight: 18, fontSize: 10, fontFamily: "Horde Cousine Regular, monospace, monospace" } }}> <Link className="log-link" to={url}>{renderLine(line, undefined, {})}</Link></Stack>);
+      const lines = item.lines.filter(line => line.message?.trim().length).map(line => <Stack key={`errorpane_line_${item.lineIndex}_${lineKey++}`} styles={{ root: { paddingLeft: 8, paddingRight: 8, lineBreak: "normal", whiteSpace: "pre-wrap", lineHeight: 18, fontSize: 10, fontFamily: "Horde Cousine Regular, monospace, monospace" } }}> <Link className="log-link" to={url}>{renderLine(line, undefined, {})}</Link></Stack>);
 
       return (<Stack className={errorStyles.itemCell} style={{ padding: 8 }}><Stack className={item.severity === EventSeverity.Warning ? errorStyles.gutterWarning : errorStyles.gutter} styles={{ root: { padding: 0, margin: 0 } }}>
          <Stack styles={{ root: { paddingLeft: 14 } }}>
@@ -1860,15 +1861,15 @@ export const IssueModalV2: React.FC<{ popHistoryOnClose: boolean, issueId?: stri
       </Modal>
    }
 
-   return <Modal isOpen={true} isBlocking={true} topOffsetFixed={true} styles={{ main: { padding: 8, width: 1700, backgroundColor: modeColors.background, hasBeenOpened: false, top: "24px", position: "absolute", height: "95vh" } }} className={hordeClasses.modal} onDismiss={() => { if (onCloseExternal) { onCloseExternal() } else { onClose() } }}>
+   return <Modal isOpen={true} isBlocking={true} topOffsetFixed={true} styles={{ main: { padding: 8, width: 1420, backgroundColor: modeColors.background, hasBeenOpened: false, top: "24px", position: "absolute", height: "95vh" } }} className={hordeClasses.modal} onDismiss={() => { if (onCloseExternal) { onCloseExternal() } else { onClose() } }}>
       {editShown && <EditIssueModal onClose={() => { setEditShown(false) }} />}
       <Stack style={{ height: "93vh" }}>
          <Stack style={{ height: "100%" }}>
             <Stack style={{ flexBasis: "70px", flexShrink: 0 }}>
                <Stack horizontal styles={{ root: { padding: 8 } }} style={{ padding: 20, paddingBottom: 8 }}>
                   <Stack horizontal style={{ width: 1024 }} tokens={{ childrenGap: 24 }} verticalAlign="center" verticalFill={true}>
-                     <Stack style={{ maxWidth: 880, lineBreak: "anywhere" }} >
-                        <Text styles={{ root: { fontWeight: "unset", fontFamily: "Horde Open Sans SemiBold", fontSize: "15px", color: "#087BC4", textDecoration: details.issue?.resolvedAt ? "line-through" : undefined } }}>{title}</Text>
+                     <Stack style={{ maxWidth: 880, lineBreak: "normal" }} >
+                        <Text styles={{ root: { fontWeight: "unset", fontFamily: "Horde Open Sans SemiBold", fontSize: "14px", color: "#087BC4", textDecoration: details.issue?.resolvedAt ? "line-through" : undefined } }}>{title}</Text>
                      </Stack>
 
                      <Stack onClick={() => { setEditShown(true) }} style={{ cursor: "pointer" }}>
diff --git a/Engine/Source/Programs/Horde/HordeDashboard/src/components/JobView.tsx b/Engine/Source/Programs/Horde/HordeDashboard/src/components/JobView.tsx
--- a/Engine/Source/Programs/Horde/HordeDashboard/src/components/JobView.tsx
+++ b/Engine/Source/Programs/Horde/HordeDashboard/src/components/JobView.tsx
@@ -61,7 +61,7 @@ const buildColumns = (jobTab: JobsTabData): IColumn[] => {
    if (jobTab.columns) {
       let total = 0;
       jobTab.columns.forEach(c => total += c.relativeWidth ?? 1);
-      const w = (jobTab.showNames ? 1000 : 800) / total;
+      const w = (jobTab.showNames ? 950 - 300 : 750 - 300) / total;
 
       jobTab.columns.forEach(c => { minWidths[c.heading] = w * (c.relativeWidth ?? 1); cnames.push(c.heading); });
    } else {
@@ -562,7 +562,7 @@ const JobList: React.FC<{ tab: string; filter: JobFilterSimple, controller: Call
 
    let nojobs = !jobHandler.initial && !jobs.length;
 
-   const width = 1800;
+   const width = 1440;
 
    // main header
    const onRenderDetailsHeader: IDetailsListProps['onRenderDetailsHeader'] = (props) => {
diff --git a/Engine/Source/Programs/Horde/HordeDashboard/src/components/JobViewAll.tsx b/Engine/Source/Programs/Horde/HordeDashboard/src/components/JobViewAll.tsx
--- a/Engine/Source/Programs/Horde/HordeDashboard/src/components/JobViewAll.tsx
+++ b/Engine/Source/Programs/Horde/HordeDashboard/src/components/JobViewAll.tsx
@@ -122,10 +122,10 @@ const JobViewAllInner: React.FC<{ filter: JobFilterSimple }> = observer(({ filte
     let columns: IColumn[] = [
         { key: 'jobview_column1', name: 'Status', minWidth: 16, maxWidth: 16 },
         { key: 'jobview_column2', name: 'Change', minWidth: 64, maxWidth: 64 },
-        { key: 'jobview_column3', name: 'Job', minWidth: 420, maxWidth: 420 },
-        { key: 'jobview_column4', name: 'Labels', minWidth: 380, maxWidth: 380 },
-        { key: 'jobview_column5', name: 'Steps', minWidth: 504, maxWidth: 504 },
-        { key: 'jobview_column6', name: 'StartedBy', minWidth: 160, maxWidth: 160 },
+        { key: 'jobview_column3', name: 'Job', minWidth: 220, maxWidth: 220 },
+        { key: 'jobview_column4', name: 'Labels', minWidth: 300, maxWidth: 300 },
+        { key: 'jobview_column5', name: 'Steps', minWidth: 324, maxWidth: 324 },
+        { key: 'jobview_column6', name: 'StartedBy', minWidth: 290, maxWidth: 290 },
         { key: 'jobview_column7', name: 'Time', minWidth: 50, maxWidth: 50 },
 
     ];
@@ -527,7 +527,7 @@ const JobViewAllInner: React.FC<{ filter: JobFilterSimple }> = observer(({ filte
     return (
         <Stack>
             <Stack tokens={{ childrenGap: 0 }} style={{}}>
-                <Stack horizontalAlign="center" style={{ backgroundColor: "#FFFFFF", width: 1800, marginLeft: 4, boxShadow: "0 3px 3.6px 0 rgba(0,0,0,0.132), 0 0.3px 0.9px 0 rgba(0,0,0,0.108)" }}>
+                <Stack horizontalAlign="center" style={{ backgroundColor: "#FFFFFF", width: 1440, marginLeft: 4, boxShadow: "0 3px 3.6px 0 rgba(0,0,0,0.132), 0 0.3px 0.9px 0 rgba(0,0,0,0.108)" }}>
                     <Stack style={{ paddingTop: 18, paddingBottom: 24 }}>
                         <JobFilterPanel />
                     </Stack>
@@ -537,7 +537,7 @@ const JobViewAllInner: React.FC<{ filter: JobFilterSimple }> = observer(({ filte
                 <FocusZone direction={FocusZoneDirection.vertical}>
                     <div className={detailClasses.container} style={{ height: 'calc(100vh - 352px)', position: 'relative', marginTop: 0 }} data-is-scrollable={true}>
                         {<ScrollablePane scrollbarVisibility={ScrollbarVisibility.always} style={{ overflow: "visible" }}>
-                            <Stack style={{ width: 1800, marginLeft: 4, boxShadow: "0 1.6px 3.6px 0 rgba(0,0,0,0.132), 0 0.3px 0.9px 0 rgba(0,0,0,0.108)", backgroundColor: "#FFFFFF" }}>
+                            <Stack style={{ width: 1440, marginLeft: 4, boxShadow: "0 1.6px 3.6px 0 rgba(0,0,0,0.132), 0 0.3px 0.9px 0 rgba(0,0,0,0.108)", backgroundColor: "#FFFFFF" }}>
                                 <Stack>
                                     <DetailsList
                                         styles={{ root: { paddingLeft: 8, paddingRight: 8 } }}
@@ -565,7 +565,7 @@ const JobViewAllInner: React.FC<{ filter: JobFilterSimple }> = observer(({ filte
                             }
 
                             {!nojobs && !jobItems.length && <Stack styles={{ root: { paddingTop: 20, paddingBottom: 20 } }}>
-                                <Stack style={{ width: 1800 }}><Spinner size={SpinnerSize.large} /></Stack>
+                                <Stack style={{ width: 1440 }}><Spinner size={SpinnerSize.large} /></Stack>
                             </Stack>
                             }
 
diff --git a/Engine/Source/Programs/Horde/HordeDashboard/src/components/JobViewIncremental.tsx b/Engine/Source/Programs/Horde/HordeDashboard/src/components/JobViewIncremental.tsx
--- a/Engine/Source/Programs/Horde/HordeDashboard/src/components/JobViewIncremental.tsx
+++ b/Engine/Source/Programs/Horde/HordeDashboard/src/components/JobViewIncremental.tsx
@@ -61,7 +61,7 @@ const buildColumns = (jobTab: JobsTabData, streamId: string): IColumn[] => {
 
       let total = 0;
       jobTab.columns.forEach(c => total += c.relativeWidth ?? 1);
-      const w = (1350) / total;
+      const w = (980) / total;
 
       jobTab.columns.forEach(c => { minWidths[c.heading] = w * (c.relativeWidth ?? 1); cnames.push(c.heading); });
    } else {
@@ -304,7 +304,7 @@ const JobList: React.FC<{ tab: string; filter: JobFilterSimple, controller: Call
 
    };
 
-   const width = 1800;
+   const width = 1440;
 
    // main header
    const onRenderDetailsHeader: IDetailsListProps['onRenderDetailsHeader'] = (props) => {
@@ -460,7 +460,7 @@ const JobList: React.FC<{ tab: string; filter: JobFilterSimple, controller: Call
 
       return (
          <Stack horizontalAlign="center" onMouseMove={() => controller.setState(undefined)}>
-            <Stack wrap horizontal horizontalAlign="center" tokens={{ childrenGap: 4 }} styles={{ root: { paddingTop: 2, width: column?.maxWidth } }}>
+            <Stack horizontal horizontalAlign="center" tokens={{ childrenGap: 4 }} styles={{ root: { paddingTop: 2, width: column?.maxWidth } }}>
                {buttons}
             </Stack>
          </Stack>
@@ -607,7 +607,7 @@ const JobList: React.FC<{ tab: string; filter: JobFilterSimple, controller: Call
                         <DetailsList
                            styles={{
                               headerWrapper: { overflow: "hidden" }, root: {
-                                 paddingBottom: 32, 
+                                 paddingBottom: 32,
                                  selectors: {
                                     '.ms-List-cell': {
                                        minHeight: lineHeight
diff --git a/Engine/Source/Programs/Horde/HordeDashboard/src/components/LogView.tsx b/Engine/Source/Programs/Horde/HordeDashboard/src/components/LogView.tsx
--- a/Engine/Source/Programs/Horde/HordeDashboard/src/components/LogView.tsx
+++ b/Engine/Source/Programs/Horde/HordeDashboard/src/components/LogView.tsx
@@ -943,8 +943,8 @@ export const LogList: React.FC<{ logId: string }> = observer(({ logId }) => {
       <Breadcrumbs items={logSource?.crumbs ?? []} title={logSource?.crumbTitle} />
       <Stack tokens={{ childrenGap: 0 }} style={{ backgroundColor: "#FFFFFF", paddingTop: 12 }}>
          <Stack horizontal >
-            <div key={`windowsize_logview1_${windowSize.width}_${windowSize.height}`} style={{ width: vw / 2 - 880, flexShrink: 0 }} />
-            <Stack tokens={{ childrenGap: 0, maxWidth: 1780 }} disableShrink={true} styles={{ root: { width: "100%", backgroundColor: "#FFFFFF", paddingLeft: 4, paddingRight: 24, paddingTop: 12 } }}>
+            <div key={`windowsize_logview1_${windowSize.width}_${windowSize.height}`} style={{ width: vw / 2 - (1440 / 2), flexShrink: 0 }} />
+            <Stack tokens={{ childrenGap: 0, maxWidth: 1440 }} disableShrink={true} styles={{ root: { width: "100%", backgroundColor: "#FFFFFF", paddingLeft: 4, paddingRight: 24, paddingTop: 12 } }}>
                <Stack horizontal style={{paddingBottom: 4}}>
                   <Stack className={hordeClasses.button} horizontal horizontalAlign={"start"} verticalAlign="center" tokens={{ childrenGap: 8 }}>
                      <Stack horizontal tokens={{ childrenGap: 2 }}>
@@ -1119,7 +1119,7 @@ export const LogList: React.FC<{ logId: string }> = observer(({ logId }) => {
 
                      }}>
                      <Stack horizontal>
-                        {!dashboard.leftAlignLog && <div key={`windowsize_logview2_${windowSize.width}_${windowSize.height}`} style={{ width: vw / 2 - 880 - 24, flexShrink: 0 }} />}
+                        {!dashboard.leftAlignLog && <div key={`windowsize_logview2_${windowSize.width}_${windowSize.height}`} style={{ width: vw / 2 - (1440 / 2) - 48, flexShrink: 0 }} />}
                         <Stack styles={{ root: { backgroundColor: "#FFFFFF", paddingLeft: "0px", paddingRight: "0px" } }}>
                            <SelectionZone selection={selection} selectionMode={SelectionMode.multiple}>
                               <List key={`log_list_key_${logListKey}`} id="LogList" ref={(list: List) => { listRef = list; }}
diff --git a/Engine/Source/Programs/Horde/HordeDashboard/src/components/NoticeView.tsx b/Engine/Source/Programs/Horde/HordeDashboard/src/components/NoticeView.tsx
--- a/Engine/Source/Programs/Horde/HordeDashboard/src/components/NoticeView.tsx
+++ b/Engine/Source/Programs/Horde/HordeDashboard/src/components/NoticeView.tsx
@@ -23,10 +23,10 @@ const NoticePanel: React.FC = observer(() => {
    if (notices.updated) { };
 
    const columns = [
-      { key: 'column_message', name: 'Message', minWidth: 1024, maxWidth: 1024, isResizable: false },
+      { key: 'column_message', name: 'Message', minWidth: 740, maxWidth: 740, isResizable: false },
       { key: 'column_starttime', name: 'Start Time', minWidth: 160, maxWidth: 160, isResizable: false },
       { key: 'column_finishtime', name: 'Finish Time', minWidth: 160, maxWidth: 160, isResizable: false },
-      { key: 'column_createdby', name: 'Created By', minWidth: 240, maxWidth: 240, isResizable: false },
+      { key: 'column_createdby', name: 'Created By', minWidth: 200, maxWidth: 200, isResizable: false },
    ];
 
    let allNotices = [...notices.allNotices];
@@ -118,9 +118,9 @@ export const NoticeView: React.FC = () => {
       <TopNav />
       <Breadcrumbs items={[{ text: 'Server Notices' }]} />
       <Stack horizontal>
-         <div key={`windowsize_noticeview_${windowSize.width}_${windowSize.height}`} style={{ width: vw / 2 - 896, flexShrink: 0, backgroundColor: modeColors.background }} />
+         <div key={`windowsize_noticeview_${windowSize.width}_${windowSize.height}`} style={{ width: vw / 2 - (1440/2), flexShrink: 0, backgroundColor: modeColors.background }} />
          <Stack tokens={{ childrenGap: 0 }} styles={{ root: { backgroundColor: modeColors.background, width: "100%" } }}>
-            <Stack style={{ maxWidth: 1778, paddingTop: 6, marginLeft: 4, height: 'calc(100vh - 8px)' }}>
+            <Stack style={{ maxWidth: 1440, paddingTop: 6, marginLeft: 4, height: 'calc(100vh - 8px)' }}>
                <Stack horizontal className={hordeClasses.raised}>
                   <Stack style={{ width: "100%", height: 'calc(100vh - 228px)' }} tokens={{ childrenGap: 18 }}>
                      <NoticePanel />
diff --git a/Engine/Source/Programs/Horde/HordeDashboard/src/components/PerforceView.tsx b/Engine/Source/Programs/Horde/HordeDashboard/src/components/PerforceView.tsx
--- a/Engine/Source/Programs/Horde/HordeDashboard/src/components/PerforceView.tsx
+++ b/Engine/Source/Programs/Horde/HordeDashboard/src/components/PerforceView.tsx
@@ -119,9 +119,9 @@ export const PerforceServerView: React.FC = () => {
       <TopNav />
       <Breadcrumbs items={[{ text: 'Perforce' }]} />
       <Stack horizontal>
-         <div key={`windowsize_streamview_${windowSize.width}_${windowSize.height}`} style={{ width: vw / 2 - 896, flexShrink: 0, backgroundColor: modeColors.background }} />
+         <div key={`windowsize_streamview_${windowSize.width}_${windowSize.height}`} style={{ width: vw / 2 - (1440/2), flexShrink: 0, backgroundColor: modeColors.background }} />
          <Stack tokens={{ childrenGap: 0 }} styles={{ root: { backgroundColor: modeColors.background, width: "100%" } }}>
-            <Stack style={{ maxWidth: 1778, paddingTop: 6, marginLeft: 4, height: 'calc(100vh - 8px)' }}>
+            <Stack style={{ maxWidth: 1440, paddingTop: 6, marginLeft: 4, height: 'calc(100vh - 8px)' }}>
                <Stack horizontal className={hordeClasses.raised}>
                   <Stack style={{ width: "100%", height: 'calc(100vh - 228px)' }} tokens={{ childrenGap: 18 }}>
                      <ServerPanel />
diff --git a/Engine/Source/Programs/Horde/HordeDashboard/src/components/PoolView.tsx b/Engine/Source/Programs/Horde/HordeDashboard/src/components/PoolView.tsx
--- a/Engine/Source/Programs/Horde/HordeDashboard/src/components/PoolView.tsx
+++ b/Engine/Source/Programs/Horde/HordeDashboard/src/components/PoolView.tsx
@@ -834,6 +834,11 @@ const AutoScalerPanel: React.FC = () => {
                   <ComboBox styles={{ root: { width: 180 } }} disabled={disabled} label="Strategy" selectedKey={state.strategy} options={stratItems} onChange={(ev, option, index, value) => {
                      setState({ ...state, modified: true, strategy: option!.data })
                   }} />
+                  <SpinButton styles={{ root: { width: 180 } }} label={`Conform Interval${state.conformInterval === 0 ? " (Disabled)" : " (Hours)"}`} min={0} max={24 * 7} value={state.conformInterval?.toString() ?? "24"} labelPosition={Position.top} onChange={(ev, value) => {
+                     setState({ ...state, modified: true, conformInterval: parseInt(value ?? "24") })
+                  }} />
+               </Stack>
+               <Stack horizontal tokens={{ childrenGap: 18 }}>
                   <SpinButton styles={{ root: { width: 128 } }} disabled={disabled} label="Minimum Agents" value={state.minAgents?.toString() ?? UNSET_VALUE} labelPosition={Position.top} onChange={(ev, value) => {
                      setState({ ...state, modified: true, minAgents: value === UNSET_VALUE ? undefined : parseInt(value ?? "0") })
                   }} />
@@ -841,15 +846,11 @@ const AutoScalerPanel: React.FC = () => {
                      setState({ ...state, modified: true, reserveAgents: value === UNSET_VALUE ? undefined : parseInt(value ?? "0") })
                   }} />
                </Stack>
-               <Stack>
-                  <SpinButton styles={{ root: { width: 180 } }} label={`Conform Interval${state.conformInterval === 0 ? " (Disabled)" : " (Hours)"}`} min={0} max={24 * 7} value={state.conformInterval?.toString() ?? "24"} labelPosition={Position.top} onChange={(ev, value) => {
-                     setState({ ...state, modified: true, conformInterval: parseInt(value ?? "24") })
-                  }} />
-               </Stack>
             </Stack>
          </Stack>
       </Stack>
    </Stack>
+   
 }
 
 
@@ -905,7 +906,7 @@ const PoolPanel: React.FC = () => {
    }
 
    const summaryColumns = [
-      { key: 'column1', name: 'Name', minWidth: 100, maxWidth: 100 },
+      { key: 'column1', name: 'Name', minWidth: 64, maxWidth: 64 },
       { key: 'column2', name: 'Value', minWidth: 100, maxWidth: 200 },
    ];
 
@@ -953,7 +954,7 @@ const PoolPanel: React.FC = () => {
    return (<Stack>
       <Stack styles={{ root: { paddingTop: 18, paddingLeft: 12, paddingRight: 12, width: "100%" } }} >
          <Stack tokens={{ childrenGap: 12 }}>
-            <Stack horizontal tokens={{ childrenGap: 118 }}>
+            <Stack horizontal tokens={{ childrenGap: 48 }}>
                <Stack style={{ minWidth: 224 }}>
                   <Stack style={{ paddingBottom: 8 }}>
                      <PrimaryButton text={pool.name} href={`/agents?agent=${encodeURI(pool.id)}&exact=true`} target="_blank" style={{ color: "#FFFFFF", backgroundColor: linearInterpolate(color), border: "unset", flexShrink: 1 }} />
@@ -1071,14 +1072,14 @@ export const PoolView: React.FC = observer(() => {
       <TopNav />
       <Breadcrumbs items={[{ text: "Agents", link: "/agents" }, { text: breadText }]} />
       <Stack horizontal>
-         <div key={`windowsize_poolview_${windowSize.width}_${windowSize.height}`} style={{ width: vw / 2 - 896, flexShrink: 0, backgroundColor: modeColors.background }} />
+         <div key={`windowsize_poolview_${windowSize.width}_${windowSize.height}`} style={{ width: vw / 2 - (1440 / 2), flexShrink: 0, backgroundColor: modeColors.background }} />
          <Stack tokens={{ childrenGap: 0 }} styles={{ root: { backgroundColor: modeColors.background, width: "100%" } }}>
-            <Stack style={{ maxWidth: 1778, paddingTop: 6, marginLeft: 4, height: 'calc(100vh - 8px)' }}>
+            <Stack style={{ maxWidth: 1440, paddingTop: 6, marginLeft: 4, height: 'calc(100vh - 8px)' }}>
                <Stack horizontal style={{ paddingTop: 8, paddingBottom: 16 }}>
                   <Stack grow />
-                  <Stack style={{width: 320}}>
+                  <Stack style={{ width: 320 }}>
                      <TagPicker inputProps={{ placeholder: "Select a pool" }}
-                        defaultSelectedItems={ defaultItems}
+                        defaultSelectedItems={defaultItems}
                         onResolveSuggestions={filterSuggestedTags}
                         getTextFromItem={getTextFromItem}
                         onEmptyResolveSuggestions={(selected) => {
@@ -1101,7 +1102,7 @@ export const PoolView: React.FC = observer(() => {
                   </Stack>
                </Stack>
                {!!pool && <Stack className={hordeClasses.raised}>
-                  <Stack style={{ position: "relative", width: "100%", height: 'calc(100vh - 228px)' }}>
+                  <Stack style={{ position: "relative", width: "1390px", height: 'calc(100vh - 228px)' }}>
                      <ScrollablePane scrollbarVisibility={ScrollbarVisibility.always}>
                         <Stack tokens={{ childrenGap: 18 }}>
                            <PoolPanel />
diff --git a/Engine/Source/Programs/Horde/HordeDashboard/src/components/PreviewChanges.tsx b/Engine/Source/Programs/Horde/HordeDashboard/src/components/PreviewChanges.tsx
--- a/Engine/Source/Programs/Horde/HordeDashboard/src/components/PreviewChanges.tsx
+++ b/Engine/Source/Programs/Horde/HordeDashboard/src/components/PreviewChanges.tsx
@@ -15,21 +15,12 @@ const changeItems: ChangeItem[] = [];
 
 let item = {
    id: "1",
-   summary: "Improved horizontal space on job detail views to better support various desktop resolutions and scaling",
+   summary: "Improved horizontal space to better support various desktop resolutions and scaling",
    link: "/job/63f7c145d436940d3e099365", 
    slack: "https://epicgames.slack.com/archives/C044JG73ARY/p1677529747047259"
 }
 changeItems.push(item);
 
-item = {
-   id: "2",
-   summary: "Timeline navigation for steps, agents, pools, and batches.  Click a span in the timeline to pin the tooltip, allowing navigation",
-   link: "/job/63f7c145d436940d3e099365",
-   slack: "https://epicgames.slack.com/archives/C044JG73ARY/p1677530063797479"
-}
-
-changeItems.push(item);
-
 const columns = [
    { key: 'column1', name: 'ID', minWidth: 100, maxWidth: 100 },
    { key: 'column2', name: 'Summary', minWidth: 540, maxWidth: 540 },
diff --git a/Engine/Source/Programs/Horde/HordeDashboard/src/components/ProjectHome.tsx b/Engine/Source/Programs/Horde/HordeDashboard/src/components/ProjectHome.tsx
--- a/Engine/Source/Programs/Horde/HordeDashboard/src/components/ProjectHome.tsx
+++ b/Engine/Source/Programs/Horde/HordeDashboard/src/components/ProjectHome.tsx
@@ -118,9 +118,9 @@ export const ProjectHome: React.FC = observer(() => {
             <TopNav />
             <Breadcrumbs items={crumbs} />
             <Stack horizontal>
-                <div key={`windowsize_streamview_${windowSize.width}_${windowSize.height}`} style={{ width: vw / 2 - 850, flexShrink: 0, backgroundColor: 'rgb(250, 249, 249)' }} />
+                <div key={`windowsize_streamview_${windowSize.width}_${windowSize.height}`} style={{ width: vw / 2 - (1440/2), flexShrink: 0, backgroundColor: 'rgb(250, 249, 249)' }} />
                 <Stack tokens={{ childrenGap: 0 }} styles={{ root: { backgroundColor: 'rgb(250, 249, 249)', width: "100%" } }}>
-                    <Stack style={{ maxWidth: 1800, paddingTop: 6, marginLeft: 4 }}>
+                    <Stack style={{ maxWidth: 1440, paddingTop: 6, marginLeft: 4 }}>
                         <FocusZone direction={FocusZoneDirection.vertical} style={{ padding: 0 }}>
                             <div className={detailClasses.container} style={{ width: "100%", height: 'calc(100vh - 208px)', position: 'relative' }} data-is-scrollable={true}>
                                 <ScrollablePane scrollbarVisibility={ScrollbarVisibility.auto} onScroll={() => { }}>
diff --git a/Engine/Source/Programs/Horde/HordeDashboard/src/components/StreamSummary.tsx b/Engine/Source/Programs/Horde/HordeDashboard/src/components/StreamSummary.tsx
--- a/Engine/Source/Programs/Horde/HordeDashboard/src/components/StreamSummary.tsx
+++ b/Engine/Source/Programs/Horde/HordeDashboard/src/components/StreamSummary.tsx
@@ -436,18 +436,22 @@ const HealthPanelIssues: React.FC<{ desktopAlerts?: boolean }> = observer(({ des
    const [items, groups] = sortIssueGroups(issues.map(i => { return { issue: i } }), projectStore);
 
    // 508
-   const columns = [
-      { key: 'health_column1', name: 'Summary', minWidth: 690, maxWidth: 690, isResizable: false },
+   const columns:IColumn[] = [
+      { key: 'health_column1', name: 'Summary', minWidth: 370, maxWidth: 370, isResizable: false }, // see summary render for px setting needed by ellipsis
       { key: 'health_column2', name: 'Quarantine', minWidth: 80, maxWidth: 80, isResizable: false },
       { key: 'health_column3', name: 'Jira', minWidth: 80, maxWidth: 80, isResizable: false },
       { key: 'health_column4', name: 'JiraPriority', minWidth: 64, maxWidth: 64, isResizable: false },
-      { key: 'health_column5', name: 'JiraAssignee', minWidth: 240, maxWidth: 240, isResizable: false },
+      { key: 'health_column5', name: 'JiraAssignee', minWidth: 200, maxWidth: 200, isResizable: false },
       { key: 'health_column6', name: 'JiraStatus', minWidth: 64, maxWidth: 64, isResizable: false },
       { key: 'health_column7', name: 'Status', minWidth: 160, maxWidth: 160, isResizable: false },
-      { key: 'health_column8', name: 'Opened', minWidth: 150, maxWidth: 150, isResizable: false },
+      { key: 'health_column8', name: 'Opened', minWidth: 140, maxWidth: 140, isResizable: false },
 
    ];
 
+   columns.forEach(c => {
+      c.isPadded = false;
+   })
+
    const renderItem = (item: HealthItem, index?: number, column?: IColumn) => {
 
       if (!column) {
@@ -477,11 +481,8 @@ const HealthPanelIssues: React.FC<{ desktopAlerts?: boolean }> = observer(({ des
       if (column.name === "Summary") {
 
          let summary = issue.summary;
-         if (summary.length > 100) {
-            summary = summary.slice(0, 100) + "...";
-         }
 
-         return <Stack horizontal disableShrink={true}>{<IssueStatusIconV2 issue={issue} />}<Text variant={textSize}>{`Issue ${issue.id} - ${summary}`}</Text></Stack>;
+         return <Stack horizontal>{<IssueStatusIconV2 issue={issue} />}<Text style={{width: "370px", display: "block", textAlign: "left", overflow: "hidden", "textOverflow" : "ellipsis"}} variant={textSize}>{`${issue.id} - ${summary}`}</Text></Stack>;
       }
 
       if (column.name === "Quarantine") {
@@ -573,7 +574,7 @@ const HealthPanelIssues: React.FC<{ desktopAlerts?: boolean }> = observer(({ des
       detailsRow: {
          selectors: {
             '.ms-DetailsRow-cell': {
-               padding: 8,
+               padding: 4,
                overflow: "hidden",
                whiteSpace: "nowrap"
             }
@@ -651,7 +652,7 @@ const HealthPanelIssues: React.FC<{ desktopAlerts?: boolean }> = observer(({ des
       pivotItems.push(<PivotItem headerText={`${config.summaryTab ?? config.id} (${workflowIssues.length})`} itemKey={`issue_pivot_item_key_${config.id}`} key={`issue_pivot_key_${config.id}`} headerButtonProps={{ onClick: () => { setCurrentPivot(config.id) } }} />);
    });
 
-   return <Stack style={{ width: 1720, marginLeft: 4 }}>
+   return <Stack style={{ width: 1324, marginLeft: 4 }}>
       <IssueModalV2 issueId={query.get("issue")} streamId={handler.streamId!} popHistoryOnClose={issueHistory} />
       <Stack styles={{ root: { paddingLeft: 4, paddingRight: 0, paddingTop: 4, paddingBottom: 4 } }}>
          <Stack>
@@ -711,7 +712,7 @@ const HealthPanel: React.FC<{ desktopAlerts?: boolean }> = observer(({ desktopAl
 
    if (handler.update) { }
 
-   return <Stack style={{ width: 1760, marginLeft: 4 }}>
+   return <Stack style={{ width: 1384, marginLeft: 4 }}>
       <Stack className={hordeClasses.raised}>
          <Stack tokens={{ childrenGap: 12 }}>
             <Stack>
@@ -741,7 +742,7 @@ const SchedulePanel: React.FC = observer(() => {
       return null;
    }
 
-   return (<Stack style={{ width: 1760, marginLeft: 4 }}>
+   return (<Stack style={{ width: 1384, marginLeft: 4 }}>
       <Stack className={hordeClasses.raised}>
          <Stack tokens={{ childrenGap: 12 }}>
             <Text variant="mediumPlus" styles={{ root: { fontFamily: "Horde Open Sans SemiBold" } }}>Schedule</Text>
@@ -765,7 +766,7 @@ const StreamSummaryInner: React.FC = observer(() => {
    // subscribe
    if (handler.update) { }
 
-   return <Stack tokens={{ childrenGap: 0 }} styles={{ root: { width: 1800, backgroundColor: "#fffffff", margin: 0, paddingTop: 8 } }}>
+   return <Stack tokens={{ childrenGap: 0 }} styles={{ root: { width: 1440, backgroundColor: "#fffffff", margin: 0, paddingTop: 8 } }}>
       <Stack style={{ padding: 0 }} className={detailClasses.detailsRow}>
          <FocusZone direction={FocusZoneDirection.vertical} style={{ padding: 0 }}>
             <div className={detailClasses.container} style={{ width: "100%", height: 'calc(100vh - 278px)', position: 'relative' }} data-is-scrollable={true}>
diff --git a/Engine/Source/Programs/Horde/HordeDashboard/src/components/StreamView.tsx b/Engine/Source/Programs/Horde/HordeDashboard/src/components/StreamView.tsx
--- a/Engine/Source/Programs/Horde/HordeDashboard/src/components/StreamView.tsx
+++ b/Engine/Source/Programs/Horde/HordeDashboard/src/components/StreamView.tsx
@@ -228,9 +228,9 @@ const StreamViewInner: React.FC = observer(() => {
          }} />
          {findJobsShown && <JobSearchSimpleModal onClose={() => { setFindJobsShown(false) }} streamId={stream.id} />}
          <Stack horizontal>
-            <div key={`windowsize_streamview_${windowSize.width}_${windowSize.height}`} style={{ width: vw / 2 - 900, flexShrink: 0, backgroundColor: modeColors.background }} />
+            <div key={`windowsize_streamview_${windowSize.width}_${windowSize.height}`} style={{ width: (vw / 2 - (1440/2)) - 12, flexShrink: 0, backgroundColor: modeColors.background }} />
             <Stack tokens={{ childrenGap: 0 }} styles={{ root: { backgroundColor: modeColors.background, width: "100%" } }}>
-               <Stack style={{ width: 1800, paddingTop: 12, marginLeft: 4 }}>
+               <Stack style={{ width: 1440, paddingTop: 12, marginLeft: 4 }}>
                   <Stack style={{ maxWidth: windowWidth - 12 }} horizontal verticalAlign='center' verticalFill={true}>
                      <Stack style={{ paddingLeft: 4, paddingTop: 2 }}>
                         <Pivot className={hordeClasses.pivot}
diff --git a/Engine/Source/Programs/Horde/HordeDashboard/src/components/ToolView.tsx b/Engine/Source/Programs/Horde/HordeDashboard/src/components/ToolView.tsx
new file mode 100644
--- /dev/null
+++ b/Engine/Source/Programs/Horde/HordeDashboard/src/components/ToolView.tsx
@@ -0,0 +1,150 @@
+// Copyright Epic Games, Inc. All Rights Reserved.
+
+import { DetailsList, DetailsListLayoutMode, IColumn, PrimaryButton, SelectionMode, Stack, Text } from "@fluentui/react";
+import { observer } from "mobx-react-lite";
+import React, { useEffect } from "react";
+import backend from "../backend";
+import { GetToolSummaryResponse } from "../backend/Api";
+import { PollBase } from "../backend/PollBase";
+import { useWindowSize } from "../base/utilities/hooks";
+import { detailClasses, hordeClasses, modeColors } from "../styles/Styles";
+import { Breadcrumbs } from "./Breadcrumbs";
+import { TopNav } from "./TopNav";
+
+class ToolHandler extends PollBase {
+
+   constructor(pollTime = 30000) {
+
+      super(pollTime);
+
+   }
+
+   clear() {
+      super.stop();
+   }
+
+   async poll(): Promise<void> {
+
+      try {
+
+         this.tools = await backend.getTools();         
+         this.setUpdated();
+
+      } catch (err) {
+
+      }
+
+   }
+
+   tools: GetToolSummaryResponse[] = [];
+}
+
+const handler = new ToolHandler();
+
+const ToolPanel: React.FC = observer(() => {
+
+   useEffect(() => {
+
+      handler.start();
+
+      return () => {
+         handler.clear();
+      };
+
+   }, []);
+
+   // subscribe
+   if (handler.updated) { };
+
+   const columns:IColumn[] = [      
+      { key: 'column_name', name: 'Name', minWidth: 240, maxWidth: 240, isResizable: false },
+      { key: 'column_desc', name: 'Description', fieldName: 'description', minWidth: 580, maxWidth: 580, isResizable: false, isMultiline: true },
+      { key: 'column_version', name: 'Version', fieldName: 'version', minWidth: 280, maxWidth: 280, isResizable: false, headerClassName: detailClasses.detailsHeader },      
+      { key: 'column_download', name: '', minWidth: 160, maxWidth: 160, isResizable: false }
+   ];
+
+   let tools = [...handler.tools];
+
+   tools = tools.sort((a, b) => a.name.localeCompare(b.name));
+
+   const renderItem = (item: any, index?: number, column?: IColumn) => {
+
+      if (!column) {
+         return null;
+      }
+
+      if (column.key === 'column_name') {
+         return <Stack verticalAlign="center" verticalFill={true}>
+            <Text style={{ fontFamily: "Horde Open Sans SemiBold", color: modeColors.text }}>{item.name}</Text>
+         </Stack>
+      }
+
+      if (column.key === 'column_version') {
+         if (!item.version) {
+            return null;
+         }
+         return <Stack horizontalAlign="center" verticalAlign="center" verticalFill={true}>
+            <Text style={{ color: modeColors.text }}>{item.version}</Text>
+         </Stack>
+      }
+
+
+      if (column.key === 'column_download') {
+         return <Stack horizontalAlign="center" verticalAlign="center" verticalFill={true}>
+            <PrimaryButton style={{width: 120, color: "#FFFFFF"}} text="Download" href={`/api/v1/tools/${item.id}?action=download`} />
+         </Stack>
+      }
+
+      if (!column?.fieldName) {
+         return null;
+      }
+      return <Stack verticalAlign="center" verticalFill={true}>
+         <Text style={{ color: modeColors.text }}>{item[column?.fieldName]}</Text>
+      </Stack>
+   };
+
+   return (<Stack>
+      <Stack styles={{ root: { paddingTop: 0, paddingLeft: 12, paddingRight: 12, width: "100%" } }} >
+         <Stack tokens={{ childrenGap: 12 }}>
+            <div style={{ overflowY: 'auto', overflowX: 'hidden', maxHeight: "calc(100vh - 312px)" }} data-is-scrollable={true}>
+               <Stack>
+                  <DetailsList
+                     isHeaderVisible={true}
+                     items={tools}
+                     columns={columns}
+                     selectionMode={SelectionMode.none}
+                     layoutMode={DetailsListLayoutMode.justified}
+                     compact={false}
+                     onRenderItemColumn={renderItem}
+                  />
+               </Stack>
+            </div>
+         </Stack>
+      </Stack>
+   </Stack>);
+});
+
+
+export const ToolView: React.FC = () => {
+
+   const windowSize = useWindowSize();
+   const vw = Math.max(document.documentElement.clientWidth, window.innerWidth || 0);
+
+   return <Stack className={hordeClasses.horde}>
+      <TopNav />
+      <Breadcrumbs items={[{ text: 'Tools' }]} />
+      <Stack horizontal>
+         <div key={`windowsize_streamview_${windowSize.width}_${windowSize.height}`} style={{ width: vw / 2 - (1440/2), flexShrink: 0, backgroundColor: modeColors.background }} />
+         <Stack tokens={{ childrenGap: 0 }} styles={{ root: { backgroundColor: modeColors.background, width: "100%" } }}>
+            <Stack style={{ maxWidth: 1440, paddingTop: 6, marginLeft: 4, height: 'calc(100vh - 8px)' }}>
+               <Stack horizontal className={hordeClasses.raised}>
+                  <Stack style={{ width: "100%", height: 'calc(100vh - 228px)' }} tokens={{ childrenGap: 18 }}>
+                     <ToolPanel />
+                  </Stack>
+               </Stack>
+            </Stack>
+         </Stack>
+      </Stack>
+   </Stack>
+};
+
diff --git a/Engine/Source/Programs/Horde/HordeDashboard/src/components/TopNav.tsx b/Engine/Source/Programs/Horde/HordeDashboard/src/components/TopNav.tsx
--- a/Engine/Source/Programs/Horde/HordeDashboard/src/components/TopNav.tsx
+++ b/Engine/Source/Programs/Horde/HordeDashboard/src/components/TopNav.tsx
@@ -6,6 +6,7 @@ import React, { useRef, useState } from 'react';
 import { Link, useNavigate } from 'react-router-dom';
 import backend, { useBackend } from '../backend';
 import { ProjectData } from "../backend/Api";
+import { getSiteConfig } from '../backend/Config';
 import dashboard from '../backend/Dashboard';
 import { ProjectStore } from '../backend/ProjectStore';
 import { modeColors } from '../styles/Styles';
@@ -376,6 +377,14 @@ export const TopNav: React.FC<{ suppressServer?: boolean }> = observer(({ suppre
          link: `/perforce/servers`
       });
 
+      if (getSiteConfig().environment !== "production") {
+         resourceItems.push({
+            key: "admin_tools",
+            text: "Tools",
+            link: `/tools`
+         });   
+      }
+
       subItems.push({
          itemType: ContextualMenuItemType.Section,
          key: `admin_resources`,
@@ -447,7 +456,7 @@ export const TopNav: React.FC<{ suppressServer?: boolean }> = observer(({ suppre
 
       subItems.push({
          itemType: ContextualMenuItemType.Section,
-         key: `admin_config`,
+         key: `admin_version`,
          sectionProps: {
             title: "Version",
             items: versionItems,
@@ -541,7 +550,7 @@ export const TopNav: React.FC<{ suppressServer?: boolean }> = observer(({ suppre
       <div style={{ backgroundColor: modeColors.header }}>
          {showVersion && <VersionModal show={true} onClose={() => { setShowVersion(false) }} />}
          {showPreviewChanges && <PreviewChangesModal onClose={() => { setShowPreviewChanges(false) }} />}
-         <Stack tokens={{ maxWidth: 1800, childrenGap: 0 }} disableShrink={true} styles={{ root: { backgroundColor: modeColors.header, margin: "auto", width: "100%" } }}>
+         <Stack tokens={{ maxWidth: 1464, childrenGap: 0 }} disableShrink={true} styles={{ root: { backgroundColor: modeColors.header, margin: "auto", width: "100%" } }}>
             <Stack horizontal verticalAlign='center' styles={{ root: { height: "60px" } }} >
 
                <Link to="/index"><Stack horizontal styles={{ root: { paddingLeft: 8, cursor: 'pointer' } }}>
diff --git a/Engine/Source/Programs/Horde/HordeDashboard/src/components/UserHome.tsx b/Engine/Source/Programs/Horde/HordeDashboard/src/components/UserHome.tsx
--- a/Engine/Source/Programs/Horde/HordeDashboard/src/components/UserHome.tsx
+++ b/Engine/Source/Programs/Horde/HordeDashboard/src/components/UserHome.tsx
@@ -65,7 +65,7 @@ const customStyles = mergeStyleSets({
 
 });
 
-const homeWidth = 1750;
+const homeWidth = 1400;
 
 /// Projects
 
@@ -224,8 +224,8 @@ const HealthPanel: React.FC = observer(() => {
    const items = issues?.map(i => { return { issue: i } });
 
    const columns = [
-      { key: 'health_column1', name: 'Summary', minWidth: 600, maxWidth: 600, isResizable: false },
-      { key: 'health_column3', name: 'Status', minWidth: 420, maxWidth: 420, isResizable: false },
+      { key: 'health_column1', name: 'Summary', minWidth: 500, maxWidth: 500, isResizable: false },
+      { key: 'health_column3', name: 'Status', minWidth: 320, maxWidth: 320, isResizable: false },
       { key: 'health_column4', name: 'Opened', minWidth: 100, maxWidth: 100, isResizable: false },
 
    ];
@@ -584,10 +584,10 @@ const JobsPanel: React.FC<{ includeOtherPreflights: boolean }> = observer(({ inc
 
    let columns: IColumn[] = [
       { key: 'jobview_column1', name: 'Status', minWidth: 16, maxWidth: 16 },
-      { key: 'jobview_column2', name: 'Change', minWidth: 220, maxWidth: 220 },
-      { key: 'jobview_column3', name: 'Submit', minWidth: 140, maxWidth: 140 },
-      { key: 'jobview_column4', name: 'Labels', minWidth: 420, maxWidth: 420 },
-      { key: 'jobview_column5', name: 'Steps', minWidth: 380, maxWidth: 380 },
+      { key: 'jobview_column2', name: 'Change', minWidth: 114, maxWidth: 114 },
+      { key: 'jobview_column3', name: 'Submit', minWidth: 130, maxWidth: 130 },
+      { key: 'jobview_column4', name: 'Labels', minWidth: 320, maxWidth: 320 },
+      { key: 'jobview_column5', name: 'Steps', minWidth: 260, maxWidth: 260 },
       { key: 'jobview_column6', name: 'StartedBy', minWidth: 140, maxWidth: 140 },
       { key: 'jobview_column7', name: 'Time', minWidth: 130, maxWidth: 130 },
       { key: 'jobview_column8', name: 'Dismiss', minWidth: 32, maxWidth: 32 },
@@ -1030,7 +1030,6 @@ const JobsPanel: React.FC<{ includeOtherPreflights: boolean }> = observer(({ inc
                   </Stack>
                </Stack>
                <Stack styles={{ root: { paddingLeft: 4, paddingRight: 0, paddingTop: 8, paddingBottom: 4 } }}>
-                  <div style={{ overflowY: 'auto', overflowX: 'hidden', maxHeight: "800px" }} data-is-scrollable={true}>
                      {!!jobItems.length && <DetailsList
                         styles={{ root: { paddingLeft: 8, paddingRight: 8, marginBottom: 18 } }}
                         compact={true}
@@ -1047,7 +1046,6 @@ const JobsPanel: React.FC<{ includeOtherPreflights: boolean }> = observer(({ inc
                         onShouldVirtualize={() => true}
                      />}
                      {!dashboard.pinnedJobsIds.length && <Stack style={{ paddingBottom: 12 }}><Text variant="medium">No jobs are pinned</Text></Stack>}
-                  </div>
                </Stack>
             </Stack>
          </Stack>
@@ -1097,9 +1095,9 @@ export const UserHomeView: React.FC = () => {
          <TopNav />
          <Breadcrumbs items={[{ text: 'Home' }]} />
          <Stack horizontal>
-            <div key={`windowsize_streamview_${windowSize.width}_${windowSize.height}`} style={{ width: vw / 2 - 900, flexShrink: 0, backgroundColor: modeColors.background }} />
+            <div key={`windowsize_streamview_${windowSize.width}_${windowSize.height}`} style={{ width: vw / 2 - (1440/2), flexShrink: 0, backgroundColor: modeColors.background }} />
             <Stack tokens={{ childrenGap: 0 }} styles={{ root: { backgroundColor: modeColors.background, width: "100%" } }}>
-               <Stack style={{ maxWidth: 1800, paddingTop: 6, marginLeft: 4 }}>
+               <Stack style={{ maxWidth: 1440, paddingTop: 6, marginLeft: 4 }}>
                   <UserHomeViewInner />
                </Stack>
             </Stack>
diff --git a/Engine/Source/Programs/Horde/HordeDashboard/src/components/UtilizationReportView.tsx b/Engine/Source/Programs/Horde/HordeDashboard/src/components/UtilizationReportView.tsx
--- a/Engine/Source/Programs/Horde/HordeDashboard/src/components/UtilizationReportView.tsx
+++ b/Engine/Source/Programs/Horde/HordeDashboard/src/components/UtilizationReportView.tsx
@@ -708,9 +708,9 @@ export const UtilizationReportView: React.FC = observer(() => {
 			<TopNav />
 			<Breadcrumbs items={[{ text: 'Utilization Report' }]} />
 			<Stack horizontal>
-				<div key={`windowsize_streamview_${windowSize.width}_${windowSize.height}`} style={{ width: vw / 2 - 900, flexShrink: 0, backgroundColor: modeColors.background }} />
+				<div key={`windowsize_streamview_${windowSize.width}_${windowSize.height}`} style={{ width: vw / 2 - (1440/2), flexShrink: 0, backgroundColor: modeColors.background }} />
 				<Stack tokens={{ childrenGap: 0 }} styles={{ root: { backgroundColor: modeColors.background, width: "100%" } }}>
-					<Stack style={{ width: 1800, paddingTop: 6, marginLeft: 4, height: 'calc(100vh - 8px)' }}>
+					<Stack style={{ width: 1440, paddingTop: 6, marginLeft: 4, height: 'calc(100vh - 8px)' }}>
 						<Stack styles={{ root: { paddingTop: 18, paddingLeft: 12, paddingRight: 12, width: "100%" } }} >
 							<Stack className={hordeClasses.raised}>
 								<div style={{ overflowY: 'auto', overflowX: 'hidden', maxHeight: "calc(100vh - 250px)" }} data-is-scrollable={true}>
diff --git a/Engine/Source/Programs/Horde/HordeDashboard/src/components/jobDetailsV2/JobDetailHealthV2.tsx b/Engine/Source/Programs/Horde/HordeDashboard/src/components/jobDetailsV2/JobDetailHealthV2.tsx
--- a/Engine/Source/Programs/Horde/HordeDashboard/src/components/jobDetailsV2/JobDetailHealthV2.tsx
+++ b/Engine/Source/Programs/Horde/HordeDashboard/src/components/jobDetailsV2/JobDetailHealthV2.tsx
@@ -201,13 +201,13 @@ export const HealthPanel: React.FC<{ jobDetails: JobDetailsV2 }> = observer(({ j
    });
 
    const columns = [
-      { key: 'health_column1', name: 'Summary', minWidth: 500, maxWidth: 500, isResizable: false },
+      { key: 'health_column1', name: 'Summary', minWidth: 460, maxWidth: 460, isResizable: false },
       { key: 'health_column2', name: 'Quarantine', minWidth: 80, maxWidth: 80, isResizable: false },
       { key: 'health_column3', name: 'Jira', minWidth: 80, maxWidth: 80, isResizable: false },
       { key: 'health_column4', name: 'JiraPriority', minWidth: 64, maxWidth: 64, isResizable: false },
-      { key: 'health_column5', name: 'JiraAssignee', minWidth: 240, maxWidth: 240, isResizable: false },
+      { key: 'health_column5', name: 'JiraAssignee', minWidth: 180, maxWidth: 180, isResizable: false },
       { key: 'health_column6', name: 'JiraStatus', minWidth: 64, maxWidth: 64, isResizable: false },
-      { key: 'health_column7', name: 'Status', minWidth: 160, maxWidth: 160, isResizable: false },
+      { key: 'health_column7', name: 'Status', minWidth: 140, maxWidth: 140, isResizable: false },
       { key: 'health_column8', name: 'Opened', minWidth: 140, maxWidth: 150, isResizable: false },
    ];
 
diff --git a/Engine/Source/Programs/Horde/HordeDashboard/src/components/jobDetailsV2/JobDetailStepTrends.tsx b/Engine/Source/Programs/Horde/HordeDashboard/src/components/jobDetailsV2/JobDetailStepTrends.tsx
--- a/Engine/Source/Programs/Horde/HordeDashboard/src/components/jobDetailsV2/JobDetailStepTrends.tsx
+++ b/Engine/Source/Programs/Horde/HordeDashboard/src/components/jobDetailsV2/JobDetailStepTrends.tsx
@@ -881,7 +881,7 @@ const StepTrendsPanelInner: React.FC<{ jobDetails: JobDetailsV2; dataView: Trend
             <Stack styles={{ root: { position: 'relative', top: -10 } }}>
                {statusDiv}
             </Stack>
-            <Stack styles={{ root: { width: 1720, minHeight: 400, height: 400 } }} onWheel={onChartScroll} onMouseEnter={onChartWheelEnter} onMouseLeave={onChartWheelLeave}>
+            <Stack styles={{ root: { width: 1370, minHeight: 400, height: 400 } }} onWheel={onChartScroll} onMouseEnter={onChartWheelEnter} onMouseLeave={onChartWheelLeave}>
                {placeholderDiv}
             </Stack>
 
diff --git a/Engine/Source/Programs/Horde/HordeDashboard/src/components/jobDetailsV2/JobDetailViewSteps.tsx b/Engine/Source/Programs/Horde/HordeDashboard/src/components/jobDetailsV2/JobDetailViewSteps.tsx
--- a/Engine/Source/Programs/Horde/HordeDashboard/src/components/jobDetailsV2/JobDetailViewSteps.tsx
+++ b/Engine/Source/Programs/Horde/HordeDashboard/src/components/jobDetailsV2/JobDetailViewSteps.tsx
@@ -183,7 +183,7 @@ export const StepsPanelInner: React.FC<{ jobDetails: JobDetailsV2, depStepId?: s
    const buildColumns = (): IColumn[] => {
 
       const widths: Record<string, number> = {
-         "Name": 1074,
+         "Name": 732,
          "Progress": 120,
          "Start": 100,
          "End": 100,
diff --git a/Engine/Source/Programs/Horde/HordeDashboard/src/components/jobDetailsV2/JobDetailViewV2.tsx b/Engine/Source/Programs/Horde/HordeDashboard/src/components/jobDetailsV2/JobDetailViewV2.tsx
--- a/Engine/Source/Programs/Horde/HordeDashboard/src/components/jobDetailsV2/JobDetailViewV2.tsx
+++ b/Engine/Source/Programs/Horde/HordeDashboard/src/components/jobDetailsV2/JobDetailViewV2.tsx
@@ -314,7 +314,7 @@ const DetailsViewOverview: React.FC<{ jobDetails: JobDetailsV2 }> = ({ jobDetail
    </Stack>
 };
 
-const rootWidth = 1776;
+const rootWidth = 1440;
 
 const ScrollRestore: React.FC<{ jobDetails: JobDetailsV2, scrollRef: React.RefObject<HTMLDivElement> }> = ({ jobDetails, scrollRef }) => {
 
@@ -345,7 +345,8 @@ const DetailsView: React.FC<{ jobDetails: JobDetailsV2 }> = ({ jobDetails }) =>
 
    const details = jobDetails;
 
-   const centerAlign = vw / 2 - 890;
+
+   const centerAlign = vw / 2 - 720 /*890*/;
    const key = `windowsize_jobdetail_view_${windowSize.width}_${windowSize.height}`;
 
    return <Stack id="rail_job_overview">
diff --git a/Engine/Source/Programs/Horde/HordeDashboard/src/styles/Styles.tsx b/Engine/Source/Programs/Horde/HordeDashboard/src/styles/Styles.tsx
--- a/Engine/Source/Programs/Horde/HordeDashboard/src/styles/Styles.tsx
+++ b/Engine/Source/Programs/Horde/HordeDashboard/src/styles/Styles.tsx
@@ -259,8 +259,8 @@ export const hordeClasses = mergeStyleSets({
             minWidth: 0,
             paddingTop: 6,
             paddingBottom: 6,
-            paddingLeft: 3,
-            paddingRight: 3,
+            paddingLeft: 1,
+            paddingRight: 1,
             height: 18
          },
          '.ms-Button:hover': {
@@ -268,7 +268,7 @@ export const hordeClasses = mergeStyleSets({
             filter: "brightness(1.15)"
          },
          '.ms-Button-label': {
-            fontSize: '10px',
+            fontSize: '9px',
             fontFamily: "Horde Open Sans SemiBold",
             whiteSpace: "nowrap"
          }
