Horde: Fix missing lease type, job and batch ids due to changes in details serialization
#rnx

diff --git a/Engine/Source/Programs/Horde/HordeDashboard/src/backend/Backend.ts b/Engine/Source/Programs/Horde/HordeDashboard/src/backend/Backend.ts
--- a/Engine/Source/Programs/Horde/HordeDashboard/src/backend/Backend.ts
+++ b/Engine/Source/Programs/Horde/HordeDashboard/src/backend/Backend.ts
@@ -173,9 +173,9 @@ export class Backend {
                         lease.finishTime = new Date(Date.parse(lease.finishTime as string));
                     }
                     if (lease.details) {
-                        lease.batchId = lease.details.BatchId;
-                        lease.jobId = lease.details.JobId;
-                        lease.type = lease.details.Type;
+                        lease.batchId = lease.details.batchId;
+                        lease.jobId = lease.details.jobId;
+                        lease.type = lease.details.type;
                     }
                 });
 
