Horde: Fix log view progress bar overflow
#rnx
#rnx

diff --git a/Engine/Source/Programs/Horde/HordeDashboard/src/components/LogView.tsx b/Engine/Source/Programs/Horde/HordeDashboard/src/components/LogView.tsx
--- a/Engine/Source/Programs/Horde/HordeDashboard/src/components/LogView.tsx
+++ b/Engine/Source/Programs/Horde/HordeDashboard/src/components/LogView.tsx
@@ -992,7 +992,7 @@ export const LogList: React.FC<{ logId: string }> = observer(({ logId }) => {
 
                      </Stack>
                      <Stack horizontal >
-                        {logSource.active && <ProgressIndicator percentComplete={percentComplete} barHeight={2} styles={{ root: { paddingLeft: 8, width: 400 } }} />}
+                        {logSource.active && <ProgressIndicator percentComplete={percentComplete} barHeight={2} styles={{ root: { paddingLeft: 12, width: 364 } }} />}
                      </Stack>
                   </Stack>
 
