Horde: Adding workflow thread to find issue responses
#rnx

diff --git a/Engine/Source/Programs/Horde/Horde.Build/Issues/IssueMessages.cs b/Engine/Source/Programs/Horde/Horde.Build/Issues/IssueMessages.cs
--- a/Engine/Source/Programs/Horde/Horde.Build/Issues/IssueMessages.cs
+++ b/Engine/Source/Programs/Horde/Horde.Build/Issues/IssueMessages.cs
@@ -767,6 +767,11 @@ namespace Horde.Build.Issues
 		/// </summary>
 		public DateTime? QuarantineTimeUtc { get; set; }
 
+		/// <summary>
+		/// The workflow thread url for this issue
+		/// </summary>
+		public Uri? WorkflowThreadUrl { get; set; }
+
 		/// <summary>
 		/// Workflows for which this issue is open
 		/// </summary>
@@ -817,6 +822,7 @@ namespace Horde.Build.Issues
 				QuarantinedBy = new GetThinUserInfoResponse(quarantinedBy);
 				QuarantineTimeUtc = issue.QuarantineTimeUtc;
 			}
+			WorkflowThreadUrl = issue.WorkflowThreadUrl;
 			OpenWorkflows = openWorkflows;
 		}
 	}
