Horde: Adding workflow thread to issue modal, stream summaries, and issue exports
#rnx
#skipci

diff --git a/Engine/Source/Programs/Horde/HordeDashboard/src/backend/Api.ts b/Engine/Source/Programs/Horde/HordeDashboard/src/backend/Api.ts
--- a/Engine/Source/Programs/Horde/HordeDashboard/src/backend/Api.ts
+++ b/Engine/Source/Programs/Horde/HordeDashboard/src/backend/Api.ts
@@ -2918,6 +2918,8 @@ export type GetIssueResponse = {
 	/** User info for who force closed the issue */
 	forceClosedByUserInfo?: GetThinUserInfoResponse;
 
+	/** Workflow thread url (Slack, etc) */
+	workflowThreadUrl?: string;
 }
 
 /**Request an issue to be updated */
@@ -4038,6 +4040,8 @@ export type FindIssueResponse = {
 	/** Workflows for which this issue is open */
 	openWorkflows: string[];
 
+	/** Workflow thread url (Slack, etc) */
+	workflowThreadUrl?: string;
 }
 
 export type GetAgentSoftwareChannelResponse = {
diff --git a/Engine/Source/Programs/Horde/HordeDashboard/src/components/IssueViewV2.tsx b/Engine/Source/Programs/Horde/HordeDashboard/src/components/IssueViewV2.tsx
--- a/Engine/Source/Programs/Horde/HordeDashboard/src/components/IssueViewV2.tsx
+++ b/Engine/Source/Programs/Horde/HordeDashboard/src/components/IssueViewV2.tsx
@@ -1167,8 +1167,8 @@ const IssueHeader: React.FC<{ items?: SummaryItem[] }> = ({ items }) => {
             <Label style={{ padding: 0, minWidth: 110 }}>{`${title}:`}</Label>
          </Stack>
          <Stack>
-            {!!link && <a href={link} target="_blank" rel="noreferrer">
-               <Text style={{ textDecoration: strike ? "line-through" : undefined }} variant="small">{text}</Text>
+            {!!link && <a href={link} target="_blank" rel="noreferrer" style={{ fontSize: "12px", textDecoration: strike ? "line-through" : undefined }}>
+               {text}
             </a>
             }
             {!link && <Text style={{ textDecoration: strike ? "line-through" : undefined }} variant="small">{text}</Text>}
@@ -1204,6 +1204,15 @@ const IssueSummaryPanel: React.FC = () => {
 
    const items: SummaryItem[] = [];
 
+   if (issue.workflowThreadUrl) {
+      items.push({
+         title: "Workflow",
+         text: "Slack Thread",
+         link: issue.workflowThreadUrl
+      })
+   }
+
+
    if (!!issue.quarantineTimeUtc && !!issue.quarantinedByUserInfo) {
       items.push({
          title: "Quarantined By",
diff --git a/Engine/Source/Programs/Horde/HordeDashboard/src/components/StreamSummary.tsx b/Engine/Source/Programs/Horde/HordeDashboard/src/components/StreamSummary.tsx
--- a/Engine/Source/Programs/Horde/HordeDashboard/src/components/StreamSummary.tsx
+++ b/Engine/Source/Programs/Horde/HordeDashboard/src/components/StreamSummary.tsx
@@ -177,6 +177,7 @@ class SummaryHandler {
         <td>Issue</td>
         <td>Summary</td>
         <td>Status</td>
+        <td>Workflow</td>
         <td>Open Since</td>
         <td>Promoted</td>
         <td>Jira</td>
@@ -214,18 +215,23 @@ class SummaryHandler {
 
             let jiraTag = "";
 
-            if (issue.externalIssueKey) {      
-               
+            if (issue.externalIssueKey) {                     
                const url = `${dashboard.externalIssueService?.url}/browse/${issue.externalIssueKey}`;      
                jiraTag = `<a href="${url}" target="_blank">${issue.externalIssueKey}</a>`;
             }
 
+            let workflowTag = "";
+            if (issue.workflowThreadUrl) {                     
+               workflowTag = `<a href="${issue.workflowThreadUrl}" target="_blank">Slack Thread</a>`;
+            }
 
             html += `<tr>
                 <td>${g.name}</td>
                 <td><a href="${href}" target="_blank">${issue.id}</a></td>
                 <td style="text-align:left"><a href="${href}" target="_blank">${summary}</a></td>
                 <td>${status}</td>
+                <td>${workflowTag}</td>
+
                 <td>${openSince}</td>
                 <td>${issue.promoted ? "Yes" : "No"}</td>
                 <td>${jiraTag}</td>
@@ -487,10 +493,14 @@ const HealthPanelIssues: React.FC<{ desktopAlerts?: boolean }> = observer(({ des
 
       if (column.name === "Quarantine") {
 
-         if (!issue.quarantinedBy) {
+         if (!issue.quarantinedBy && !issue.workflowThreadUrl) {
             return null;
          }
 
+         if (issue.workflowThreadUrl) {
+            return <a href={issue.workflowThreadUrl} target="_blank"  rel="noreferrer">Slack Thread</a>
+         }
+
          return <Stack horizontalAlign="start" disableShrink={true} verticalAlign="center" ><Text variant={textSize}>Quarantined</Text></Stack>;
       }
 
