Horde: Add slack threads to user home, tweaks to open slack and not propogate through to issue modal
#rnx
#skipci

diff --git a/Engine/Source/Programs/Horde/HordeDashboard/src/components/StreamSummary.tsx b/Engine/Source/Programs/Horde/HordeDashboard/src/components/StreamSummary.tsx
--- a/Engine/Source/Programs/Horde/HordeDashboard/src/components/StreamSummary.tsx
+++ b/Engine/Source/Programs/Horde/HordeDashboard/src/components/StreamSummary.tsx
@@ -443,14 +443,14 @@ const HealthPanelIssues: React.FC<{ desktopAlerts?: boolean }> = observer(({ des
 
    // 508
    const columns:IColumn[] = [
-      { key: 'health_column1', name: 'Summary', minWidth: 370, maxWidth: 370, isResizable: false }, // see summary render for px setting needed by ellipsis
-      { key: 'health_column2', name: 'Quarantine', minWidth: 80, maxWidth: 80, isResizable: false },
-      { key: 'health_column3', name: 'Jira', minWidth: 80, maxWidth: 80, isResizable: false },
-      { key: 'health_column4', name: 'JiraPriority', minWidth: 64, maxWidth: 64, isResizable: false },
-      { key: 'health_column5', name: 'JiraAssignee', minWidth: 200, maxWidth: 200, isResizable: false },
-      { key: 'health_column6', name: 'JiraStatus', minWidth: 64, maxWidth: 64, isResizable: false },
-      { key: 'health_column7', name: 'Status', minWidth: 160, maxWidth: 160, isResizable: false },
-      { key: 'health_column8', name: 'Opened', minWidth: 140, maxWidth: 140, isResizable: false },
+      { key: 'health_column1', name: 'Summary', minWidth: 370, maxWidth: 370, isResizable: false, isPadded: false }, // see summary render for px setting needed by ellipsis
+      { key: 'health_column2', name: 'Quarantine', minWidth: 80, maxWidth: 80, isResizable: false, isPadded: false },
+      { key: 'health_column3', name: 'Jira', minWidth: 80, maxWidth: 80, isResizable: false, isPadded: false },
+      { key: 'health_column4', name: 'JiraPriority', minWidth: 64, maxWidth: 64, isResizable: false, isPadded: false },
+      { key: 'health_column5', name: 'JiraAssignee', minWidth: 200, maxWidth: 200, isResizable: false, isPadded: false },
+      { key: 'health_column6', name: 'JiraStatus', minWidth: 64, maxWidth: 64, isResizable: false, isPadded: false },
+      { key: 'health_column7', name: 'Status', minWidth: 160, maxWidth: 160, isResizable: false, isPadded: false },
+      { key: 'health_column8', name: 'Opened', minWidth: 140, maxWidth: 140, isResizable: false, isPadded: false },
 
    ];
 
@@ -488,7 +488,7 @@ const HealthPanelIssues: React.FC<{ desktopAlerts?: boolean }> = observer(({ des
 
          let summary = issue.summary;
 
-         return <Stack horizontal>{<IssueStatusIconV2 issue={issue} />}<Text style={{width: "370px", display: "block", textAlign: "left", overflow: "hidden", "textOverflow" : "ellipsis"}} variant={textSize}>{`${issue.id} - ${summary}`}</Text></Stack>;
+         return <Stack horizontal style={{ paddingTop: 6, height: "100%"}}>{<IssueStatusIconV2 issue={issue} />}<Text style={{width: "370px", display: "block", textAlign: "left", overflow: "hidden", "textOverflow" : "ellipsis"}} variant={textSize}>{`${issue.id} - ${summary}`}</Text></Stack>;
       }
 
       if (column.name === "Quarantine") {
@@ -498,10 +498,10 @@ const HealthPanelIssues: React.FC<{ desktopAlerts?: boolean }> = observer(({ des
          }
 
          if (issue.workflowThreadUrl) {
-            return <a href={issue.workflowThreadUrl} target="_blank"  rel="noreferrer">Slack Thread</a>
+            return <a onClick={(e) => e.stopPropagation()}  href={issue.workflowThreadUrl} target="_blank"  rel="noreferrer"><Stack style={{ paddingTop: 6, height: "100%"}} horizontalAlign="start" disableShrink={true}  >Slack Thread</Stack></a>
          }
 
-         return <Stack horizontalAlign="start" disableShrink={true} verticalAlign="center" ><Text variant={textSize}>Quarantined</Text></Stack>;
+         return <Stack style={{ paddingTop: 6, height: "100%"}} horizontalAlign="start" disableShrink={true} ><Text variant={textSize}>Quarantined</Text></Stack>;
       }
 
 
@@ -511,7 +511,7 @@ const HealthPanelIssues: React.FC<{ desktopAlerts?: boolean }> = observer(({ des
             status += ` (${ack})`;
          }
 
-         return <Stack horizontalAlign="end" disableShrink={true} verticalAlign="center"><Text variant={textSize}>{status}</Text></Stack>;
+         return <Stack style={{ paddingTop: 6, height: "100%"}} horizontalAlign="end" disableShrink={true} ><Text variant={textSize}>{status}</Text></Stack>;
       }
 
       if (issue.externalIssueKey) {
@@ -531,7 +531,7 @@ const HealthPanelIssues: React.FC<{ desktopAlerts?: boolean }> = observer(({ des
                priority = priority.slice(priority.indexOf("- ") + 2);
             }
 
-            return <Stack horizontalAlign="start" disableShrink={true}><a href={jissue.link} target="_blank" rel="noreferrer" onClick={(ev) => ev.stopPropagation()}><Text style={{ textDecoration: textDecoration }} variant={textSize}>{priority}</Text> </a></Stack>;
+            return <a href={jissue.link} target="_blank" rel="noreferrer" onClick={(ev) => ev.stopPropagation()}><Stack style={{ paddingTop: 6, height: "100%"}} horizontalAlign="start" disableShrink={true}><Text style={{ textDecoration: textDecoration }} variant={textSize}>{priority}</Text></Stack></a>;
 
          }
 
@@ -542,7 +542,7 @@ const HealthPanelIssues: React.FC<{ desktopAlerts?: boolean }> = observer(({ des
                return null;
             }
 
-            return <Stack horizontalAlign="center" disableShrink={true}><a href={jissue.link} target="_blank" rel="noreferrer" onClick={(ev) => ev.stopPropagation()}><Text style={{ textDecoration: textDecoration }} variant={textSize}>{jissue.assigneeDisplayName ?? "Unassigned (Jira)"}</Text> </a></Stack>;
+            return <a href={jissue.link} target="_blank" rel="noreferrer" onClick={(ev) => ev.stopPropagation()}><Stack style={{ paddingTop: 6, height: "100%"}} horizontalAlign="center" disableShrink={true}><Text style={{ textDecoration: textDecoration }} variant={textSize}>{jissue.assigneeDisplayName ?? "Unassigned (Jira)"}</Text></Stack></a>
          }
 
          if (column.name === "JiraStatus") {
@@ -553,7 +553,7 @@ const HealthPanelIssues: React.FC<{ desktopAlerts?: boolean }> = observer(({ des
                return null;
             }
 
-            return <Stack horizontalAlign="start" disableShrink={true}><a href={jissue.link} target="_blank" rel="noreferrer" onClick={(ev) => ev.stopPropagation()}><Text style={{ textDecoration: textDecoration }} variant={textSize}>{desc}</Text> </a></Stack>;
+            return <a href={jissue.link} target="_blank" rel="noreferrer" onClick={(ev) => ev.stopPropagation()}><Stack style={{ paddingTop: 6, height: "100%"}} horizontalAlign="start" disableShrink={true}><Text style={{ textDecoration: textDecoration }} variant={textSize}>{desc}</Text></Stack></a>;
          }
 
          if (column.name === "Jira") {
@@ -562,7 +562,7 @@ const HealthPanelIssues: React.FC<{ desktopAlerts?: boolean }> = observer(({ des
                return null;
             }
 
-            return <Stack horizontalAlign="start" disableShrink={true}><a href={jissue.link} target="_blank" rel="noreferrer" onClick={(ev) => ev.stopPropagation()}><Text style={{ textDecoration: textDecoration }} variant={textSize}>{jissue.key}</Text> </a></Stack>;
+            return <a href={jissue.link} target="_blank" rel="noreferrer" onClick={(ev) => ev.stopPropagation()}><Stack style={{ paddingTop: 6, height: "100%"}} horizontalAlign="start" disableShrink={true}><Text style={{ textDecoration: textDecoration }} variant={textSize}>{jissue.key}</Text></Stack></a>;
          }
       }
 
@@ -574,7 +574,7 @@ const HealthPanelIssues: React.FC<{ desktopAlerts?: boolean }> = observer(({ des
 
       if (column.name === "Opened") {
          const openSince = `${getShortNiceTime(issue.createdAt)} (${getElapsedString(moment(issue.createdAt), moment.utc(), false).trim()})`;
-         return <Stack style={{ paddingRight: 8 }} horizontalAlign="end" verticalAlign="center"><Text variant={textSize}>{openSince}</Text></Stack>;
+         return <Stack style={{ paddingRight: 8, paddingTop: 6, height: "100%" }} horizontalAlign="end"><Text variant={textSize}>{openSince}</Text></Stack>;
       }
 
       return null;
@@ -583,8 +583,8 @@ const HealthPanelIssues: React.FC<{ desktopAlerts?: boolean }> = observer(({ des
    const classes = mergeStyleSets({
       detailsRow: {
          selectors: {
-            '.ms-DetailsRow-cell': {
-               padding: 4,
+            '.ms-DetailsRow-cell': {             
+               padding: 0,
                overflow: "hidden",
                whiteSpace: "nowrap"
             }
diff --git a/Engine/Source/Programs/Horde/HordeDashboard/src/components/UserHome.tsx b/Engine/Source/Programs/Horde/HordeDashboard/src/components/UserHome.tsx
--- a/Engine/Source/Programs/Horde/HordeDashboard/src/components/UserHome.tsx
+++ b/Engine/Source/Programs/Horde/HordeDashboard/src/components/UserHome.tsx
@@ -223,11 +223,11 @@ const HealthPanel: React.FC = observer(() => {
 
    const items = issues?.map(i => { return { issue: i } });
 
-   const columns = [
-      { key: 'health_column1', name: 'Summary', minWidth: 500, maxWidth: 500, isResizable: false },
+   const columns: IColumn[] = [
+      { key: 'health_column1', name: 'Summary', minWidth: 600, maxWidth: 600, isResizable: false },
+      { key: 'health_column2', name: 'Workflow', minWidth: 100, maxWidth: 100, isResizable: false},
       { key: 'health_column3', name: 'Status', minWidth: 320, maxWidth: 320, isResizable: false },
       { key: 'health_column4', name: 'Opened', minWidth: 100, maxWidth: 100, isResizable: false },
-
    ];
 
    const renderItem = (item: HealthItem, index?: number, column?: IColumn) => {
@@ -261,7 +261,7 @@ const HealthPanel: React.FC = observer(() => {
             summary = summary.slice(0, 100) + "...";
          }
 
-         return <Stack horizontal disableShrink={true}>{<IssueStatusIcon issue={issue} />}<Text>{`Issue ${issue.id} - ${summary}`}</Text></Stack>;
+         return <Stack style={{padding: 8}} horizontal disableShrink={true}>{<IssueStatusIcon issue={issue} />}<Text>{`Issue ${issue.id} - ${summary}`}</Text></Stack>;
       }
 
       if (column.name === "Status") {
@@ -270,17 +270,20 @@ const HealthPanel: React.FC = observer(() => {
             status += ` (${ack})`;
          }
 
-         return <Stack horizontalAlign="end" disableShrink={true}><Text>{status}</Text></Stack>;
+         return <Stack style={{padding: 8}} horizontalAlign="end" disableShrink={true}><Text>{status}</Text></Stack>;
       }
 
-      if (column.name === "Ack") {
-         return <Text>{ack}</Text>;
-      }
+      if (column.name === "Workflow") {
+         if (!issue.workflowThreadUrl) {
+            return null;
+         }
 
+         return <a onClick={(e) => e.stopPropagation()} style={{ fontSize: "13px" }} href={issue.workflowThreadUrl} target="_blank" rel="noreferrer"><Stack style={{padding: 8}} horizontalAlign="start" disableShrink={true}>Slack Thread</Stack></a>;
+      }
 
       if (column.name === "Opened") {
          const openSince = `${getShortNiceTime(issue.createdAt)} (${getElapsedString(moment(issue.createdAt), moment.utc(), false).trim()})`;
-         return <Stack style={{ paddingRight: 8 }} horizontalAlign="end"><Text>{openSince}</Text></Stack>;
+         return <Stack style={{ padding: 8 }} horizontalAlign="end"><Text>{openSince}</Text></Stack>;
       }
 
       return <Stack />;
@@ -289,8 +292,7 @@ const HealthPanel: React.FC = observer(() => {
    const classes = mergeStyleSets({
       detailsRow: {
          selectors: {
-            '.ms-DetailsRow-cell': {
-               padding: 8,
+            '.ms-DetailsRow-cell': {               
                overflow: "hidden",
                whiteSpace: "nowrap"
             }
@@ -1030,22 +1032,22 @@ const JobsPanel: React.FC<{ includeOtherPreflights: boolean }> = observer(({ inc
                   </Stack>
                </Stack>
                <Stack styles={{ root: { paddingLeft: 4, paddingRight: 0, paddingTop: 8, paddingBottom: 4 } }}>
-                     {!!jobItems.length && <DetailsList
-                        styles={{ root: { paddingLeft: 8, paddingRight: 8, marginBottom: 18 } }}
-                        compact={true}
-                        isHeaderVisible={false}
-                        indentWidth={0}
-                        items={jobItems}
-                        columns={columns}
-                        setKey="set"
-                        selectionMode={SelectionMode.none}
-                        layoutMode={DetailsListLayoutMode.fixedColumns}
-                        onRenderDetailsHeader={onRenderDetailsHeader}
-                        onRenderItemColumn={renderItem}
-                        onRenderRow={renderRow}
-                        onShouldVirtualize={() => true}
-                     />}
-                     {!dashboard.pinnedJobsIds.length && <Stack style={{ paddingBottom: 12 }}><Text variant="medium">No jobs are pinned</Text></Stack>}
+                  {!!jobItems.length && <DetailsList
+                     styles={{ root: { paddingLeft: 8, paddingRight: 8, marginBottom: 18 } }}
+                     compact={true}
+                     isHeaderVisible={false}
+                     indentWidth={0}
+                     items={jobItems}
+                     columns={columns}
+                     setKey="set"
+                     selectionMode={SelectionMode.none}
+                     layoutMode={DetailsListLayoutMode.fixedColumns}
+                     onRenderDetailsHeader={onRenderDetailsHeader}
+                     onRenderItemColumn={renderItem}
+                     onRenderRow={renderRow}
+                     onShouldVirtualize={() => true}
+                  />}
+                  {!dashboard.pinnedJobsIds.length && <Stack style={{ paddingBottom: 12 }}><Text variant="medium">No jobs are pinned</Text></Stack>}
                </Stack>
             </Stack>
          </Stack>
@@ -1095,7 +1097,7 @@ export const UserHomeView: React.FC = () => {
          <TopNav />
          <Breadcrumbs items={[{ text: 'Home' }]} />
          <Stack horizontal>
-            <div key={`windowsize_streamview_${windowSize.width}_${windowSize.height}`} style={{ width: vw / 2 - (1440/2), flexShrink: 0, backgroundColor: modeColors.background }} />
+            <div key={`windowsize_streamview_${windowSize.width}_${windowSize.height}`} style={{ width: vw / 2 - (1440 / 2), flexShrink: 0, backgroundColor: modeColors.background }} />
             <Stack tokens={{ childrenGap: 0 }} styles={{ root: { backgroundColor: modeColors.background, width: "100%" } }}>
                <Stack style={{ maxWidth: 1440, paddingTop: 6, marginLeft: 4 }}>
                   <UserHomeViewInner />
