Horde: Move Tools page from Server => Resources to Server => Software sections
#rnx
#skipci

diff --git a/Engine/Source/Programs/Horde/HordeDashboard/src/components/TopNav.tsx b/Engine/Source/Programs/Horde/HordeDashboard/src/components/TopNav.tsx
--- a/Engine/Source/Programs/Horde/HordeDashboard/src/components/TopNav.tsx
+++ b/Engine/Source/Programs/Horde/HordeDashboard/src/components/TopNav.tsx
@@ -377,14 +377,6 @@ export const TopNav: React.FC<{ suppressServer?: boolean }> = observer(({ suppre
          link: `/perforce/servers`
       });
 
-      if (getSiteConfig().environment !== "production") {
-         resourceItems.push({
-            key: "admin_tools",
-            text: "Tools",
-            link: `/tools`
-         });   
-      }
-
       subItems.push({
          itemType: ContextualMenuItemType.Section,
          key: `admin_resources`,
@@ -446,6 +438,26 @@ export const TopNav: React.FC<{ suppressServer?: boolean }> = observer(({ suppre
          }
       });
 
+      if (getSiteConfig().environment !== "production") {
+
+         const softwareItems: IContextualMenuItem[] = [];
+         softwareItems.push({
+            key: "software_tools",
+            text: "Tools",
+            link: `/tools`
+         });
+
+         subItems.push({
+            itemType: ContextualMenuItemType.Section,
+            key: `server_software`,
+            sectionProps: {
+               title: "Software",
+               items: softwareItems,
+               bottomDivider: true
+            }
+         });
+      }
+
       // Configuration
       const versionItems: IContextualMenuItem[] = [];
       versionItems.push({
@@ -470,7 +482,7 @@ export const TopNav: React.FC<{ suppressServer?: boolean }> = observer(({ suppre
          const previewItem: ICommandBarItemProps = {
             key: "preview_button",
             text: "CHANGES",
-            onClick: () => {setShowPreviewChanges(true)}
+            onClick: () => { setShowPreviewChanges(true) }
          };
          cbProps.push(previewItem);
       }
