Horde: Adjust agent names on agent view
#rnx
#skipci

diff --git a/Engine/Source/Programs/Horde/HordeDashboard/src/components/AgentView.tsx b/Engine/Source/Programs/Horde/HordeDashboard/src/components/AgentView.tsx
--- a/Engine/Source/Programs/Horde/HordeDashboard/src/components/AgentView.tsx
+++ b/Engine/Source/Programs/Horde/HordeDashboard/src/components/AgentView.tsx
@@ -621,7 +621,7 @@ class LocalState {
          {
             key: 'name',
             displayText: 'Name',
-            colSize: 145,
+            colSize: 170,
             isChecked: true,
             isCheckable: false,
             isSorted: true,
@@ -2262,7 +2262,7 @@ export const AgentView: React.FC = observer(() => {
                      {getAgentStatusIcon(agent)}
                   </Stack.Item>
                   <Stack.Item align={"center"}>
-                     <ReactLink styles={{ root: { paddingTop: '1px' } }} title="Lease and Session History" onClick={() => localState.setAgentId(agent.id) }>{agent.name}</ReactLink>
+                     <ReactLink styles={{ root: { paddingTop: '1px', fontSize: "12px"} }} title="Lease and Session History" onClick={() => localState.setAgentId(agent.id) }>{agent.name}</ReactLink>
                   </Stack.Item>
                </Stack>
             );
