Horde: Add method attributes for storage controller, causing exceptions loading Swagger pages.

diff --git a/Engine/Source/Programs/Horde/Horde.Build/Storage/StorageController.cs b/Engine/Source/Programs/Horde/Horde.Build/Storage/StorageController.cs
--- a/Engine/Source/Programs/Horde/Horde.Build/Storage/StorageController.cs
+++ b/Engine/Source/Programs/Horde/Horde.Build/Storage/StorageController.cs
@@ -475,6 +475,7 @@ namespace Horde.Build.Storage
 		/// <param name="includePackets">Whether to include packets for the bundle</param>
 		/// <param name="cancellationToken">Cancellation token for the operation</param>
 		/// <returns></returns>
+		[HttpGet]
 		[Route("/api/v1/storage/{namespaceId}/bundles/{*locator}")]
 		public async Task<ActionResult<object>> GetBundleAsync(NamespaceId namespaceId, BlobLocator locator, [FromQuery(Name = "imports")] bool includeImports = false, [FromQuery(Name = "exports")] bool includeExports = true, [FromQuery(Name = "packets")] bool includePackets = false, CancellationToken cancellationToken = default)
 		{
@@ -556,6 +557,7 @@ namespace Horde.Build.Storage
 		/// <param name="exportIdx">Index of the export</param>
 		/// <param name="cancellationToken">Cancellation token for the operation</param>
 		/// <returns></returns>
+		[HttpGet]
 		[Route("/api/v1/storage/{namespaceId}/nodes/{*locator}")]
 		public async Task<ActionResult<object>> GetNodeAsync(NamespaceId namespaceId, BlobLocator locator, [FromQuery(Name = "export")] int exportIdx, CancellationToken cancellationToken = default)
 		{
