Horde: Terminate the session if an agent reports its status as stopping.

diff --git a/Engine/Source/Programs/Horde/Horde.Build/Agents/AgentService.cs b/Engine/Source/Programs/Horde/Horde.Build/Agents/AgentService.cs
--- a/Engine/Source/Programs/Horde/Horde.Build/Agents/AgentService.cs
+++ b/Engine/Source/Programs/Horde/Horde.Build/Agents/AgentService.cs
@@ -629,36 +629,19 @@ namespace Horde.Build.Agents
 			IAgent? agent = inAgent;
 			while (agent != null)
 			{
-				// If the agent is stopping and doesn't have any leases, we can terminate the current session.
-				if (status == AgentStatus.Stopping && newLeases.Count == 0)
+				// Check the session id is correct.
+				if (agent.SessionId != sessionId)
 				{
-					// If we've already decided to terminate the session, this update is redundant but harmless
-					if (agent.SessionId != sessionId)
+					if (status == AgentStatus.Stopping)
 					{
-						break;
+						break; // Harmless; agent is not doing any work.
 					}
-
-					// If the session is valid, we can terminate once the agent leases are also empty
-					if (agent.Leases.Count == 0)
+					else
 					{
-						IAgent? terminatedAgent = await TryTerminateSessionAsync(agent);
-						if (terminatedAgent == null)
-						{
-							agent = await GetAgentAsync(agent.Id);
-							continue;
-						}
-
-						agent = terminatedAgent;
-						break;
+						throw new InvalidOperationException($"Invalid agent session {sessionId}");
 					}
 				}
 
-				// Check the session id is correct.
-				if (agent.SessionId != sessionId)
-				{
-					throw new InvalidOperationException($"Invalid agent session {sessionId}");
-				}
-
 				// Check the session hasn't expired
 				if (!agent.IsSessionValid(utcNow))
 				{
@@ -739,6 +722,19 @@ namespace Horde.Build.Agents
 				// Fetch the agent again
 				agent = await GetAgentAsync(agent.Id);
 			}
+
+			// If the agent is stopping, terminate the session
+			while (agent != null && agent.Status == AgentStatus.Stopping && agent.Leases.Count == 0)
+			{
+				IAgent? terminatedAgent = await TryTerminateSessionAsync(agent);
+				if (terminatedAgent != null)
+				{
+					agent = terminatedAgent;
+					break;
+				}
+				agent = await GetAgentAsync(agent.Id);
+			}
+
 			return agent;
 		}
 
