Horde: Rename the Horde.Build project to Horde.Server.

diff --git a/Engine/Source/Programs/Horde/Horde.Agent/Horde.Agent.csproj b/Engine/Source/Programs/Horde/Horde.Agent/Horde.Agent.csproj
--- a/Engine/Source/Programs/Horde/Horde.Agent/Horde.Agent.csproj
+++ b/Engine/Source/Programs/Horde/Horde.Agent/Horde.Agent.csproj
@@ -97,22 +97,22 @@
     <Compile Include="..\..\UnrealBuildTool\Matchers\*.cs">
       <Link>Parser\Matchers\%(FileName)%(Extension)</Link>
     </Compile>
-    <Content Include="..\Horde.Build\Protos\horde\*.proto">
+    <Content Include="..\Horde.Server\Protos\horde\*.proto">
       <Link>Protos\%(FileName)%(Extension)</Link>
     </Content>
     <Content Include="Horde.bat">
       <CopyToOutputDirectory>PreserveNewest</CopyToOutputDirectory>
     </Content>
-    <Content Include="..\Horde.Build\Protos\horde\common\**\*.proto">
+    <Content Include="..\Horde.Server\Protos\horde\common\**\*.proto">
       <Link>Protos\horde\common\%(RecursiveDir)%(FileName)%(Extension)</Link>
     </Content>
-    <Content Include="..\Horde.Build\Protos\horde\messages\**\*.proto">
+    <Content Include="..\Horde.Server\Protos\horde\messages\**\*.proto">
       <Link>Protos\horde\messages\%(RecursiveDir)%(FileName)%(Extension)</Link>
     </Content>
-    <Content Include="..\Horde.Build\Protos\horde\tasks\**\*.proto">
+    <Content Include="..\Horde.Server\Protos\horde\tasks\**\*.proto">
       <Link>Protos\horde\tasks\%(RecursiveDir)%(FileName)%(Extension)</Link>
     </Content>
-    <Compile Include="..\Horde.Build\Protos\horde\**\*.cs">
+    <Compile Include="..\Horde.Server\Protos\horde\**\*.cs">
       <Link>Protos\horde\%(RecursiveDir)%(FileName)%(Extension)</Link>
     </Compile>
     <Compile Include="../../../../Platforms/*/Source/Programs/Horde/Horde.Agent/**/*.cs">
@@ -131,34 +131,34 @@
   </ItemGroup>
 
   <ItemGroup>
-    <Protobuf Include="..\Horde.Build\Protos\horde\horde_rpc.proto" ProtoRoot="..\Horde.Build\Protos" GrpcServices="Client" Access="Public" />
-    <Protobuf Include="..\Horde.Build\Protos\horde\horde_rpc_messages.proto" ProtoRoot="..\Horde.Build\Protos" Access="Public" />
-    <Protobuf Include="..\Horde.Build\Protos\horde\job_rpc.proto" ProtoRoot="..\Horde.Build\Protos" GrpcServices="Client" Access="Public" />
-    <Protobuf Include="..\Horde.Build\Protos\horde\job_rpc_messages.proto" ProtoRoot="..\Horde.Build\Protos" Access="Public" />
-    <Protobuf Include="..\Horde.Build\Protos\horde\log_rpc.proto" ProtoRoot="..\Horde.Build\Protos" GrpcServices="Client" Access="Public" />
-    <Protobuf Include="..\Horde.Build\Protos\horde\log_rpc_messages.proto" ProtoRoot="..\Horde.Build\Protos" Access="Public" />
-    <Protobuf Include="..\Horde.Build\Protos\horde\common\agent_status.proto" ProtoRoot="..\Horde.Build\Protos" Access="Public" />
-    <Protobuf Include="..\Horde.Build\Protos\horde\common\event_severity.proto" ProtoRoot="..\Horde.Build\Protos" Access="Public" />
-    <Protobuf Include="..\Horde.Build\Protos\horde\common\io_hash.proto" ProtoRoot="..\Horde.Build\Protos" Access="Public" />
-    <Protobuf Include="..\Horde.Build\Protos\horde\common\jobstep_outcome.proto" ProtoRoot="..\Horde.Build\Protos" Access="Public" />
-    <Protobuf Include="..\Horde.Build\Protos\horde\common\jobstep_state.proto" ProtoRoot="..\Horde.Build\Protos" Access="Public" />
-    <Protobuf Include="..\Horde.Build\Protos\horde\common\label_change.proto" ProtoRoot="..\Horde.Build\Protos" Access="Public" />
-    <Protobuf Include="..\Horde.Build\Protos\horde\common\lease_outcome.proto" ProtoRoot="..\Horde.Build\Protos" Access="Public" />
-    <Protobuf Include="..\Horde.Build\Protos\horde\common\lease_state.proto" ProtoRoot="..\Horde.Build\Protos" Access="Public" />
-    <Protobuf Include="..\Horde.Build\Protos\horde\common\namespace_id.proto" ProtoRoot="..\Horde.Build\Protos" Access="Public" />
-    <Protobuf Include="..\Horde.Build\Protos\horde\common\priority.proto" ProtoRoot="..\Horde.Build\Protos" Access="Public" />
-    <Protobuf Include="..\Horde.Build\Protos\horde\messages\agent_capabilities.proto" ProtoRoot="..\Horde.Build\Protos" Access="Public" />
-    <Protobuf Include="..\Horde.Build\Protos\horde\messages\agent_requirements.proto" ProtoRoot="..\Horde.Build\Protos" Access="Public" />
-    <Protobuf Include="..\Horde.Build\Protos\horde\messages\agent_workspace.proto" ProtoRoot="..\Horde.Build\Protos" Access="Public" />
-    <Protobuf Include="..\Horde.Build\Protos\horde\messages\lease.proto" ProtoRoot="..\Horde.Build\Protos" Access="Public" />
-    <Protobuf Include="..\Horde.Build\Protos\horde\messages\aws_lambda.proto" ProtoRoot="..\Horde.Build\Protos" Access="Public" />
-    <Protobuf Include="..\Horde.Build\Protos\horde\tasks\compute_task.proto" ProtoRoot="..\Horde.Build\Protos" Access="Public" />
-    <Protobuf Include="..\Horde.Build\Protos\horde\tasks\conform_task.proto" ProtoRoot="..\Horde.Build\Protos" Access="Public" />
-    <Protobuf Include="..\Horde.Build\Protos\horde\tasks\job_task.proto" ProtoRoot="..\Horde.Build\Protos" Access="Public" />
-    <Protobuf Include="..\Horde.Build\Protos\horde\tasks\restart_task.proto" ProtoRoot="..\Horde.Build\Protos" Access="Public" />
-    <Protobuf Include="..\Horde.Build\Protos\horde\tasks\shutdown_task.proto" ProtoRoot="..\Horde.Build\Protos" Access="Public" />
-    <Protobuf Include="..\Horde.Build\Protos\horde\tasks\test_task.proto" ProtoRoot="..\Horde.Build\Protos" Access="Public" />
-    <Protobuf Include="..\Horde.Build\Protos\horde\tasks\upgrade_task.proto" ProtoRoot="..\Horde.Build\Protos" Access="Public" />
+    <Protobuf Include="..\Horde.Server\Protos\horde\horde_rpc.proto" ProtoRoot="..\Horde.Server\Protos" GrpcServices="Client" Access="Public" />
+    <Protobuf Include="..\Horde.Server\Protos\horde\horde_rpc_messages.proto" ProtoRoot="..\Horde.Server\Protos" Access="Public" />
+    <Protobuf Include="..\Horde.Server\Protos\horde\job_rpc.proto" ProtoRoot="..\Horde.Server\Protos" GrpcServices="Client" Access="Public" />
+    <Protobuf Include="..\Horde.Server\Protos\horde\job_rpc_messages.proto" ProtoRoot="..\Horde.Server\Protos" Access="Public" />
+    <Protobuf Include="..\Horde.Server\Protos\horde\log_rpc.proto" ProtoRoot="..\Horde.Server\Protos" GrpcServices="Client" Access="Public" />
+    <Protobuf Include="..\Horde.Server\Protos\horde\log_rpc_messages.proto" ProtoRoot="..\Horde.Server\Protos" Access="Public" />
+    <Protobuf Include="..\Horde.Server\Protos\horde\common\agent_status.proto" ProtoRoot="..\Horde.Server\Protos" Access="Public" />
+    <Protobuf Include="..\Horde.Server\Protos\horde\common\event_severity.proto" ProtoRoot="..\Horde.Server\Protos" Access="Public" />
+    <Protobuf Include="..\Horde.Server\Protos\horde\common\io_hash.proto" ProtoRoot="..\Horde.Server\Protos" Access="Public" />
+    <Protobuf Include="..\Horde.Server\Protos\horde\common\jobstep_outcome.proto" ProtoRoot="..\Horde.Server\Protos" Access="Public" />
+    <Protobuf Include="..\Horde.Server\Protos\horde\common\jobstep_state.proto" ProtoRoot="..\Horde.Server\Protos" Access="Public" />
+    <Protobuf Include="..\Horde.Server\Protos\horde\common\label_change.proto" ProtoRoot="..\Horde.Server\Protos" Access="Public" />
+    <Protobuf Include="..\Horde.Server\Protos\horde\common\lease_outcome.proto" ProtoRoot="..\Horde.Server\Protos" Access="Public" />
+    <Protobuf Include="..\Horde.Server\Protos\horde\common\lease_state.proto" ProtoRoot="..\Horde.Server\Protos" Access="Public" />
+    <Protobuf Include="..\Horde.Server\Protos\horde\common\namespace_id.proto" ProtoRoot="..\Horde.Server\Protos" Access="Public" />
+    <Protobuf Include="..\Horde.Server\Protos\horde\common\priority.proto" ProtoRoot="..\Horde.Server\Protos" Access="Public" />
+    <Protobuf Include="..\Horde.Server\Protos\horde\messages\agent_capabilities.proto" ProtoRoot="..\Horde.Server\Protos" Access="Public" />
+    <Protobuf Include="..\Horde.Server\Protos\horde\messages\agent_requirements.proto" ProtoRoot="..\Horde.Server\Protos" Access="Public" />
+    <Protobuf Include="..\Horde.Server\Protos\horde\messages\agent_workspace.proto" ProtoRoot="..\Horde.Server\Protos" Access="Public" />
+    <Protobuf Include="..\Horde.Server\Protos\horde\messages\lease.proto" ProtoRoot="..\Horde.Server\Protos" Access="Public" />
+    <Protobuf Include="..\Horde.Server\Protos\horde\messages\aws_lambda.proto" ProtoRoot="..\Horde.Server\Protos" Access="Public" />
+    <Protobuf Include="..\Horde.Server\Protos\horde\tasks\compute_task.proto" ProtoRoot="..\Horde.Server\Protos" Access="Public" />
+    <Protobuf Include="..\Horde.Server\Protos\horde\tasks\conform_task.proto" ProtoRoot="..\Horde.Server\Protos" Access="Public" />
+    <Protobuf Include="..\Horde.Server\Protos\horde\tasks\job_task.proto" ProtoRoot="..\Horde.Server\Protos" Access="Public" />
+    <Protobuf Include="..\Horde.Server\Protos\horde\tasks\restart_task.proto" ProtoRoot="..\Horde.Server\Protos" Access="Public" />
+    <Protobuf Include="..\Horde.Server\Protos\horde\tasks\shutdown_task.proto" ProtoRoot="..\Horde.Server\Protos" Access="Public" />
+    <Protobuf Include="..\Horde.Server\Protos\horde\tasks\test_task.proto" ProtoRoot="..\Horde.Server\Protos" Access="Public" />
+    <Protobuf Include="..\Horde.Server\Protos\horde\tasks\upgrade_task.proto" ProtoRoot="..\Horde.Server\Protos" Access="Public" />
   </ItemGroup>
 
   <ItemGroup>
diff --git a/Engine/Source/Programs/Horde/Horde.Server.Tests/AgentControllerTest.cs b/Engine/Source/Programs/Horde/Horde.Server.Tests/AgentControllerTest.cs
new file mode 100644
--- /dev/null
+++ b/Engine/Source/Programs/Horde/Horde.Server.Tests/AgentControllerTest.cs
@@ -0,0 +1,32 @@
+// Copyright Epic Games, Inc. All Rights Reserved.
+
+using System.Threading.Tasks;
+using Horde.Server.Agents;
+using Microsoft.AspNetCore.Mvc;
+using Microsoft.VisualStudio.TestTools.UnitTesting;
+
+namespace Horde.Server.Tests;
+
+[TestClass]
+public class AgentControllerDbTest : TestSetup
+{
+	[TestMethod]
+	public async Task UpdateAgent()
+	{
+		Fixture fixture = await CreateFixtureAsync();
+		IAgent fixtureAgent = fixture.Agent1;
+
+		ActionResult<object> obj = await AgentsController.GetAgentAsync(fixtureAgent.Id);
+		GetAgentResponse getRes = (obj.Value as GetAgentResponse)!;
+		Assert.AreEqual(fixture!.Agent1Name.ToUpper(), getRes.Name);
+		Assert.IsNull(getRes.Comment);
+
+		UpdateAgentRequest updateReq = new();
+		updateReq.Comment = "foo bar baz";
+		await AgentsController.UpdateAgentAsync(fixtureAgent.Id, updateReq);
+
+		obj = await AgentsController.GetAgentAsync(fixtureAgent.Id);
+		getRes = (obj.Value as GetAgentResponse)!;
+		Assert.AreEqual("foo bar baz", getRes.Comment);
+	}
+}
\ No newline at end of file
diff --git a/Engine/Source/Programs/Horde/Horde.Server.Tests/ControllerIntegrationTest.cs b/Engine/Source/Programs/Horde/Horde.Server.Tests/ControllerIntegrationTest.cs
new file mode 100644
--- /dev/null
+++ b/Engine/Source/Programs/Horde/Horde.Server.Tests/ControllerIntegrationTest.cs
@@ -0,0 +1,106 @@
+// Copyright Epic Games, Inc. All Rights Reserved.
+
+using System;
+using System.Collections.Generic;
+using System.Net.Http;
+using System.Threading.Tasks;
+using Amazon.CloudWatch;
+using Horde.Server.Jobs;
+using Horde.Server.Server;
+using Horde.Server.Perforce;
+using Horde.Server.Tests.Stubs.Services;
+using Microsoft.AspNetCore.Hosting;
+using Microsoft.AspNetCore.Mvc.Testing;
+using Microsoft.AspNetCore.TestHost;
+using Microsoft.Extensions.Configuration;
+using Microsoft.Extensions.DependencyInjection;
+using Horde.Server.Agents;
+using Horde.Server.Streams;
+using Horde.Server.Jobs.Templates;
+using Horde.Server.Configuration;
+using Horde.Server.Jobs.Graphs;
+using Horde.Server.Jobs.Artifacts;
+using Moq;
+using Microsoft.Extensions.Options;
+
+namespace Horde.Server.Tests;
+
+public class TestWebApplicationFactory<TStartup> : WebApplicationFactory<TStartup> where TStartup : class
+{
+	private readonly MongoDbInstance _mongoDbInstance;
+
+	public TestWebApplicationFactory(MongoDbInstance mongoDbInstance)
+	{
+		_mongoDbInstance = mongoDbInstance;
+	}
+
+	protected override void ConfigureWebHost(IWebHostBuilder builder)
+	{
+		Dictionary<string, string?> dict = new()
+		{
+			{ "Horde:DatabaseConnectionString", _mongoDbInstance.ConnectionString },
+			{ "Horde:DatabaseName", _mongoDbInstance.DatabaseName },
+			{ "Horde:LogServiceWriteCacheType", "inmemory" },
+			{ "Horde:DisableAuth", "true" },
+			{ "Horde:OidcAuthority", null },
+			{ "Horde:OidcClientId", null }
+		};
+
+		Mock<IAmazonCloudWatch> cloudWatchMock = new (MockBehavior.Strict);
+		builder.ConfigureAppConfiguration((hostingContext, config) => { config.AddInMemoryCollection(dict); });
+		builder.ConfigureTestServices(collection =>
+		{
+			collection.AddSingleton<IPerforceService, PerforceServiceStub>();
+			collection.AddSingleton<IAmazonCloudWatch>(x => cloudWatchMock.Object);
+		});
+	}
+}
+
+public class ControllerIntegrationTest : IDisposable
+{
+	private readonly Lazy<Task<Fixture>> _fixture;
+
+	public ControllerIntegrationTest()
+	{
+		MongoDbInstance = new MongoDbInstance();
+		Factory = new TestWebApplicationFactory<Startup>(MongoDbInstance);
+		Client = Factory.CreateClient();
+
+		_fixture = new Lazy<Task<Fixture>>(Task.Run(() => CreateFixture()));
+	}
+
+	protected MongoDbInstance MongoDbInstance { get; }
+	private TestWebApplicationFactory<Startup> Factory { get; }
+	protected HttpClient Client { get; }
+
+	public void Dispose()
+	{
+		Dispose(true);
+		GC.SuppressFinalize(this);
+	}
+
+	protected virtual void Dispose(bool disposing)
+	{
+		MongoDbInstance.Dispose();
+	}
+
+	public Task<Fixture> GetFixture()
+	{
+		return _fixture.Value;
+	}
+
+	private async Task<Fixture> CreateFixture()
+	{
+		IServiceProvider services = Factory.Services;
+		ConfigService configService = services.GetRequiredService<ConfigService>();
+		MongoService mongoService = services.GetRequiredService<MongoService>();
+		ITemplateCollection templateService = services.GetRequiredService<ITemplateCollection>();
+		JobService jobService = services.GetRequiredService<JobService>();
+		IArtifactCollectionV1 artifactCollection = services.GetRequiredService<IArtifactCollectionV1>();
+		AgentService agentService = services.GetRequiredService<AgentService>();
+		GraphCollection graphCollection = new(mongoService);
+		IOptions<ServerSettings> serverSettings = services.GetRequiredService<IOptions<ServerSettings>>();
+
+		return await Fixture.Create(configService, graphCollection, templateService, jobService, artifactCollection, agentService, serverSettings.Value);
+	}
+}
\ No newline at end of file
diff --git a/Engine/Source/Programs/Horde/Horde.Server.Tests/DatabaseIntegrationTest.cs b/Engine/Source/Programs/Horde/Horde.Server.Tests/DatabaseIntegrationTest.cs
new file mode 100644
--- /dev/null
+++ b/Engine/Source/Programs/Horde/Horde.Server.Tests/DatabaseIntegrationTest.cs
@@ -0,0 +1,228 @@
+// Copyright Epic Games, Inc. All Rights Reserved.
+
+using System;
+using System.Collections.Generic;
+using System.Net;
+using Horde.Server.Server;
+using Horde.Server.Utilities;
+using Microsoft.Extensions.DependencyInjection;
+using Microsoft.Extensions.Logging;
+using Microsoft.Extensions.Options;
+using Microsoft.VisualStudio.TestTools.UnitTesting;
+using MongoDB.Driver;
+using StackExchange.Redis;
+
+namespace Horde.Server.Tests
+{
+	// Stub for fulfilling IOptions interface during testing
+	public sealed class TestOptions<T> : IOptions<T> where T : class
+	{
+		public TestOptions(T options) => Value = options;
+		public T Value { get; }
+	}
+
+	// Stub for fulfilling IOptionsMonitor interface during testing
+	public sealed class TestOptionsMonitor<T> : IOptionsMonitor<T>
+        where T : class, new()
+    {
+		sealed class Disposable : IDisposable
+		{
+			public void Dispose() { }
+		}
+
+		public TestOptionsMonitor(T currentValue)
+        {
+            CurrentValue = currentValue;
+        }
+
+        public T Get(string name)
+        {
+            return CurrentValue;
+        }
+
+        public IDisposable OnChange(Action<T, string> listener)
+        {
+			return new Disposable();
+        }
+
+        public T CurrentValue { get; }
+    }
+
+	public sealed class MongoDbInstance : IDisposable
+	{
+		public string DatabaseName { get; }
+		public string ConnectionString { get; }
+		MongoClient Client { get; }
+
+		private static readonly object s_lockObject = new object();
+		private static MongoDbRunnerLocal? s_mongoDbRunner;
+		private static int s_nextDatabaseIndex = 1;
+		public const string MongoDbDatabaseNamePrefix = "HordeServerTest_";
+
+		public MongoDbInstance()
+		{
+			int databaseIndex;
+			lock (s_lockObject)
+			{
+				if (s_mongoDbRunner == null)
+				{
+					// One-time setup per test run to avoid overhead of starting the external MongoDB process
+					Startup.ConfigureMongoDbClient();
+					s_mongoDbRunner = new MongoDbRunnerLocal();
+					s_mongoDbRunner.Start();
+
+					// Drop all the previous databases
+					MongoClientSettings mongoSettings = MongoClientSettings.FromConnectionString(s_mongoDbRunner.GetConnectionString());
+					MongoClient client = new MongoClient(mongoSettings);
+
+					List<string> dropDatabaseNames = client.ListDatabaseNames().ToList();
+					foreach (string dropDatabaseName in dropDatabaseNames)
+					{
+						if (dropDatabaseName.StartsWith(MongoDbDatabaseNamePrefix, StringComparison.Ordinal))
+						{
+							client.DropDatabase(dropDatabaseName);
+						}
+					}
+				}
+				databaseIndex = s_nextDatabaseIndex++;
+			}
+
+			DatabaseName = $"{MongoDbDatabaseNamePrefix}{databaseIndex}";
+			ConnectionString = $"{s_mongoDbRunner.GetConnectionString()}/{DatabaseName}";
+			Client = new MongoClient(MongoClientSettings.FromConnectionString(ConnectionString));
+		}
+
+		public void Dispose()
+		{
+			Client.DropDatabase(DatabaseName);
+		}
+	}
+
+	public class ServiceTest : IDisposable
+	{
+		private ServiceProvider? _serviceProvider = null;
+
+		public IServiceProvider ServiceProvider
+		{
+			get
+			{
+				if (_serviceProvider == null)
+				{
+					IServiceCollection services = new ServiceCollection();
+					ConfigureServices(services);
+
+					_serviceProvider = services.BuildServiceProvider();
+				}
+				return _serviceProvider;
+			}
+		}
+
+		public void Dispose()
+		{
+			Dispose(true);
+			GC.SuppressFinalize(this);
+		}
+
+		protected virtual void Dispose(bool disposing)
+		{
+			_serviceProvider?.Dispose();
+		}
+
+		protected virtual void ConfigureSettings(ServerSettings settings)
+		{
+		}
+
+		protected virtual void ConfigureServices(IServiceCollection services)
+		{
+		}
+	}
+
+	public class DatabaseIntegrationTest : ServiceTest
+    {
+		private static readonly object s_lockObject = new object();
+
+		private MongoDbInstance? _mongoDbInstance;
+		private MongoService? _mongoService;
+		private readonly LoggerFactory _loggerFactory = new LoggerFactory();
+
+		private static RedisRunner? s_redisRunner;
+		private RedisService? _redisService;
+        public const int RedisDbNum = 15;
+
+		public DatabaseIntegrationTest()
+		{
+		}
+
+		protected override void ConfigureServices(IServiceCollection services)
+		{
+			services.AddSingleton(GetMongoServiceSingleton());
+			services.AddSingleton(GetRedisServiceSingleton());
+		}
+
+		protected override void Dispose(bool disposing)
+		{
+			base.Dispose(disposing);
+
+			_mongoDbInstance?.Dispose();
+			_mongoService?.Dispose();
+			_redisService?.Dispose();
+			_loggerFactory.Dispose();
+		}
+
+		public MongoService GetMongoServiceSingleton()
+        {
+			lock(s_lockObject)
+			{
+				if (_mongoService == null)
+				{
+					_mongoDbInstance = new MongoDbInstance();
+
+					ServerSettings ss = new ServerSettings();
+					ss.DatabaseName = _mongoDbInstance.DatabaseName;
+					ss.DatabaseConnectionString = _mongoDbInstance.ConnectionString;
+
+					_mongoService = new MongoService(Options.Create(ss), _loggerFactory);
+				}
+			}
+			return _mongoService;
+        }
+
+		private static RedisRunner GetRedisRunner()
+		{
+			lock (s_lockObject)
+			{
+				if (s_redisRunner == null)
+				{
+					// One-time setup per test run to avoid overhead of starting the external Redis process
+					s_redisRunner = new RedisRunner();
+					s_redisRunner.Start();
+				}
+			}
+			return s_redisRunner;
+		}
+
+		public RedisService GetRedisServiceSingleton()
+        {
+			if (_redisService == null)
+			{
+				RedisRunner redisRunner = GetRedisRunner();
+
+				(string host, int port) = redisRunner.GetListenAddress();
+				_redisService = new RedisService($"{host}:{port},allowAdmin=true", RedisDbNum, _loggerFactory);
+				IConnectionMultiplexer cm = _redisService.ConnectionPool.GetConnection();
+
+				foreach (EndPoint endpoint in cm.GetEndPoints())
+				{
+					cm.GetServer(endpoint).FlushDatabase(RedisDbNum);
+				}
+			}
+			return _redisService;
+        }
+        
+		public static T Deref<T>(T? item)
+		{
+			Assert.IsNotNull(item);
+			return item!;
+		}
+	}
+}
\ No newline at end of file
diff --git a/Engine/Source/Programs/Horde/Horde.Server.Tests/Fixture.cs b/Engine/Source/Programs/Horde/Horde.Server.Tests/Fixture.cs
new file mode 100644
--- /dev/null
+++ b/Engine/Source/Programs/Horde/Horde.Server.Tests/Fixture.cs
@@ -0,0 +1,134 @@
+// Copyright Epic Games, Inc. All Rights Reserved.
+
+using System.Collections.Generic;
+using System.IO;
+using System.Text;
+using System.Threading.Tasks;
+using EpicGames.Core;
+using Horde.Server.Agents;
+using Horde.Server.Agents.Pools;
+using Horde.Server.Configuration;
+using Horde.Server.Jobs;
+using Horde.Server.Jobs.Artifacts;
+using Horde.Server.Jobs.Graphs;
+using Horde.Server.Jobs.Templates;
+using Horde.Server.Projects;
+using Horde.Server.Server;
+using Horde.Server.Streams;
+using Horde.Server.Utilities;
+using HordeCommon;
+
+namespace Horde.Server.Tests
+{
+	public class Fixture
+	{
+		public IJob Job1 { get; private set; } = null!;
+		public IJob Job2 { get; private set; } = null!;
+		public ITemplate Template { get; private set; } = null!;
+		public IGraph Graph { get; private set; } = null!;
+		public StreamId StreamId { get; private set; }
+		public StreamConfig? StreamConfig { get; private set; }
+		public TemplateId TemplateRefId1 { get; private set; }
+		public TemplateId TemplateRefId2 { get; private set; }
+		public IArtifactV1 Job1Artifact { get; private set; } = null!;
+		public string Job1ArtifactData { get; private set; } = null!;
+		public IAgent Agent1 { get; private set; } = null!;
+		public string Agent1Name { get; private set; } = null!;
+		public const string PoolName = "TestingPool";
+
+		public static async Task<Fixture> Create(ConfigService configService, IGraphCollection graphCollection, ITemplateCollection templateCollection, JobService jobService, IArtifactCollectionV1 artifactCollection, AgentService agentService, ServerSettings serverSettings)
+		{
+			Fixture fixture = new Fixture();
+			await fixture.Populate(configService, graphCollection, templateCollection, jobService, artifactCollection, agentService, serverSettings);
+
+//			(PerforceService as PerforceServiceStub)?.AddChange("//UE5/Main", 112233, "leet.coder", "Did stuff", new []{"file.cpp"});
+//			(PerforceService as PerforceServiceStub)?.AddChange("//UE5/Main", 1111, "swarm", "A shelved CL here", new []{"renderer.cpp"});
+		
+			return fixture;
+		}
+
+		private async Task Populate(ConfigService configService, IGraphCollection graphCollection, ITemplateCollection templateCollection, JobService jobService, IArtifactCollectionV1 artifactCollection, AgentService agentService, ServerSettings serverSettings)
+		{
+			FixtureGraph fg = new FixtureGraph();
+			fg.Id = ContentHash.Empty;
+			fg.Schema = 1122;
+			fg.Groups = new List<INodeGroup>();
+			fg.Aggregates = new List<IAggregate>();
+			fg.Labels = new List<ILabel>();
+
+			Template = await templateCollection.GetOrAddAsync(new TemplateConfig { Name = "Test template" });
+			Graph = await graphCollection.AddAsync(Template, null);
+
+			TemplateRefId1 = new TemplateId("template1");
+			TemplateRefId2 = new TemplateId("template2");
+
+			List<TemplateRefConfig> templates = new List<TemplateRefConfig>();
+			templates.Add(new TemplateRefConfig { Id = TemplateRefId1, Name = "Test Template" });
+			templates.Add(new TemplateRefConfig { Id = TemplateRefId2, Name = "Test Template" });
+
+			List<TabConfig> tabs = new List<TabConfig>();
+			tabs.Add(new JobsTabConfig { Title = "foo", Templates = new List<TemplateId> { TemplateRefId1, TemplateRefId2 } });
+
+			Dictionary<string, AgentConfig> agentTypes = new()
+			{
+				{ "Win64", new() { Pool = new PoolId(PoolName) } }
+			};
+
+			StreamId streamId = new StreamId("ue5-main");
+			StreamConfig streamConfig = new StreamConfig { Id = streamId, Name = "//UE5/Main", Tabs = tabs, Templates = templates, AgentTypes = agentTypes };
+
+			ProjectId projectId = new ProjectId("ue5");
+			ProjectConfig projectConfig = new ProjectConfig { Id = projectId, Name = "UE5", Streams = new List<StreamConfig> { streamConfig } };
+
+			GlobalConfig globalConfig = new GlobalConfig { Projects = new List<ProjectConfig> { projectConfig } };
+			globalConfig.PostLoad(serverSettings);
+			configService.Set(IoHash.Zero, globalConfig);
+
+			StreamId = streamId;
+			StreamConfig = streamConfig;
+
+			Job1 = await jobService.CreateJobAsync(
+				jobId: JobId.Parse("5f283932841e7fdbcafb6ab5"),
+				streamConfig: streamConfig,
+				templateRefId: TemplateRefId1,
+				templateHash: Template.Hash,
+				graph: Graph,
+				name: "hello1",
+				change: 1000001,
+				codeChange: 1000002,
+				new CreateJobOptions { PreflightChange = 1001 }
+			);
+			Job1 = (await jobService.GetJobAsync(Job1.Id))!;
+
+			Job2 = await jobService.CreateJobAsync(
+				jobId: JobId.Parse("5f69ea1b68423e921b035106"),
+				streamConfig: streamConfig,
+				templateRefId: new TemplateId("template-id-1"),
+				templateHash: ContentHash.MD5("made-up-template-hash"),
+				graph: fg,
+				name: "hello2",
+				change: 2000001,
+				codeChange: 2000002,
+				new CreateJobOptions()
+			);
+			Job2 = (await jobService.GetJobAsync(Job2.Id))!;
+
+			Job1ArtifactData = "For The Horde!";
+			using MemoryStream job1ArtifactStream = new MemoryStream(Encoding.UTF8.GetBytes(Job1ArtifactData));
+			Job1Artifact = await artifactCollection.CreateArtifactAsync(Job1.Id, SubResourceId.Parse("22"), "myFile.txt",
+				"text/plain", job1ArtifactStream);
+
+			Agent1Name = "testAgent1";
+			Agent1 = await agentService.CreateAgentAsync(Agent1Name, true, null);
+		}
+
+		private class FixtureGraph : IGraph
+		{
+			public ContentHash Id { get; set; } = ContentHash.Empty;
+			public int Schema { get; set; }
+			public IReadOnlyList<INodeGroup> Groups { get; set; } = null!;
+			public IReadOnlyList<IAggregate> Aggregates { get; set; } = null!;
+			public IReadOnlyList<ILabel> Labels { get; set; } = null!;
+		}
+	}
+}
\ No newline at end of file
diff --git a/Engine/Source/Programs/Horde/Horde.Server.Tests/Fleet/JobQueueStrategyTest.cs b/Engine/Source/Programs/Horde/Horde.Server.Tests/Fleet/JobQueueStrategyTest.cs
new file mode 100644
--- /dev/null
+++ b/Engine/Source/Programs/Horde/Horde.Server.Tests/Fleet/JobQueueStrategyTest.cs
@@ -0,0 +1,160 @@
+// Copyright Epic Games, Inc. All Rights Reserved.
+
+using System;
+using System.Collections.Generic;
+using System.Threading.Tasks;
+using EpicGames.Core;
+using Horde.Server.Agents.Pools;
+using Horde.Server.Agents.Fleet;
+using Horde.Server.Projects;
+using Horde.Server.Streams;
+using Horde.Server.Jobs;
+using Horde.Server.Jobs.Graphs;
+using Horde.Server.Server;
+using Horde.Server.Utilities;
+using Microsoft.VisualStudio.TestTools.UnitTesting;
+using Horde.Server.Agents;
+
+namespace Horde.Server.Tests.Fleet
+{
+	[TestClass]
+	public class JobQueueStrategyTest : TestSetup
+	{
+		[TestMethod]
+		public async Task GetPoolQueueSizes()
+		{
+			(JobQueueStrategy strategy, PoolSizeResult poolSizeResult, IPool pool, List<IAgent> agents) = await SetUpJobsAsync(1, 5);
+			await Clock.AdvanceAsync(TimeSpan.FromSeconds(strategy.Settings.ReadyTimeThresholdSec) + TimeSpan.FromSeconds(5));
+			Dictionary<PoolId, int> poolQueueSizes = await strategy.GetPoolQueueSizesAsync(Clock.UtcNow - TimeSpan.FromHours(2));
+			Assert.AreEqual(1, poolQueueSizes.Count);
+			Assert.AreEqual(5, poolQueueSizes[pool.Id]);
+		}
+		
+		[TestMethod]
+		public async Task EmptyJobQueue()
+		{
+			await AssertAgentCount(0, -1, false);
+		}
+		
+		[TestMethod]
+		public async Task EmptyJobQueueWithLargePool()
+		{
+			await AssertAgentCount(0, -2, false, 20);
+		}
+		
+		[TestMethod]
+		public async Task NoAgentsInPool()
+		{
+			await AssertAgentCount(1, 1, true, 0);
+		}
+		
+		[TestMethod]
+		public async Task BatchesNotWaitingLongEnough()
+		{
+			await AssertAgentCount(3, -1, false);
+		}
+		
+		[TestMethod]
+		public async Task NumQueuedJobs1()
+		{
+			await AssertAgentCount(1, 1, true);
+		}
+		
+		[TestMethod]
+		public async Task NumQueuedJobs3()
+		{
+			await AssertAgentCount(3, 1, true);
+		}
+		
+		[TestMethod]
+		public async Task NumQueuedJobs6()
+		{
+			await AssertAgentCount(6, 2);
+		}
+		
+		[TestMethod]
+		public async Task NumQueuedJobs25()
+		{
+			await AssertAgentCount(25, 7);
+		}
+
+		public async Task AssertAgentCount(int numBatchesReady, int expectedAgentDelta, bool waitedBeyondThreshold = true, int numAgents = 8)
+		{
+			(JobQueueStrategy strategy, PoolSizeResult poolSizeResult, IPool pool, List<IAgent> agents) = await SetUpJobsAsync(1, numBatchesReady, numAgents);
+			TimeSpan timeToWait = waitedBeyondThreshold
+				? TimeSpan.FromSeconds(strategy.Settings.ReadyTimeThresholdSec) + TimeSpan.FromSeconds(5)
+				: TimeSpan.FromSeconds(15);
+			
+			await Clock.AdvanceAsync(timeToWait);
+
+			PoolSizeResult result = await strategy.CalculatePoolSizeAsync(pool, agents);
+			Assert.AreEqual(agents.Count + expectedAgentDelta, result.DesiredAgentCount);
+		}
+	
+		/// <summary>
+		/// Set up a fixture for job queue tests, ensuring a certain number of job batches are in running or waiting state
+		/// </summary>
+		/// <param name="numBatchesRunning">Num of job batches that should be in state running</param>
+		/// <param name="numBatchesReady">Num of job batches that should be in state waiting</param>
+		private async Task<(JobQueueStrategy, PoolSizeResult, IPool, List<IAgent> agents)> SetUpJobsAsync(int numBatchesRunning, int numBatchesReady, int numAgents = 8)
+		{
+			IPool pool = await PoolService.CreatePoolAsync("bogusPool1", null, true, 0, 0);
+			List<IAgent> agents = new();
+			for (int i = 0; i < numAgents; i++)
+			{
+				agents.Add(await CreateAgentAsync(pool));
+			}
+			
+			PoolSizeResult poolSize = new (agents.Count, agents.Count, null);
+			
+			string agentTypeName1 = "bogusAgentType1";
+			Dictionary<string, AgentConfig> agentTypes = new() { {agentTypeName1, new() { Pool = new PoolId(pool.Name) } }, };
+
+			StreamConfig streamConfig = new StreamConfig { Id = new StreamId("ue5"), Name = "//UE5/Main", AgentTypes = agentTypes };
+
+			ProjectConfig projectConfig = new ProjectConfig();
+			projectConfig.Id = new ProjectId("ue5");
+			projectConfig.Streams.Add(streamConfig);
+
+			GlobalConfig globalConfig = new GlobalConfig();
+			globalConfig.Projects.Add(projectConfig);
+
+			SetConfig(globalConfig);
+
+			string nodeForAgentType1 = "bogusNodeOnAgentType1";
+			IGraph graph = await GraphCollection.AppendAsync(null, new()
+			{
+				new NewGroup(agentTypeName1, new List<NewNode>
+				{
+					new (nodeForAgentType1),
+				})
+			});
+
+			for (int i = 0; i < numBatchesRunning; i++)
+			{
+				IJob job = await AddPlaceholderJob(graph, streamConfig.Id, nodeForAgentType1);
+				await JobCollection.TryUpdateBatchAsync(job, graph, job.Batches[0].Id, null, JobStepBatchState.Running, null);
+			}
+			
+			for (int i = 0; i < numBatchesReady; i++)
+			{
+				IJob job = await AddPlaceholderJob(graph, streamConfig.Id, nodeForAgentType1);
+				await JobCollection.TryUpdateBatchAsync(job, graph, job.Batches[0].Id, null, JobStepBatchState.Ready, null);
+			}
+
+			return (new (JobCollection, GraphCollection, StreamCollection, Clock, Cache, GlobalConfig), poolSize, pool, agents);
+		}
+		
+		private async Task<IJob> AddPlaceholderJob(IGraph graph, StreamId streamId, string nodeNameToExecute)
+		{
+			CreateJobOptions options = new CreateJobOptions();
+			options.Arguments.Add($"-Target={nodeNameToExecute}");
+
+			IJob job = await JobCollection.AddAsync(JobId.GenerateNewId(), streamId,
+				new TemplateId("bogusTemplateRefId"), ContentHash.Empty, graph, "bogusJobName",
+				1000, 1000, options);
+
+			return job;
+		}
+	}
+}
diff --git a/Engine/Source/Programs/Horde/Horde.Server.Tests/Fleet/LeaseUtilizationStrategyTest.cs b/Engine/Source/Programs/Horde/Horde.Server.Tests/Fleet/LeaseUtilizationStrategyTest.cs
new file mode 100644
--- /dev/null
+++ b/Engine/Source/Programs/Horde/Horde.Server.Tests/Fleet/LeaseUtilizationStrategyTest.cs
@@ -0,0 +1,117 @@
+// Copyright Epic Games, Inc. All Rights Reserved.
+
+using System;
+using System.Collections.Generic;
+using System.Threading.Tasks;
+using Google.Protobuf;
+using Google.Protobuf.WellKnownTypes;
+using Horde.Server.Agents;
+using Horde.Server.Agents.Fleet;
+using Horde.Server.Agents.Leases;
+using Horde.Server.Streams;
+using Horde.Server.Agents.Pools;
+using Horde.Server.Utilities;
+using HordeCommon;
+using HordeCommon.Rpc.Tasks;
+using Microsoft.VisualStudio.TestTools.UnitTesting;
+
+namespace Horde.Server.Tests.Fleet
+{
+	[TestClass]
+	public class LeaseUtilizationStrategyTest : TestSetup
+	{
+		private IAgent _agent1 = null!;
+		private IAgent _agent2 = null!;
+		private IAgent _agent3 = null!;
+		private IAgent _agent4 = null!;
+		private List<IAgent> _poolAgents = null!;
+
+		public void CreateAgents(IPool pool)
+		{
+			_agent1 = CreateAgentAsync(pool).Result;
+			_agent2 = CreateAgentAsync(pool).Result;
+			_agent3 = CreateAgentAsync(pool).Result;
+			_agent4 = CreateAgentAsync(pool).Result;
+			_poolAgents = new() { _agent1, _agent2, _agent3, _agent4 };
+		}
+
+		[TestMethod]
+		public async Task UtilizationFull()
+		{
+			IPool pool = await PoolService.CreatePoolAsync("AutoscalePool1", null, true, 0, 0);
+			CreateAgents(pool);
+			
+			await AddPlaceholderLease(_agent1, pool, Clock.UtcNow - TimeSpan.FromMinutes(120), TimeSpan.FromMinutes(120));
+			await AddPlaceholderLease(_agent2, pool, Clock.UtcNow - TimeSpan.FromMinutes(120), TimeSpan.FromMinutes(120));
+			await AddPlaceholderLease(_agent3, pool, Clock.UtcNow - TimeSpan.FromMinutes(120), TimeSpan.FromMinutes(120));
+			await AddPlaceholderLease(_agent4, pool, Clock.UtcNow - TimeSpan.FromMinutes(120), TimeSpan.FromMinutes(120));
+			await AssertPoolSizeAsync(pool, 4);
+		}
+		
+		[TestMethod]
+		public async Task UtilizationHalf()
+		{
+			IPool pool = await PoolService.CreatePoolAsync("AutoscalePool1", null, true, 0, 0);
+			CreateAgents(pool);
+			
+			await AddPlaceholderLease(_agent1, pool, Clock.UtcNow - TimeSpan.FromMinutes(120), TimeSpan.FromMinutes(120));
+			await AddPlaceholderLease(_agent2, pool, Clock.UtcNow - TimeSpan.FromMinutes(120), TimeSpan.FromMinutes(120));
+			await AssertPoolSizeAsync(pool, 2);
+		}
+		
+		[TestMethod]
+		public async Task UtilizationZero()
+		{
+			IPool pool = await PoolService.CreatePoolAsync("AutoscalePool1", null, true, 0, 0);
+			CreateAgents(pool);
+
+			await AssertPoolSizeAsync(pool, 0);
+		}
+		
+		[TestMethod]
+		public async Task ReserveAgents()
+		{
+			IPool pool = await PoolService.CreatePoolAsync("AutoscalePool1", null, true, 0, 5);
+			CreateAgents(pool);
+			
+			await AddPlaceholderLease(_agent1, pool, Clock.UtcNow - TimeSpan.FromMinutes(120), TimeSpan.FromMinutes(120));
+			await AddPlaceholderLease(_agent2, pool, Clock.UtcNow - TimeSpan.FromMinutes(120), TimeSpan.FromMinutes(120));
+			await AddPlaceholderLease(_agent3, pool, Clock.UtcNow - TimeSpan.FromMinutes(120), TimeSpan.FromMinutes(120));
+			await AddPlaceholderLease(_agent4, pool, Clock.UtcNow - TimeSpan.FromMinutes(120), TimeSpan.FromMinutes(120));
+			
+			// Full utilization should mean all agents plus the reserve agents
+			await AssertPoolSizeAsync(pool, 9);
+		}
+		
+		[TestMethod]
+		public async Task MinAgents()
+		{
+			IPool pool = await PoolService.CreatePoolAsync("AutoscalePool1", null, true, 2, 0);
+			CreateAgents(pool);
+			
+			// Even with no utilization, expect at least the min number of agents
+			await AssertPoolSizeAsync(pool, 2);
+		}
+
+		private async Task AssertPoolSizeAsync(IPool pool, int expectedNumAgents)
+		{
+			IPoolSizeStrategy strategy = FleetService.CreatePoolSizeStrategy(pool);
+			PoolSizeResult output = await strategy.CalculatePoolSizeAsync(pool, _poolAgents);
+			Assert.AreEqual(expectedNumAgents, output.DesiredAgentCount);
+		}
+		
+		private async Task<ILease> AddPlaceholderLease(IAgent agent, IPool pool, DateTime startTime, TimeSpan duration)
+		{
+			Assert.IsNotNull(agent.SessionId);
+			
+			ExecuteJobTask placeholderJobTask = new();
+			placeholderJobTask.JobName = "placeholderJobName";
+			byte[] payload = Any.Pack(placeholderJobTask).ToByteArray();
+
+			ILease lease = await LeaseCollection.AddAsync(LeaseId.GenerateNewId(), "placeholderLease", agent.Id, agent.SessionId!.Value, new StreamId("placeholderStream"), pool.Id, null, startTime, payload);
+			bool wasModified = await LeaseCollection.TrySetOutcomeAsync(lease.Id, startTime + duration, LeaseOutcome.Success, null);
+			Assert.IsTrue(wasModified);
+			return lease;
+		}
+	}
+}
\ No newline at end of file
diff --git a/Engine/Source/Programs/Horde/Horde.Server.Tests/Horde.Server.Tests.csproj b/Engine/Source/Programs/Horde/Horde.Server.Tests/Horde.Server.Tests.csproj
new file mode 100644
--- /dev/null
+++ b/Engine/Source/Programs/Horde/Horde.Server.Tests/Horde.Server.Tests.csproj
@@ -0,0 +1,60 @@
+<Project Sdk="Microsoft.NET.Sdk">
+
+  <PropertyGroup>
+    <TargetFramework>net6.0</TargetFramework>
+    <Nullable>enable</Nullable>
+    <IsPackable>false</IsPackable>
+    <NeutralLanguage>en-US</NeutralLanguage>
+  </PropertyGroup>
+
+  <ItemGroup>
+    <Compile Remove="Stubs\Collections\JobCollectionStub.cs" />
+    <Compile Remove="Stubs\Collections\StreamCollectionStub.cs" />
+  </ItemGroup>
+
+  <ItemGroup>
+    <None Include="Stubs\Collections\JobCollectionStub.cs" />
+    <None Include="Stubs\Collections\StreamCollectionStub.cs" />
+  </ItemGroup>
+
+  <ItemGroup>
+    <PackageReference Include="Microsoft.AspNetCore.Mvc.Testing" Version="6.0.10" />
+    <PackageReference Include="Microsoft.NET.Test.Sdk" Version="17.3.2" />
+    <PackageReference Include="Moq" Version="4.18.2" />
+    <PackageReference Include="MSTest.TestAdapter" Version="2.2.10" />
+    <PackageReference Include="MSTest.TestFramework" Version="2.2.10" />
+    <PackageReference Include="Microsoft.Extensions.Caching.Memory" Version="6.0.1" />
+    <PackageReference Include="coverlet.collector" Version="3.1.2">
+      <PrivateAssets>all</PrivateAssets>
+      <IncludeAssets>runtime; build; native; contentfiles; analyzers; buildtransitive</IncludeAssets>
+    </PackageReference>
+  </ItemGroup>
+
+  <ItemGroup>
+    <ProjectReference Include="..\..\Shared\EpicGames.Serialization\EpicGames.Serialization.csproj" />
+    <ProjectReference Include="..\Horde.Agent\Horde.Agent.csproj">
+      <Aliases>HordeAgent</Aliases>
+    </ProjectReference>
+    <ProjectReference Include="..\Horde.Server\Horde.Server.csproj" />
+  </ItemGroup>
+
+  <ItemGroup>
+    <Compile Update="Properties\Resources.Designer.cs">
+      <DesignTime>True</DesignTime>
+      <AutoGen>True</AutoGen>
+      <DependentUpon>Resources.resx</DependentUpon>
+    </Compile>
+  </ItemGroup>
+
+  <ItemGroup>
+    <EmbeddedResource Update="Properties\Resources.resx">
+      <Generator>ResXFileCodeGenerator</Generator>
+      <LastGenOutput>Resources.Designer.cs</LastGenOutput>
+    </EmbeddedResource>
+  </ItemGroup>
+
+  <ItemGroup>
+    <Folder Include="Resources\" />
+  </ItemGroup>
+
+</Project>
diff --git a/Engine/Source/Programs/Horde/Horde.Server.Tests/JobCollectionTests.cs b/Engine/Source/Programs/Horde/Horde.Server.Tests/JobCollectionTests.cs
new file mode 100644
--- /dev/null
+++ b/Engine/Source/Programs/Horde/Horde.Server.Tests/JobCollectionTests.cs
@@ -0,0 +1,345 @@
+// Copyright Epic Games, Inc. All Rights Reserved.
+
+using System;
+using System.Collections.Generic;
+using System.Linq;
+using System.Threading.Tasks;
+using EpicGames.Core;
+using Horde.Server.Jobs;
+using Horde.Server.Jobs.Graphs;
+using Horde.Server.Agents;
+using Horde.Server.Logs;
+using Horde.Server.Agents.Leases;
+using Horde.Server.Agents.Pools;
+using Horde.Server.Streams;
+using Horde.Server.Jobs.Templates;
+using Horde.Server.Utilities;
+using HordeCommon;
+using Microsoft.VisualStudio.TestTools.UnitTesting;
+using MongoDB.Bson;
+using Moq;
+using Horde.Server.Agents.Sessions;
+
+namespace Horde.Server.Tests
+{
+	[TestClass]
+	public class JobCollectionTests : TestSetup
+	{
+		static NewGroup AddGroup(List<NewGroup> groups)
+		{
+			NewGroup @group = new NewGroup("win64", new List<NewNode>());
+			groups.Add(@group);
+			return @group;
+		}
+
+		static NewNode AddNode(NewGroup @group, string name, string[]? inputDependencies, Action<NewNode>? action = null)
+		{
+			NewNode node = new NewNode(name, inputDependencies: inputDependencies?.ToList(), orderDependencies: inputDependencies?.ToList());
+			if (action != null)
+			{
+				action.Invoke(node);
+			}
+			@group.Nodes.Add(node);
+			return node;
+		}
+
+		async Task<IJob> StartBatch(IJob job, IGraph graph, int batchIdx)
+		{
+			Assert.AreEqual(JobStepBatchState.Ready, job.Batches[batchIdx].State);
+			job = Deref(await JobCollection.TryUpdateBatchAsync(job, graph, job.Batches[batchIdx].Id, null, JobStepBatchState.Running, null));
+			Assert.AreEqual(JobStepBatchState.Running, job.Batches[batchIdx].State);
+			return job;
+		}
+
+		async Task<IJob> RunStep(IJob job, IGraph graph, int batchIdx, int stepIdx, JobStepOutcome outcome)
+		{
+			Assert.AreEqual(JobStepState.Ready, job.Batches[batchIdx].Steps[stepIdx].State);
+			job = Deref(await JobCollection.TryUpdateStepAsync(job, graph, job.Batches[batchIdx].Id, job.Batches[batchIdx].Steps[stepIdx].Id, JobStepState.Running, JobStepOutcome.Success));
+			Assert.AreEqual(JobStepState.Running, job.Batches[batchIdx].Steps[stepIdx].State);
+			job = Deref(await JobCollection.TryUpdateStepAsync(job, graph, job.Batches[batchIdx].Id, job.Batches[batchIdx].Steps[stepIdx].Id, JobStepState.Completed, outcome));
+			Assert.AreEqual(JobStepState.Completed, job.Batches[batchIdx].Steps[stepIdx].State);
+			Assert.AreEqual(outcome, job.Batches[batchIdx].Steps[stepIdx].Outcome);
+			return job;
+		}
+
+		[TestMethod]
+		public async Task TestStates()
+		{
+			Mock<ITemplate> templateMock = new Mock<ITemplate>(MockBehavior.Strict);
+			templateMock.SetupGet(x => x.InitialAgentType).Returns((string?)null);
+
+			IGraph baseGraph = await GraphCollection.AddAsync(templateMock.Object, null);
+
+			CreateJobOptions options = new CreateJobOptions();
+			options.Arguments.Add("-Target=Publish Client");
+			options.Arguments.Add("-Target=Post-Publish Client");
+
+			IJob job = await JobCollection.AddAsync(JobId.GenerateNewId(), new StreamId("ue4-main"), new TemplateId("test-build"), ContentHash.SHA1("hello"), baseGraph, "Test job", 123, 123, options);
+
+			job = await StartBatch(job, baseGraph, 0);
+			job = await RunStep(job, baseGraph, 0, 0, JobStepOutcome.Success); // Setup Build
+
+			List<NewGroup> newGroups = new List<NewGroup>();
+
+			NewGroup initialGroup = AddGroup(newGroups);
+			AddNode(initialGroup, "Update Version Files", null);
+			AddNode(initialGroup, "Compile Editor", new[] { "Update Version Files" });
+
+			NewGroup compileGroup = AddGroup(newGroups);
+			AddNode(compileGroup, "Compile Client", new[] { "Update Version Files" });
+
+			NewGroup publishGroup = AddGroup(newGroups);
+			AddNode(publishGroup, "Cook Client", new[] { "Compile Editor" }, x => x.RunEarly = true);
+			AddNode(publishGroup, "Publish Client", new[] { "Compile Client", "Cook Client" });
+			AddNode(publishGroup, "Post-Publish Client", null, x => x.OrderDependencies = new List<string> { "Publish Client" });
+
+			IGraph graph = await GraphCollection.AppendAsync(baseGraph, newGroups, null, null);
+			job = Deref(await JobCollection.TryUpdateGraphAsync(job, graph));
+
+			job = await StartBatch(job, graph, 1);
+			job = await RunStep(job, graph, 1, 0, JobStepOutcome.Success); // Update Version Files
+			job = await RunStep(job, graph, 1, 1, JobStepOutcome.Success); // Compile Editor
+
+			job = await StartBatch(job, graph, 2);
+			job = await RunStep(job, graph, 2, 0, JobStepOutcome.Success); // Compile Client
+
+			job = await StartBatch(job, graph, 3);
+			job = await RunStep(job, graph, 3, 0, JobStepOutcome.Failure); // Cook Client
+			Assert.AreEqual(JobStepState.Skipped, job.Batches[3].Steps[1].State); // Publish Client
+			Assert.AreEqual(JobStepState.Skipped, job.Batches[3].Steps[2].State); // Post-Publish Client
+		}
+
+		[TestMethod]
+		public async Task TryAssignLeaseTest()
+		{
+			Fixture fixture = await CreateFixtureAsync();
+
+			SessionId sessionId1 = SessionId.GenerateNewId();
+			await JobCollection.TryAssignLeaseAsync(fixture.Job1, 0, new PoolId("foo"), fixture.Agent1.Id,
+				sessionId1, LeaseId.GenerateNewId(), LogId.GenerateNewId());
+			
+			SessionId sessionId2 = SessionId.GenerateNewId();
+			IJob job = (await JobCollection.GetAsync(fixture.Job1.Id))!;
+			await JobCollection.TryAssignLeaseAsync(job, 0, new PoolId("foo"), fixture.Agent1.Id,
+				sessionId2, LeaseId.GenerateNewId(), LogId.GenerateNewId());
+			
+			// Manually verify the log output
+		}
+
+		[TestMethod]
+		public Task LostLeaseTestWithDependency()
+		{
+			return LostLeaseTestInternal(true);
+		}
+
+		[TestMethod]
+		public Task LostLeaseTestWithoutDependency()
+		{
+			return LostLeaseTestInternal(false);
+		}
+
+		public async Task LostLeaseTestInternal(bool hasDependency)
+		{
+			Mock<ITemplate> templateMock = new Mock<ITemplate>(MockBehavior.Strict);
+			templateMock.SetupGet(x => x.InitialAgentType).Returns((string?)null);
+
+			IGraph baseGraph = await GraphCollection.AddAsync(templateMock.Object, null);
+
+			CreateJobOptions options = new CreateJobOptions();
+			options.Arguments.Add("-Target=Step 1");
+			options.Arguments.Add("-Target=Step 3");
+
+			IJob job = await JobCollection.AddAsync(JobId.GenerateNewId(), new StreamId("ue4-main"), new TemplateId("test-build"), ContentHash.SHA1("hello"), baseGraph, "Test job", 123, 123, options);
+
+			job = await StartBatch(job, baseGraph, 0);
+			job = await RunStep(job, baseGraph, 0, 0, JobStepOutcome.Success); // Setup Build
+
+			List<NewGroup> newGroups = new List<NewGroup>();
+
+			NewGroup initialGroup = AddGroup(newGroups);
+			AddNode(initialGroup, "Step 1", null);
+			AddNode(initialGroup, "Step 2", hasDependency? new[] { "Step 1" } : null);
+			AddNode(initialGroup, "Step 3", new[] { "Step 2" });
+
+			IGraph graph = await GraphCollection.AppendAsync(baseGraph, newGroups, null, null);
+			job = Deref(await JobCollection.TryUpdateGraphAsync(job, graph));
+
+			job = await StartBatch(job, graph, 1);
+			job = await RunStep(job, graph, 1, 0, JobStepOutcome.Success); // Step 1
+			job = await RunStep(job, graph, 1, 1, JobStepOutcome.Success); // Step 2
+
+			// Force an error executing the batch
+			job = Deref(await JobCollection.TryUpdateBatchAsync(job, graph, job.Batches[1].Id, null, JobStepBatchState.Complete, JobStepBatchError.Incomplete));
+
+			// Check that it restarted all three nodes
+			IJob newJob = (await JobCollection.GetAsync(job.Id))!;
+			Assert.AreEqual(3, newJob.Batches.Count);
+			Assert.AreEqual(1, newJob.Batches[2].GroupIdx);
+
+			if (hasDependency)
+			{
+				Assert.AreEqual(3, newJob.Batches[2].Steps.Count);
+
+				Assert.AreEqual(0, newJob.Batches[2].Steps[0].NodeIdx);
+				Assert.AreEqual(1, newJob.Batches[2].Steps[1].NodeIdx);
+				Assert.AreEqual(2, newJob.Batches[2].Steps[2].NodeIdx);
+			}
+			else
+			{
+				Assert.AreEqual(2, newJob.Batches[2].Steps.Count);
+
+				Assert.AreEqual(1, newJob.Batches[2].Steps[0].NodeIdx);
+				Assert.AreEqual(2, newJob.Batches[2].Steps[1].NodeIdx);
+			}
+		}
+
+		[TestMethod]
+		public async Task IncompleteBatchAsync()
+		{
+			Mock<ITemplate> templateMock = new Mock<ITemplate>(MockBehavior.Strict);
+			templateMock.SetupGet(x => x.InitialAgentType).Returns((string?)null);
+
+			IGraph baseGraph = await GraphCollection.AddAsync(templateMock.Object, null);
+
+			CreateJobOptions options = new CreateJobOptions();
+			options.Arguments.Add("-Target=Step 1");
+			options.Arguments.Add("-Target=Step 3");
+
+			IJob job = await JobCollection.AddAsync(JobId.GenerateNewId(), new StreamId("ue4-main"), new TemplateId("test-build"), ContentHash.SHA1("hello"), baseGraph, "Test job", 123, 123, options);
+			Assert.AreEqual(1, job.Batches.Count);
+
+			job = await StartBatch(job, baseGraph, 0);
+			job = Deref(await JobCollection.TryUpdateBatchAsync(job, baseGraph, job.Batches[0].Id, null, JobStepBatchState.Complete, JobStepBatchError.Incomplete));
+
+			job = (await JobCollection.GetAsync(job.Id))!;
+			Assert.AreEqual(2, job.Batches.Count);
+			Assert.AreEqual(JobStepBatchState.Complete, job.Batches[0].State);
+			Assert.AreEqual(0, job.Batches[0].Steps.Count);
+			Assert.AreEqual(JobStepBatchState.Ready, job.Batches[1].State);
+			Assert.AreEqual(1, job.Batches[1].Steps.Count);
+
+			job = Deref(await JobCollection.TryUpdateBatchAsync(job, baseGraph, job.Batches[1].Id, null, JobStepBatchState.Complete, JobStepBatchError.Incomplete));
+
+			job = (await JobCollection.GetAsync(job.Id))!;
+			Assert.AreEqual(3, job.Batches.Count);
+			Assert.AreEqual(JobStepBatchState.Complete, job.Batches[0].State);
+			Assert.AreEqual(0, job.Batches[0].Steps.Count);
+			Assert.AreEqual(JobStepBatchState.Complete, job.Batches[1].State);
+			Assert.AreEqual(0, job.Batches[1].Steps.Count);
+			Assert.AreEqual(JobStepBatchState.Ready, job.Batches[2].State);
+			Assert.AreEqual(1, job.Batches[2].Steps.Count);
+
+			job = Deref(await JobCollection.TryUpdateBatchAsync(job, baseGraph, job.Batches[2].Id, null, JobStepBatchState.Complete, JobStepBatchError.Incomplete));
+
+			job = (await JobCollection.GetAsync(job.Id))!;
+			Assert.AreEqual(3, job.Batches.Count);
+			Assert.AreEqual(JobStepBatchState.Complete, job.Batches[0].State);
+			Assert.AreEqual(0, job.Batches[0].Steps.Count);
+			Assert.AreEqual(JobStepBatchState.Complete, job.Batches[1].State);
+			Assert.AreEqual(0, job.Batches[1].Steps.Count);
+			Assert.AreEqual(JobStepBatchState.Complete, job.Batches[2].State);
+			Assert.AreEqual(1, job.Batches[2].Steps.Count);
+			Assert.AreEqual(JobStepState.Skipped, job.Batches[2].Steps[0].State);
+		}
+
+		[TestMethod]
+		public async Task IncompleteBatchRunningAsync()
+		{
+			// Same as IncompleteBatchAsync, but with steps moved to running state
+
+			Mock<ITemplate> templateMock = new Mock<ITemplate>(MockBehavior.Strict);
+			templateMock.SetupGet(x => x.InitialAgentType).Returns((string?)null);
+
+			IGraph baseGraph = await GraphCollection.AddAsync(templateMock.Object, null);
+
+			CreateJobOptions options = new CreateJobOptions();
+			options.Arguments.Add("-Target=Step 1");
+			options.Arguments.Add("-Target=Step 3");
+
+			IJob job = await JobCollection.AddAsync(JobId.GenerateNewId(), new StreamId("ue4-main"), new TemplateId("test-build"), ContentHash.SHA1("hello"), baseGraph, "Test job", 123, 123, options);
+			Assert.AreEqual(1, job.Batches.Count);
+
+			// First retry
+
+			job = await StartBatch(job, baseGraph, 0);
+			job = Deref(await JobCollection.TryUpdateStepAsync(job, baseGraph, job.Batches[0].Id, job.Batches[0].Steps[0].Id, newState: JobStepState.Running));
+			job = Deref(await JobCollection.TryUpdateBatchAsync(job, baseGraph, job.Batches[0].Id, null, JobStepBatchState.Complete, JobStepBatchError.Incomplete));
+
+			job = (await JobCollection.GetAsync(job.Id))!;
+			Assert.AreEqual(2, job.Batches.Count);
+
+			Assert.AreEqual(JobStepBatchState.Complete, job.Batches[0].State);
+			Assert.AreEqual(1, job.Batches[0].Steps.Count);
+			Assert.AreEqual(JobStepState.Completed, job.Batches[0].Steps[0].State);
+			Assert.AreEqual(JobStepError.Incomplete, job.Batches[0].Steps[0].Error);
+
+			Assert.AreEqual(JobStepBatchState.Ready, job.Batches[1].State);
+			Assert.AreEqual(1, job.Batches[1].Steps.Count);
+
+			// Second retry
+
+			job = Deref(await JobCollection.TryUpdateStepAsync(job, baseGraph, job.Batches[1].Id, job.Batches[1].Steps[0].Id, newState: JobStepState.Running));
+			job = Deref(await JobCollection.TryUpdateBatchAsync(job, baseGraph, job.Batches[1].Id, null, JobStepBatchState.Complete, JobStepBatchError.Incomplete));
+
+			job = (await JobCollection.GetAsync(job.Id))!;
+			Assert.AreEqual(3, job.Batches.Count);
+
+			Assert.AreEqual(JobStepBatchState.Complete, job.Batches[0].State);
+			Assert.AreEqual(1, job.Batches[0].Steps.Count);
+			Assert.AreEqual(JobStepState.Completed, job.Batches[0].Steps[0].State);
+			Assert.AreEqual(JobStepError.Incomplete, job.Batches[0].Steps[0].Error);
+
+			Assert.AreEqual(JobStepBatchState.Complete, job.Batches[1].State);
+			Assert.AreEqual(1, job.Batches[1].Steps.Count);
+			Assert.AreEqual(JobStepState.Completed, job.Batches[1].Steps[0].State);
+			Assert.AreEqual(JobStepError.Incomplete, job.Batches[1].Steps[0].Error);
+
+			Assert.AreEqual(JobStepBatchState.Ready, job.Batches[2].State);
+			Assert.AreEqual(1, job.Batches[2].Steps.Count);
+
+			// Check it doesn't retry a third time
+
+			job = Deref(await JobCollection.TryUpdateStepAsync(job, baseGraph, job.Batches[2].Id, job.Batches[2].Steps[0].Id, newState: JobStepState.Running));
+			job = Deref(await JobCollection.TryUpdateBatchAsync(job, baseGraph, job.Batches[2].Id, null, JobStepBatchState.Complete, JobStepBatchError.Incomplete));
+
+			job = (await JobCollection.GetAsync(job.Id))!;
+			Assert.AreEqual(3, job.Batches.Count);
+
+			Assert.AreEqual(JobStepBatchState.Complete, job.Batches[0].State);
+			Assert.AreEqual(1, job.Batches[0].Steps.Count);
+			Assert.AreEqual(JobStepState.Completed, job.Batches[0].Steps[0].State);
+			Assert.AreEqual(JobStepError.Incomplete, job.Batches[0].Steps[0].Error);
+
+			Assert.AreEqual(JobStepBatchState.Complete, job.Batches[1].State);
+			Assert.AreEqual(1, job.Batches[1].Steps.Count);
+			Assert.AreEqual(JobStepState.Completed, job.Batches[1].Steps[0].State);
+			Assert.AreEqual(JobStepError.Incomplete, job.Batches[1].Steps[0].Error);
+
+			Assert.AreEqual(JobStepBatchState.Complete, job.Batches[2].State);
+			Assert.AreEqual(1, job.Batches[2].Steps.Count);
+			Assert.AreEqual(JobStepState.Completed, job.Batches[2].Steps[0].State);
+			Assert.AreEqual(JobStepError.Incomplete, job.Batches[2].Steps[0].Error);
+		}
+
+		[TestMethod]
+		public async Task UnknownShelfAsync()
+		{
+			Mock<ITemplate> templateMock = new Mock<ITemplate>(MockBehavior.Strict);
+			templateMock.SetupGet(x => x.InitialAgentType).Returns((string?)null);
+
+			IGraph baseGraph = await GraphCollection.AddAsync(templateMock.Object, null);
+
+			CreateJobOptions options = new CreateJobOptions();
+			options.Arguments.Add("-Target=Step 1");
+			options.Arguments.Add("-Target=Step 3");
+
+			IJob job = await JobCollection.AddAsync(JobId.GenerateNewId(), new StreamId("ue4-main"), new TemplateId("test-build"), ContentHash.SHA1("hello"), baseGraph, "Test job", 123, 123, options);
+			Assert.AreEqual(1, job.Batches.Count);
+
+			job = await StartBatch(job, baseGraph, 0);
+			job = Deref(await JobCollection.TryUpdateBatchAsync(job, baseGraph, job.Batches[0].Id, null, JobStepBatchState.Complete, JobStepBatchError.UnknownShelf));
+
+			Assert.AreEqual(1, job.Batches.Count);
+		}
+	}
+}
diff --git a/Engine/Source/Programs/Horde/Horde.Server.Tests/JobServiceTests.cs b/Engine/Source/Programs/Horde/Horde.Server.Tests/JobServiceTests.cs
new file mode 100644
--- /dev/null
+++ b/Engine/Source/Programs/Horde/Horde.Server.Tests/JobServiceTests.cs
@@ -0,0 +1,244 @@
+// Copyright Epic Games, Inc. All Rights Reserved.
+
+using System;
+using System.Collections.Generic;
+using System.Runtime.InteropServices;
+using System.Text;
+using System.Threading;
+using System.Threading.Tasks;
+using EpicGames.Core;
+using EpicGames.Horde.Common;
+using Horde.Server.Agents;
+using Horde.Server.Agents.Pools;
+using Horde.Server.Jobs;
+using Horde.Server.Jobs.Graphs;
+using Horde.Server.Jobs.Templates;
+using Horde.Server.Logs;
+using Horde.Server.Perforce;
+using Horde.Server.Projects;
+using Horde.Server.Server;
+using Horde.Server.Streams;
+using Horde.Server.Users;
+using Horde.Server.Utilities;
+using HordeCommon;
+using Microsoft.CodeAnalysis;
+using Microsoft.VisualStudio.TestTools.UnitTesting;
+
+namespace Horde.Server.Tests
+{
+	using ProjectId = Projects.ProjectId;
+
+	[TestClass]
+	public class JobServiceTests : TestSetup
+	{
+		[TestMethod]
+		public async Task TestChainedJobs()
+		{
+			ProjectId projectId = new ProjectId("ue5");
+			StreamId streamId = new StreamId("ue5-main");
+			TemplateId templateRefId1 = new TemplateId("template1");
+			TemplateId templateRefId2 = new TemplateId("template2");
+
+			StreamConfig streamConfig = new StreamConfig { Id = streamId };
+			streamConfig.Templates.Add(new TemplateRefConfig { Id = templateRefId1, Name = "Test Template", ChainedJobs = new List<ChainedJobTemplateConfig> { new ChainedJobTemplateConfig { TemplateId = templateRefId2, Trigger = "Setup Build" } } });
+			streamConfig.Templates.Add(new TemplateRefConfig { Id = templateRefId2, Name = "Test Template" });
+			streamConfig.Tabs.Add(new JobsTabConfig { Title = "foo", Templates = new List<TemplateId> { templateRefId1, templateRefId2 } });
+
+			ProjectConfig projectConfig = new ProjectConfig { Id = projectId };
+			projectConfig.Streams.Add(streamConfig);
+
+			GlobalConfig globalConfig = new GlobalConfig();
+			globalConfig.Projects.Add(projectConfig);
+
+			SetConfig(globalConfig);
+
+			CreateJobOptions options = new CreateJobOptions();
+			options.PreflightChange = 999;
+			options.JobTriggers.AddRange(streamConfig.Templates[0].ChainedJobs!);
+
+			ITemplate template = await TemplateCollection.GetOrAddAsync(streamConfig.Templates[0]);
+
+			IGraph graph = await GraphCollection.AddAsync(template, null);
+
+			IJob job = await JobService.CreateJobAsync(null, streamConfig, templateRefId1, template.Hash, graph, "Hello", 1234, 1233, options);
+			Assert.AreEqual(1, job.ChainedJobs.Count);
+
+			job = Deref(await JobService.UpdateBatchAsync(job, job.Batches[0].Id, streamConfig, LogId.GenerateNewId(), JobStepBatchState.Running));
+			job = Deref(await JobService.UpdateStepAsync(job, job.Batches[0].Id, job.Batches[0].Steps[0].Id, streamConfig, JobStepState.Running));
+			job = Deref(await JobService.UpdateStepAsync(job, job.Batches[0].Id, job.Batches[0].Steps[0].Id, streamConfig, JobStepState.Completed, JobStepOutcome.Success));
+
+			Assert.IsNotNull(job.ChainedJobs[0].JobId);
+
+			IJob? chainedJob = await JobCollection.GetAsync(job.ChainedJobs[0].JobId!.Value);
+			Assert.IsNotNull(chainedJob);
+			Assert.AreEqual(chainedJob!.Id, job!.ChainedJobs[0].JobId);
+
+			Assert.AreEqual(chainedJob!.Change, job!.Change);
+			Assert.AreEqual(chainedJob!.CodeChange, job!.CodeChange);
+			Assert.AreEqual(chainedJob!.PreflightChange, job!.PreflightChange);
+			Assert.AreEqual(chainedJob!.StartedByUserId, job!.StartedByUserId);
+		}
+
+		[TestMethod]
+		public async Task ChangeQueryTestAsync()
+		{
+			StreamId streamId = new StreamId("ue5-main");
+			StreamConfig streamConfig = new StreamConfig { Id = streamId };
+
+			ProjectId projectId = new ProjectId("ue5");
+			ProjectConfig projectConfig = new ProjectConfig { Id = projectId, Name = "UE4", Streams = new List<StreamConfig> { streamConfig } };
+
+			SetConfig(new GlobalConfig { Projects = new List<ProjectConfig> { projectConfig } });
+
+			IUser user = await UserCollection.FindOrAddUserByLoginAsync("Bob");
+
+			PerforceService.AddChange(streamId, 1000, user, "", new[] { "Foo.cpp" });
+			PerforceService.AddChange(streamId, 1001, user, "", new[] { "Bar.cpp" });
+			PerforceService.AddChange(streamId, 1002, user, "", new[] { "Baz.cpp" });
+			PerforceService.AddChange(streamId, 1003, user, "", new[] { "Foo.uasset" });
+			PerforceService.AddChange(streamId, 1004, user, "", new[] { "Foo.uasset" });
+
+			ICommitCollection commits = PerforceService.GetCommits(streamConfig);
+
+			{
+				int? change = await JobService.EvaluateChangeQueryAsync(streamId, new ChangeQueryConfig { CommitTag = CommitTag.Code }, null, commits, CancellationToken.None);
+				Assert.AreEqual(1002, change);
+			}
+			{
+				int? change = await JobService.EvaluateChangeQueryAsync(streamId, new ChangeQueryConfig { CommitTag = CommitTag.Content }, null, commits, CancellationToken.None);
+				Assert.AreEqual(1004, change);
+			}
+			{
+				int? change = await JobService.EvaluateChangeQueryAsync(streamId, new ChangeQueryConfig { CommitTag = CommitTag.Content, Condition = "tag.code == 1" }, new List<CommitTag> { CommitTag.Code }, commits, CancellationToken.None);
+				Assert.AreEqual(1004, change);
+			}
+			{
+				int? change = await JobService.EvaluateChangeQueryAsync(streamId, new ChangeQueryConfig { CommitTag = CommitTag.Content, Condition = "tag.content == 1" }, new List<CommitTag> { CommitTag.Code }, commits, CancellationToken.None);
+				Assert.IsNull(change);
+			}
+		}
+
+		[TestMethod]
+		public async Task StopAnyDuplicateJobsByPreflight()
+		{
+			Fixture fixture = await CreateFixtureAsync();
+
+			string[] args = {"-Target=bogus"};
+			IJob orgJob = await CreatePreflightJob(fixture, "tpl-ref-1", "tpl-hash-1", "elvis", 1000, args);
+			IJob newJob = await CreatePreflightJob(fixture, "tpl-ref-1", "tpl-hash-1", "elvis", 1000, args);
+			IJob differentTplRef = await CreatePreflightJob(fixture, "tpl-ref-other", "tpl-hash-1", "elvis", 1000, args);
+			IJob differentTplHash = await CreatePreflightJob(fixture, "tpl-ref-1", "tpl-hash-other", "elvis", 1000, args);
+			IJob differentUserName = await CreatePreflightJob(fixture, "tpl-ref-1", "tpl-hash-1", "julia", 1000, args);
+			IJob differentArgs = await CreatePreflightJob(fixture, "tpl-ref-1", "tpl-hash-1", "elvis", 1000, new []{"-Target=other"});
+			
+			orgJob = (await JobService.GetJobAsync(orgJob.Id))!;
+			newJob = (await JobService.GetJobAsync(newJob.Id))!;
+			differentTplRef = (await JobService.GetJobAsync(differentTplRef.Id))!;
+			differentTplHash = (await JobService.GetJobAsync(differentTplHash.Id))!;
+			differentUserName = (await JobService.GetJobAsync(differentUserName.Id))!;
+			differentArgs = (await JobService.GetJobAsync(differentArgs.Id))!;
+			
+			Assert.AreEqual(KnownUsers.System, orgJob.AbortedByUserId);
+			Assert.IsNull(newJob.AbortedByUserId);
+			Assert.IsNull(differentTplRef.AbortedByUserId);
+			Assert.IsNull(differentTplHash.AbortedByUserId);
+			Assert.IsNull(differentUserName.AbortedByUserId);
+			Assert.IsNull(differentArgs.AbortedByUserId);
+		}
+
+		private async Task<IJob> CreatePreflightJob(Fixture fixture, string templateRefId, string templateHash, string startedByUserName, int preflightChange, string[] arguments)
+		{
+			IUser user = await UserCollection.FindOrAddUserByLoginAsync(startedByUserName);
+
+			CreateJobOptions options = new CreateJobOptions();
+			options.PreflightChange = preflightChange;
+			options.StartedByUserId = user.Id;
+			options.Arguments.AddRange(arguments);
+
+			return await JobService.CreateJobAsync(
+				jobId: JobId.GenerateNewId(),
+				streamConfig: fixture!.StreamConfig!,
+				templateRefId: new TemplateId(templateRefId),
+				templateHash: new ContentHash(Encoding.ASCII.GetBytes(templateHash)),
+				graph: fixture!.Graph,
+				name: "hello1",
+				change: 1000001,
+				codeChange: 1000002,
+				options
+			);
+		}
+		
+		// Only test for cancelled preflights
+		// [TestMethod]
+		// public async Task StopAnyDuplicateJobsByChange()
+		// {
+		// 	TestSetup TestSetup = await GetTestSetup();
+		// 	Fixture Fix = Fixture!;
+		// 	IJob Job1 = Fix.Job1;
+		// 	IJob Job2 = Fix.Job2;
+		// 	
+		// 	Assert.AreEqual(JobState.Waiting, Job1.GetState());
+		// 	Assert.IsTrue(await JobService.UpdateBatchAsync(Job1, Job1.Batches[0].Id, ObjectId.GenerateNewId(), JobStepBatchState.Running));
+		//
+		// 	Job1 = (await JobService.GetJobAsync(Job1.Id))!;
+		// 	Job2 = (await JobService.GetJobAsync(Job2.Id))!;
+		// 	Assert.AreEqual(JobState.Running, Job1.GetState());
+		// 	Assert.AreEqual(JobState.Complete, Job2.GetState());
+		// 	
+		// 	IJob Job3 = await JobService.CreateJobAsync(null, Fix.Stream!.Id, Fix.TemplateRefId1, Fix.Template.Id, Fix.Graph, "Hello", Job1.Change, Job1.CodeChange, null, "joe", null, null, true, true, null, null, new List<string>());
+		// 	await DispatchService.TickOnlyForTestingAsync();
+		// 	Job1 = (await JobService.GetJobAsync(Job1.Id))!;
+		// 	Assert.AreEqual(JobState.Complete, Job1.GetState());
+		// 	Assert.AreEqual(JobState.Waiting, Job3.GetState());
+		// }
+
+		[TestMethod]
+		public async Task TestRunEarly()
+		{
+			StreamId streamId = new StreamId("ue5-main");
+			StreamConfig streamConfig = new StreamConfig { Id = streamId };
+
+			ProjectId projectId = new ProjectId("ue5");
+			ProjectConfig projectConfig = new ProjectConfig { Id = projectId, Name = "UE5", Streams = new List<StreamConfig> { streamConfig } };
+
+			SetConfig(new GlobalConfig { Projects = new List<ProjectConfig> { projectConfig } });
+
+			// ----
+
+			IAgent? agent = await AgentService.CreateAgentAsync("TestAgent", true, new List<PoolId> { new PoolId("win") });
+			await AgentService.CreateSessionAsync(agent, AgentStatus.Ok, new List<string>(), new Dictionary<string, int>(), null);
+
+			ITemplate template = await TemplateCollection.GetOrAddAsync(new TemplateConfig { Name = "Test template" });
+			IGraph graph = await GraphCollection.AddAsync(template, null);
+
+			NewGroup groupA = new NewGroup("win", new List<NewNode>());
+			groupA.Nodes.Add(new NewNode("Compile"));
+
+			NewGroup groupB = new NewGroup("win", new List<NewNode>());
+			groupB.Nodes.Add(new NewNode("Cook", runEarly: true));
+			groupB.Nodes.Add(new NewNode("Middle"));
+			groupB.Nodes.Add(new NewNode("Pak", inputDependencies: new List<string> { "Compile", "Cook" }));
+
+			graph = await GraphCollection.AppendAsync(graph, new List<NewGroup> { groupA, groupB });
+
+			CreateJobOptions options = new CreateJobOptions();
+			options.PreflightChange = 999;
+			options.Arguments.Add("-Target=Pak");
+
+			IJob job = await JobService.CreateJobAsync(null, streamConfig!, new TemplateId("temp"), new ContentHash(new byte[] { 1, 2, 3 }), graph, "Hello", 1234, 1233, options);
+
+			job = Deref(await JobService.UpdateBatchAsync(job, job.Batches[0].Id, streamConfig, LogId.GenerateNewId(), JobStepBatchState.Running));
+			job = Deref(await JobService.UpdateStepAsync(job, job.Batches[0].Id, job.Batches[0].Steps[0].Id, streamConfig, JobStepState.Running));
+			job = Deref(await JobService.UpdateStepAsync(job, job.Batches[0].Id, job.Batches[0].Steps[0].Id, streamConfig, JobStepState.Completed, JobStepOutcome.Success));
+
+			job = Deref(await JobService.UpdateBatchAsync(job, job.Batches[1].Id, streamConfig, LogId.GenerateNewId(), JobStepBatchState.Running));
+			job = Deref(await JobService.UpdateStepAsync(job, job.Batches[1].Id, job.Batches[1].Steps[0].Id, streamConfig, JobStepState.Running));
+
+			job = Deref(await JobService.UpdateBatchAsync(job, job.Batches[2].Id, streamConfig, LogId.GenerateNewId(), JobStepBatchState.Running));
+			job = Deref(await JobService.UpdateStepAsync(job, job.Batches[2].Id, job.Batches[2].Steps[0].Id, streamConfig, JobStepState.Running));
+			job = Deref(await JobService.UpdateStepAsync(job, job.Batches[2].Id, job.Batches[2].Steps[0].Id, streamConfig, JobStepState.Completed, JobStepOutcome.Success));
+
+			Assert.AreEqual(JobStepState.Waiting, job.Batches[2].Steps[1].State);
+		}
+	}
+}
diff --git a/Engine/Source/Programs/Horde/Horde.Server.Tests/JobTaskSourceTests.cs b/Engine/Source/Programs/Horde/Horde.Server.Tests/JobTaskSourceTests.cs
new file mode 100644
--- /dev/null
+++ b/Engine/Source/Programs/Horde/Horde.Server.Tests/JobTaskSourceTests.cs
@@ -0,0 +1,180 @@
+// Copyright Epic Games, Inc. All Rights Reserved.
+
+using System.Collections.Generic;
+using System.Threading;
+using System.Threading.Tasks;
+using Horde.Server.Jobs;
+using Horde.Server.Jobs.Graphs;
+using Horde.Server.Users;
+using Horde.Server.Utilities;
+using HordeCommon;
+using Microsoft.VisualStudio.TestTools.UnitTesting;
+using System.Linq;
+using System;
+using Horde.Server.Agents;
+using Horde.Server.Agents.Pools;
+using Horde.Server.Streams;
+
+namespace Horde.Server.Tests
+{
+	[TestClass]
+	public class JobTaskSourceTests : TestSetup
+	{
+		private bool _eventReceived;
+		private bool? _eventPoolHasAgentsOnline;
+
+		static NewGroup AddGroup(List<NewGroup> groups)
+		{
+			NewGroup @group = new NewGroup("Win64", new List<NewNode>());
+			groups.Add(@group);
+			return @group;
+		}
+
+		static NewNode AddNode(NewGroup @group, string name, string[]? inputDependencies, Action<NewNode>? action = null)
+		{
+			NewNode node = new NewNode(name, inputDependencies: inputDependencies?.ToList(), orderDependencies: inputDependencies?.ToList());
+			if (action != null)
+			{
+				action.Invoke(node);
+			}
+			@group.Nodes.Add(node);
+			return node;
+		}
+
+		[TestMethod]
+		public async Task UpdateJobQueueNormal()
+		{
+			Fixture fixture = await SetupPoolWithAgentAsync(isPoolAutoScaled: true, shouldCreateAgent: true, isAgentEnabled: true);
+
+			await JobTaskSource.TickAsync(CancellationToken.None);
+			Assert.AreEqual(1, JobTaskSource.GetQueueForTesting().Count);
+			Assert.AreEqual(fixture.Job1.Id, JobTaskSource.GetQueueForTesting().Min!.Id.Item1);
+			Assert.AreEqual(JobStepBatchState.Ready, JobTaskSource.GetQueueForTesting().Min!.Batch.State);
+			
+			Assert.IsTrue(_eventReceived);
+			Assert.IsTrue(_eventPoolHasAgentsOnline!.Value);
+		}
+		
+		[TestMethod]
+		public async Task UpdateJobQueueWithNoAgentsInPool()
+		{
+			Fixture fixture = await SetupPoolWithAgentAsync(isPoolAutoScaled: true, shouldCreateAgent: false, isAgentEnabled: false);
+			
+			await JobTaskSource.TickAsync(CancellationToken.None);
+			Assert.AreEqual(0, JobTaskSource.GetQueueForTesting().Count);
+
+			IJob job = (await JobService.GetJobAsync(fixture.Job1.Id))!;
+			Assert.AreEqual(JobStepBatchError.NoAgentsInPool, job.Batches[0].Error);
+			
+			Assert.IsFalse(_eventReceived);
+		}
+		
+		[TestMethod]
+		public async Task UpdateJobQueueWithNoAgentsOnlineInPool()
+		{
+			Fixture fixture = await SetupPoolWithAgentAsync(isPoolAutoScaled: false, shouldCreateAgent: true, isAgentEnabled: false);
+			
+			await JobTaskSource.TickAsync(CancellationToken.None);
+			Assert.AreEqual(1, JobTaskSource.GetQueueForTesting().Count);
+
+			Assert.AreEqual(fixture.Job1.Id, JobTaskSource.GetQueueForTesting().Min!.Id.Item1);
+			Assert.AreEqual(JobStepBatchState.Ready, JobTaskSource.GetQueueForTesting().Min!.Batch.State);
+
+			Assert.IsTrue(_eventReceived);
+			Assert.IsFalse(_eventPoolHasAgentsOnline!.Value);
+		}
+		
+		[TestMethod]
+		public async Task UpdateJobQueueWithNoAgentsOnlineInAutoScaledPool()
+		{
+			Fixture fixture = await SetupPoolWithAgentAsync(isPoolAutoScaled: true, shouldCreateAgent: true, isAgentEnabled: false);
+			
+			await JobTaskSource.TickAsync(CancellationToken.None);
+			Assert.AreEqual(1, JobTaskSource.GetQueueForTesting().Count);
+
+			Assert.AreEqual(fixture.Job1.Id, JobTaskSource.GetQueueForTesting().Min!.Id.Item1);
+			Assert.AreEqual(JobStepBatchState.Ready, JobTaskSource.GetQueueForTesting().Min!.Batch.State);
+
+			Assert.IsTrue(_eventReceived);
+			Assert.IsFalse(_eventPoolHasAgentsOnline!.Value);
+		}
+
+		[TestMethod]
+		public async Task UpdateJobQueueWithPausedStep()
+		{
+			Fixture fixture = await SetupPoolWithAgentAsync(isPoolAutoScaled: true, shouldCreateAgent: true, isAgentEnabled: true);
+
+			// update template with some step states
+			IStream Stream = await StreamCollection.GetAsync(fixture.StreamConfig!);
+			Stream = Deref(await StreamCollection.TryUpdateTemplateRefAsync(Stream, fixture.TemplateRefId1, new List<UpdateStepStateRequest>() { new UpdateStepStateRequest() { Name = "Paused Step", PausedByUserId = UserId.GenerateNewId().ToString() } }));
+
+			// create a new graph with the associated nodes
+			List<NewGroup> newGroups = new List<NewGroup>();
+
+			NewGroup initialGroup = AddGroup(newGroups);			
+			AddNode(initialGroup, "Update Version Files", null);
+			AddNode(initialGroup, "Paused Step", new[] { "Update Version Files" });
+			AddNode(initialGroup, "Step That Depends on Paused Step", new[] { "Paused Step" });
+			AddNode(initialGroup, "Step That Depends on Update Version Files", new[] { "Update Version Files" });
+
+			IGraph graph = await GraphCollection.AppendAsync(null, newGroups);
+
+			// remove the default fixture jobs
+			List<IJob> Jobs = await JobCollection.FindAsync();
+			for (int i = 0; i < Jobs.Count; i++)
+			{
+				await JobCollection.RemoveAsync(Jobs[i]);
+			}
+
+			// create a new job
+			CreateJobOptions options = new CreateJobOptions();
+			options.Arguments.Add("-Target=Step That Depends on Paused Step;Step That Depends on Update Version Files");
+
+			IJob job = await JobCollection.AddAsync(JobId.GenerateNewId(), Stream.Id,
+				fixture.TemplateRefId1, fixture.Template.Hash, graph, "Test Paused Step Job",
+				1000, 1000, options);
+
+			// validate
+			await JobTaskSource.TickAsync(CancellationToken.None);
+			Assert.AreEqual(1, JobTaskSource.GetQueueForTesting().Count);
+			Assert.AreEqual(job.Id, JobTaskSource.GetQueueForTesting().Min!.Id.Item1);
+			Assert.AreEqual(JobStepBatchState.Ready, JobTaskSource.GetQueueForTesting().Min!.Batch.State);
+
+			Assert.IsTrue(_eventReceived);
+			Assert.IsTrue(_eventPoolHasAgentsOnline!.Value);
+
+			// make sure we get db value
+			job = Deref(await JobService.GetJobAsync(job.Id));
+
+			IJobStepBatch batch = job.Batches[0];
+
+			Assert.AreEqual(batch.Steps.Count, 4);
+			Assert.AreEqual(batch.Steps[0].State, JobStepState.Ready);
+			Assert.AreEqual(batch.Steps[1].State, JobStepState.Skipped);
+			Assert.AreEqual(batch.Steps[1].Error, JobStepError.Paused);
+			Assert.AreEqual(batch.Steps[2].State, JobStepState.Skipped);
+			Assert.AreEqual(batch.Steps[3].State, JobStepState.Waiting);
+
+		}
+
+		private async Task<Fixture> SetupPoolWithAgentAsync(bool isPoolAutoScaled, bool shouldCreateAgent, bool isAgentEnabled)
+		{
+			Fixture fixture = await CreateFixtureAsync();
+			IPool pool = await PoolService.CreatePoolAsync(Fixture.PoolName, null, isPoolAutoScaled, 0, 0);
+
+			if (shouldCreateAgent)
+			{
+				IAgent? agent = await AgentService.CreateAgentAsync("TestAgent", isAgentEnabled, new List<PoolId> { pool.Id });
+				await AgentService.CreateSessionAsync(agent, AgentStatus.Ok, new List<string>(), new Dictionary<string, int>(), null);
+			}
+			
+			JobTaskSource.OnJobScheduled += (pool, poolHasAgentsOnline, job, graph, batchId) =>
+			{
+				_eventReceived = true;
+				_eventPoolHasAgentsOnline = poolHasAgentsOnline;
+			};
+
+			return fixture;
+		}
+	}
+}
diff --git a/Engine/Source/Programs/Horde/Horde.Server.Tests/JobsControllerTest.cs b/Engine/Source/Programs/Horde/Horde.Server.Tests/JobsControllerTest.cs
new file mode 100644
--- /dev/null
+++ b/Engine/Source/Programs/Horde/Horde.Server.Tests/JobsControllerTest.cs
@@ -0,0 +1,35 @@
+﻿// Copyright Epic Games, Inc. All Rights Reserved.
+
+using System;
+using System.Net;
+using System.Threading.Tasks;
+using Horde.Server.Jobs.Artifacts;
+using Microsoft.VisualStudio.TestTools.UnitTesting;
+
+namespace Horde.Server.Tests
+{
+	[TestClass]
+    public class JobsControllerTest : ControllerIntegrationTest
+    {
+        private static Uri GetUri(string jobId, string stepId, string fileName)
+        {
+            return new Uri($"http://localhost/api/v1/jobs/{jobId}/steps/{stepId}/artifacts/{fileName}/data");
+        }
+
+        [TestMethod]
+        public async Task GetArtifactDataByFilenameTest()
+        {
+			Fixture? fixture = await GetFixture();
+			IArtifactV1? art = fixture.Job1Artifact;
+
+			// Test existing filename
+			System.Net.Http.HttpResponseMessage? res = await Client.GetAsync(GetUri(art.JobId.ToString(), art.StepId.ToString()!, art.Name));
+            res.EnsureSuccessStatusCode();
+            Assert.AreEqual(fixture.Job1ArtifactData, await res.Content.ReadAsStringAsync());
+
+            // Test non-existing filename
+            res = await Client.GetAsync(GetUri(art.JobId.ToString(), art.StepId.ToString()!, "bogus.txt"));
+            Assert.AreEqual(res.StatusCode, HttpStatusCode.NotFound);
+        }
+    }
+}
\ No newline at end of file
diff --git a/Engine/Source/Programs/Horde/Horde.Server.Tests/LeasesControllerDbTest.cs b/Engine/Source/Programs/Horde/Horde.Server.Tests/LeasesControllerDbTest.cs
new file mode 100644
--- /dev/null
+++ b/Engine/Source/Programs/Horde/Horde.Server.Tests/LeasesControllerDbTest.cs
@@ -0,0 +1,47 @@
+// Copyright Epic Games, Inc. All Rights Reserved.
+
+using System;
+using System.Collections.Generic;
+using System.Threading.Tasks;
+using Horde.Server.Agents;
+using Horde.Server.Agents.Leases;
+using Horde.Server.Agents.Sessions;
+using Horde.Server.Utilities;
+using HordeCommon;
+using Microsoft.AspNetCore.Mvc;
+using Microsoft.VisualStudio.TestTools.UnitTesting;
+
+namespace Horde.Server.Tests
+{
+	/// <summary>
+	/// Database-only integration test for LeasesController
+	/// </summary>
+	[TestClass]
+	public class LeasesControllerDbTest : TestSetup
+	{
+		[TestMethod]
+		public async Task FindLeases()
+		{
+			DateTimeOffset minTime = Clock.UtcNow - TimeSpan.FromMinutes(5);
+			DateTimeOffset maxTime = Clock.UtcNow;
+
+			ILease lease1 = await CreateLease(Clock.UtcNow - TimeSpan.FromMinutes(10), TimeSpan.FromMinutes(6));
+			ILease lease2 = await CreateLease(Clock.UtcNow - TimeSpan.FromMinutes(7), TimeSpan.FromMinutes(3.1));
+			/*ILease outOfTimeWindow = */await CreateLease(Clock.UtcNow - TimeSpan.FromMinutes(7), TimeSpan.FromMinutes(25));
+			
+			ActionResult<List<object>> res = await LeasesController.FindLeasesAsync(null, null, null, null, minTime, maxTime);
+			Assert.AreEqual(2, res.Value!.Count);
+			Assert.AreEqual(lease2.Id.ToString(), (res.Value[0] as GetAgentLeaseResponse)!.Id);
+			Assert.AreEqual(lease1.Id.ToString(), (res.Value[1] as GetAgentLeaseResponse)!.Id);
+		}
+
+		private async Task<ILease> CreateLease(DateTime startTime, TimeSpan duration)
+		{
+			LeaseId id = LeaseId.GenerateNewId();
+			SessionId sessionId = SessionId.GenerateNewId();
+			ILease lease = await LeaseCollection.AddAsync(id, "myLease", new AgentId("agent-1"), sessionId, null, null, null, startTime, Array.Empty<byte>());
+			await LeaseCollection.TrySetOutcomeAsync(id, startTime + duration, LeaseOutcome.Success, null);
+			return lease;
+		}
+	}
+}
\ No newline at end of file
diff --git a/Engine/Source/Programs/Horde/Horde.Server.Tests/LogBuilderTests.cs b/Engine/Source/Programs/Horde/Horde.Server.Tests/LogBuilderTests.cs
new file mode 100644
--- /dev/null
+++ b/Engine/Source/Programs/Horde/Horde.Server.Tests/LogBuilderTests.cs
@@ -0,0 +1,99 @@
+// Copyright Epic Games, Inc. All Rights Reserved.
+
+using System;
+using System.Collections.Generic;
+using System.Text;
+using System.Threading.Tasks;
+using Horde.Server.Logs;
+using Horde.Server.Logs.Builder;
+using Horde.Server.Logs.Data;
+using Horde.Server.Server;
+using Horde.Server.Utilities;
+using Microsoft.Extensions.Logging.Abstractions;
+using Microsoft.VisualStudio.TestTools.UnitTesting;
+
+namespace Horde.Server.Tests
+{
+	[TestClass]
+	public class LogBuilderTests : DatabaseIntegrationTest
+	{
+		[TestMethod]
+		public async Task TestLocalLogBuilder()
+		{
+			ILogBuilder builder = new LocalLogBuilder();
+			await TestBuilder(builder);
+		}
+
+		[TestMethod]
+		public async Task TestRedisLogBuilder()
+		{
+			RedisService redisService = GetRedisServiceSingleton();
+			ILogBuilder builder = new RedisLogBuilder(redisService.ConnectionPool, NullLogger.Instance);
+			await TestBuilder(builder);
+		}
+
+		public static async Task TestBuilder(ILogBuilder builder)
+		{
+			LogId logId = LogId.GenerateNewId();
+
+			const long Offset = 100;
+			Assert.IsTrue(await builder.AppendAsync(logId, Offset, Offset, 0, 1, Encoding.UTF8.GetBytes("hello\n"), LogType.Text));
+
+			LogChunkData? chunk1 = await builder.GetChunkAsync(logId, Offset, 0);
+			Assert.IsNotNull(chunk1);
+			Assert.AreEqual(1, chunk1!.LineCount);
+			Assert.AreEqual(6, chunk1!.Length);
+			Assert.AreEqual(1, chunk1!.SubChunks.Count);
+			Assert.AreEqual(1, chunk1!.SubChunks[0].LineCount);
+			Assert.AreEqual(6, chunk1!.SubChunks[0].Length);
+
+			Assert.IsTrue(await builder.AppendAsync(logId, Offset, Offset + 6, 1, 1, Encoding.UTF8.GetBytes("world\n"), LogType.Text));
+
+			LogChunkData? chunk2 = await builder.GetChunkAsync(logId, Offset, 1);
+			Assert.IsNotNull(chunk2);
+			Assert.AreEqual(2, chunk2!.LineCount);
+			Assert.AreEqual(12, chunk2!.Length);
+			Assert.AreEqual(1, chunk2!.SubChunks.Count);
+			Assert.AreEqual(2, chunk2!.SubChunks[0].LineCount);
+			Assert.AreEqual(12, chunk2!.SubChunks[0].Length);
+
+			await builder.CompleteSubChunkAsync(logId, Offset);
+			Assert.IsTrue(await builder.AppendAsync(logId, Offset, Offset + 12, 2, 1, Encoding.UTF8.GetBytes("foo\n"), LogType.Text));
+			Assert.IsTrue(await builder.AppendAsync(logId, Offset, Offset + 16, 3, 2, Encoding.UTF8.GetBytes("bar\nbaz\n"), LogType.Text));
+			await builder.CompleteSubChunkAsync(logId, Offset);
+
+			LogChunkData? chunk3 = await builder.GetChunkAsync(logId, Offset, 4);
+			Assert.IsNotNull(chunk3);
+			Assert.AreEqual(5, chunk3!.LineCount);
+			Assert.AreEqual(24, chunk3!.Length);
+			Assert.AreEqual(2, chunk3!.SubChunks.Count);
+			Assert.AreEqual(2, chunk3!.SubChunks[0].LineCount);
+			Assert.AreEqual(12, chunk3!.SubChunks[0].Length);
+			Assert.AreEqual(3, chunk3!.SubChunks[1].LineCount);
+			Assert.AreEqual(12, chunk3!.SubChunks[1].Length);
+
+			Assert.AreEqual(0, chunk3!.GetLineOffsetWithinChunk(0));
+			Assert.AreEqual(6, chunk3!.GetLineOffsetWithinChunk(1));
+			Assert.AreEqual(12, chunk3!.GetLineOffsetWithinChunk(2));
+			Assert.AreEqual(16, chunk3!.GetLineOffsetWithinChunk(3));
+			Assert.AreEqual(20, chunk3!.GetLineOffsetWithinChunk(4));
+
+			List<(LogId, long)> chunks1 = await builder.TouchChunksAsync(TimeSpan.Zero);
+			Assert.AreEqual(1, chunks1.Count);
+			Assert.AreEqual((logId, Offset), chunks1[0]);
+
+			await builder.CompleteChunkAsync(logId, Offset);
+
+			LogChunkData? chunk4 = await builder.GetChunkAsync(logId, Offset, 5);
+			Assert.AreEqual(24, chunk4!.Length);
+
+			List<(LogId, long)> chunks2 = await builder.TouchChunksAsync(TimeSpan.Zero);
+			Assert.AreEqual(1, chunks2.Count);
+
+			await builder.RemoveChunkAsync(logId, Offset);
+
+			List<(LogId, long)> chunks3 = await builder.TouchChunksAsync(TimeSpan.Zero);
+			Assert.AreEqual(0, chunks3.Count);
+		}
+	}
+}
diff --git a/Engine/Source/Programs/Horde/Horde.Server.Tests/LogChunkTests.cs b/Engine/Source/Programs/Horde/Horde.Server.Tests/LogChunkTests.cs
new file mode 100644
--- /dev/null
+++ b/Engine/Source/Programs/Horde/Horde.Server.Tests/LogChunkTests.cs
@@ -0,0 +1,48 @@
+// Copyright Epic Games, Inc. All Rights Reserved.
+
+using System;
+using System.Collections.Generic;
+using System.Text;
+using Horde.Server.Logs;
+using Horde.Server.Logs.Data;
+using Microsoft.Extensions.Logging.Abstractions;
+using Microsoft.VisualStudio.TestTools.UnitTesting;
+
+namespace Horde.Server.Tests
+{
+	[TestClass]
+	public class LogChunkTests
+	{
+		readonly string[] _text =
+		{
+			@"{ ""id"": 1, ""message"": ""foo"" }",
+			@"{ ""id"": 2, ""message"": ""bar"" }",
+			@"{ ""id"": 3, ""message"": ""baz"" }",
+		};
+
+		[TestMethod]
+		public void TestSerialization()
+		{
+			byte[] textData = Encoding.UTF8.GetBytes(String.Join("\n", _text) + "\n");
+
+			List<LogSubChunkData> subChunks = new List<LogSubChunkData>();
+			subChunks.Add(new LogSubChunkData(LogType.Text, 0, 0, new LogText(textData, textData.Length)));
+			subChunks.Add(new LogSubChunkData(LogType.Json, textData.Length, 3, new LogText(textData, textData.Length)));
+
+			LogChunkData oldChunkData = new LogChunkData(0, 0, subChunks);
+			byte[] data = oldChunkData.ToByteArray(NullLogger.Instance);
+			LogChunkData newChunkData = LogChunkData.FromMemory(data, 0, 0);
+
+			Assert.AreEqual(oldChunkData.Length, newChunkData.Length);
+			Assert.AreEqual(oldChunkData.SubChunks.Count, newChunkData.SubChunks.Count);
+			Assert.IsTrue(oldChunkData.SubChunkOffset.AsSpan(0, oldChunkData.SubChunkOffset.Length).SequenceEqual(newChunkData.SubChunkOffset.AsSpan(0, newChunkData.SubChunkOffset.Length)));
+			Assert.IsTrue(oldChunkData.SubChunkLineIndex.AsSpan(0, oldChunkData.SubChunkLineIndex.Length).SequenceEqual(newChunkData.SubChunkLineIndex.AsSpan(0, newChunkData.SubChunkLineIndex.Length)));
+			for (int idx = 0; idx < oldChunkData.SubChunks.Count; idx++)
+			{
+				LogSubChunkData oldSubChunkData = oldChunkData.SubChunks[idx];
+				LogSubChunkData newSubChunkData = newChunkData.SubChunks[idx];
+				Assert.IsTrue(oldSubChunkData.InflateText().Data.Span.SequenceEqual(newSubChunkData.InflateText().Data.Span));
+			}
+		}
+	}
+}
