Horde: Rename the Horde.Build project to Horde.Server.

diff --git a/Engine/Source/Programs/Horde/Horde.Server.Tests/LogFileServiceTests.cs b/Engine/Source/Programs/Horde/Horde.Server.Tests/LogFileServiceTests.cs
new file mode 100644
--- /dev/null
+++ b/Engine/Source/Programs/Horde/Horde.Server.Tests/LogFileServiceTests.cs
@@ -0,0 +1,346 @@
+// Copyright Epic Games, Inc. All Rights Reserved.
+
+using System;
+using System.Collections.Generic;
+using System.IO;
+using System.Linq;
+using System.Security.Claims;
+using System.Text;
+using System.Threading;
+using System.Threading.Tasks;
+using EpicGames.Core;
+using Horde.Server.Agents.Sessions;
+using Horde.Server.Jobs;
+using Horde.Server.Logs;
+using Horde.Server.Logs.Builder;
+using Horde.Server.Logs.Data;
+using Horde.Server.Logs.Storage;
+using Horde.Server.Storage;
+using Horde.Server.Storage.Backends;
+using Horde.Server.Utilities;
+using HordeCommon;
+using Microsoft.Extensions.Logging;
+using Microsoft.Extensions.Logging.Abstractions;
+using Microsoft.VisualStudio.TestTools.UnitTesting;
+using MongoDB.Bson;
+
+namespace Horde.Server.Tests
+{
+	[TestClass]
+    public class LogFileServiceTest : DatabaseIntegrationTest
+    {
+		private readonly FakeClock _clock;
+        private readonly LogFileService _logFileService;
+		private readonly ILoggerFactory _loggerFactory;
+		private readonly ILogStorage _logStorage;
+
+		public LogFileServiceTest()
+        {
+            LogFileCollection logFileCollection = new LogFileCollection(GetMongoServiceSingleton());
+
+			_loggerFactory = new LoggerFactory();
+            ILogger<LogFileService> logger = _loggerFactory.CreateLogger<LogFileService>();
+
+			ILogBuilder logBuilder = new RedisLogBuilder(GetRedisServiceSingleton().ConnectionPool, NullLogger.Instance);
+			_logStorage = new PersistentLogStorage(new MemoryStorageBackend().ForType<PersistentLogStorage>(), NullLogger<PersistentLogStorage>.Instance);
+			_clock = new FakeClock();
+			TestOptions<ServerSettings> settingsOpts = new (new ServerSettings());
+			_logFileService = new LogFileService(logFileCollection, null!, logBuilder, _logStorage, _clock, null!, null!, settingsOpts, logger);
+        }
+
+		protected override void Dispose(bool disposing)
+		{
+			base.Dispose(disposing);
+
+			_logFileService.Dispose();
+			_logStorage.Dispose();
+			_loggerFactory.Dispose();
+		}
+
+		[TestMethod]
+		public void PlainTextDecoder()
+		{
+			string jsonText = "{\"time\":\"2022-01-26T14:18:29\",\"level\":\"Error\",\"message\":\"       \\tdepends on: /mnt/horde/U5\\u002BR5.0\\u002BInc\\u002BMin/Sync/Engine/Binaries/Linux/libUnrealEditor-Core.so\",\"id\":1,\"line\":5,\"lineCount\":1028}";
+			byte[] jsonData = Encoding.UTF8.GetBytes(jsonText);
+
+			byte[] outputData = new byte[jsonData.Length];
+			int outputLength = LogText.ConvertToPlainText(jsonData, outputData, 0);
+			string outputText = Encoding.UTF8.GetString(outputData, 0, outputLength);
+
+			string expectText = "       \tdepends on: /mnt/horde/U5+R5.0+Inc+Min/Sync/Engine/Binaries/Linux/libUnrealEditor-Core.so\n";
+			Assert.AreEqual(outputText, expectText);
+			
+		}
+
+		[TestMethod]
+        public async Task WriteLogLifecycleOldTest()
+        {
+			JobId jobId = JobId.GenerateNewId();
+            ILogFile logFile = await _logFileService.CreateLogFileAsync(jobId, null, null, LogType.Text, useNewStorageBackend: false, logId: null, cancellationToken: CancellationToken.None);
+
+            logFile = (await ((ILogFileService)_logFileService).WriteLogDataAsync(logFile, 0, 0, Encoding.ASCII.GetBytes("hello\n"), true))!;
+			logFile = (await ((ILogFileService)_logFileService).WriteLogDataAsync(logFile, 6, 1, Encoding.ASCII.GetBytes("foo\nbar\n"), true))!;
+			logFile = (await ((ILogFileService)_logFileService).WriteLogDataAsync(logFile, 6 + 8, 3, Encoding.ASCII.GetBytes("baz\n"), false))!;
+
+            Assert.AreEqual("hello", await ReadLogFile(_logFileService, logFile, 0, 5));
+            Assert.AreEqual("foo\nbar\nbaz\n", await ReadLogFile(_logFileService, logFile, 6, 12));
+
+            LogMetadata metadata = await _logFileService.GetMetadataAsync(logFile, CancellationToken.None);
+            Assert.AreEqual(6 + 8 + 4, metadata.Length);
+            Assert.AreEqual(4, metadata.MaxLineIndex);
+
+			List<Utf8String> lines = await _logFileService.ReadLinesAsync(logFile, 0, 100, CancellationToken.None);
+			Assert.AreEqual("hello", lines[0].ToString());
+			Assert.AreEqual("foo", lines[1].ToString());
+			Assert.AreEqual("bar", lines[2].ToString());
+			Assert.AreEqual("baz", lines[3].ToString());
+        }
+        
+        [TestMethod]
+        public async Task WriteLogLifecycleTest()
+        {
+	        await WriteLogLifecycle(_logFileService, 30);
+        }
+
+        
+        private static async Task AssertMetadata(ILogFileService logFileService, ILogFile logFile, long expectedLength, long expectedMaxLineIndex)
+        {
+	        LogMetadata metadata = await logFileService.GetMetadataAsync(logFile, CancellationToken.None);
+	        Assert.AreEqual(expectedLength, metadata.Length);
+	        Assert.AreEqual(expectedMaxLineIndex, metadata.MaxLineIndex);
+        }
+
+        private static async Task AssertChunk(ILogFileService logFileService, LogId logFileId, long numChunks, int chunkId, long offset, long length,
+	        long lineIndex)
+        {
+	        ILogFile? logFile = await logFileService.GetLogFileAsync(logFileId, CancellationToken.None);
+	        Assert.AreEqual(numChunks, logFile!.Chunks.Count);
+
+	        ILogChunk chunk = logFile.Chunks[chunkId];
+	        Assert.AreEqual(offset, chunk.Offset);
+	        Assert.AreEqual(length, chunk.Length);
+	        Assert.AreEqual(lineIndex, chunk.LineIndex);
+        }
+
+        private static async Task AssertLineOffset(ILogFileService logFileService, LogId logFileId, int lineIndex, int clampedLineIndex, long offset)
+        {
+	        ILogFile? logFile = await logFileService.GetLogFileAsync(logFileId, CancellationToken.None);
+
+			List<Utf8String> lines = await logFileService.ReadLinesAsync(logFile!, 0, lineIndex + 1, CancellationToken.None);
+			Assert.AreEqual(clampedLineIndex, Math.Min(lineIndex, lines.Count));
+			Assert.AreEqual(offset, lines.Take(lineIndex).Sum(x => x.Length + 1));
+        }
+        
+        protected static async Task<string> ReadLogFile(ILogFileService logFileService, ILogFile logFile, long offset, long length)
+        {
+			using Stream stream = await logFileService.OpenRawStreamAsync(logFile);
+
+			byte[] prefix = new byte[(int)offset];
+			await stream.ReadGreedyAsync(prefix);
+
+			byte[] data = new byte[(int)length];
+			int actualLength = await stream.ReadGreedyAsync(data);
+
+			return Encoding.UTF8.GetString(data.AsSpan(0, actualLength));
+        }
+        
+        public async Task WriteLogLifecycle(ILogFileService lfs, int maxChunkLength)
+        {
+			JobId jobId = JobId.GenerateNewId();
+            ILogFile logFile = await lfs.CreateLogFileAsync(jobId, null, null, LogType.Text, useNewStorageBackend: false);
+
+            string str1 = "hello\n";
+            string str2 = "foo\nbar\n";
+            string str3 = "baz\nqux\nquux\n";
+            string str4 = "quuz\n";
+
+            int lineIndex = 0;
+            int offset = 0;
+
+            // First write with flush. Will become chunk #1
+            logFile = (await lfs.WriteLogDataAsync(logFile, offset, lineIndex, Encoding.ASCII.GetBytes(str1), true, maxChunkLength))!;
+			await _clock.AdvanceAsync(TimeSpan.FromSeconds(2.0));
+            Assert.AreEqual(str1, await ReadLogFile(lfs, logFile, 0, str1.Length));
+            Assert.AreEqual(str1, await ReadLogFile(lfs, logFile, 0, str1.Length + 100)); // Reading too far is valid?
+            await AssertMetadata(lfs, logFile, str1.Length, 1);
+            await AssertChunk(lfs, logFile.Id, 1, 0, 0, str1.Length, lineIndex);
+            await AssertLineOffset(lfs, logFile.Id, 0, 0, 0);
+            await AssertLineOffset(lfs, logFile.Id, 1, 1, 6);
+            await AssertLineOffset(lfs, logFile.Id, 1, 1, 6);
+            await AssertLineOffset(lfs, logFile.Id, 1, 1, 6);
+
+            // Second write without flushing. Will become chunk #2
+            offset += str1.Length;
+            lineIndex += str1.Count(f => f == '\n');
+            logFile = (await lfs.WriteLogDataAsync(logFile, offset, lineIndex, Encoding.ASCII.GetBytes(str2), false, maxChunkLength))!;
+			await _clock.AdvanceAsync(TimeSpan.FromSeconds(2.0));
+			Assert.AreEqual(str1 + str2, await ReadLogFile(lfs, logFile, 0, str1.Length + str2.Length));
+            await AssertMetadata(lfs, logFile, str1.Length + str2.Length, 3); // FIXME: what are max line index?
+            await AssertChunk(lfs, logFile.Id, 2, 0, 0, str1.Length, 0);
+            await AssertChunk(lfs, logFile.Id, 2, 1, str1.Length, 0, 1); // Last chunk have length zero as it's being written
+            await AssertLineOffset(lfs, logFile.Id, 0, 0, 0);
+            await AssertLineOffset(lfs, logFile.Id, 1, 1, 6);
+            await AssertLineOffset(lfs, logFile.Id, 2, 2, 10);
+            await AssertLineOffset(lfs, logFile.Id, 3, 3, 14);
+
+            // Third write without flushing. Will become chunk #2
+            offset += str2.Length;
+            lineIndex += str3.Count(f => f == '\n');
+            logFile = (await lfs.WriteLogDataAsync(logFile, offset, lineIndex, Encoding.ASCII.GetBytes(str3), false, maxChunkLength))!;
+			await _clock.AdvanceAsync(TimeSpan.FromSeconds(2.0));
+			Assert.AreEqual(str1 + str2 + str3, await ReadLogFile(lfs, logFile, 0, str1.Length + str2.Length + str3.Length));
+            //await AssertMetadata(Lfs, LogFile, Str1.Length + Str2.Length + Str3.Length, 8);
+            // Since no flush has happened, chunks should be identical to last write
+            await AssertChunk(lfs, logFile.Id, 2, 0, 0, str1.Length, 0);
+            await AssertChunk(lfs, logFile.Id, 2, 1, str1.Length, 0, 1);
+            await AssertLineOffset(lfs, logFile.Id, 0, 0, 0);
+            await AssertLineOffset(lfs, logFile.Id, 1, 1, 6);
+            await AssertLineOffset(lfs, logFile.Id, 2, 2, 10);
+            await AssertLineOffset(lfs, logFile.Id, 3, 3, 14);
+            await AssertLineOffset(lfs, logFile.Id, 4, 4, 18);
+
+            // Fourth write with flush. Will become chunk #2
+            offset += str3.Length;
+            lineIndex += str4.Count(f => f == '\n');
+            logFile = (await lfs.WriteLogDataAsync(logFile, offset, lineIndex, Encoding.ASCII.GetBytes(str4), true, maxChunkLength))!;
+			await _clock.AdvanceAsync(TimeSpan.FromSeconds(2.0));
+			Assert.AreEqual(str1 + str2 + str3 + str4,
+                await ReadLogFile(lfs, logFile, 0, str1.Length + str2.Length + str3.Length + str4.Length));
+            Assert.AreEqual(str3 + str4,
+                await ReadLogFile(lfs, logFile, str1.Length + str2.Length, str3.Length + str4.Length));
+            await AssertMetadata(lfs, logFile, str1.Length + str2.Length + str3.Length + str4.Length, 7);
+            await AssertChunk(lfs, logFile.Id, 2, 0, 0, str1.Length, 0);
+            await AssertChunk(lfs, logFile.Id, 2, 1, str1.Length, str2.Length + str3.Length + str4.Length, 1);
+            await AssertLineOffset(lfs, logFile.Id, 0, 0, 0);
+            await AssertLineOffset(lfs, logFile.Id, 1, 1, 6);
+            await AssertLineOffset(lfs, logFile.Id, 2, 2, 10);
+            await AssertLineOffset(lfs, logFile.Id, 3, 3, 14);
+            await AssertLineOffset(lfs, logFile.Id, 4, 4, 18);
+            await AssertLineOffset(lfs, logFile.Id, 5, 5, 22);
+            await AssertLineOffset(lfs, logFile.Id, 6, 6, 27);
+            await AssertLineOffset(lfs, logFile.Id, 7, 7, 32);
+            
+            // Fifth write with flush and data that will span more than chunk. Will become chunk #3
+            string a = "Lorem ipsum dolor sit amet\n";
+            string b = "consectetur adipiscing\n";
+            string str5 = a + b;
+            
+            offset += str4.Length;
+            lineIndex += str5.Count(f => f == '\n');
+            
+            // Using this single write below will fail the ReadLogFile assert below. A bug?
+            //await LogFileService.WriteLogDataAsync(LogFile, Offset, LineIndex, Encoding.ASCII.GetBytes(Str5), true);
+            
+            // Dividing it in two like this will work however
+            logFile = (await lfs.WriteLogDataAsync(logFile, offset, lineIndex, Encoding.ASCII.GetBytes(a), false, maxChunkLength))!;
+			logFile = (await lfs.WriteLogDataAsync(logFile, offset + a.Length, lineIndex + 1, Encoding.ASCII.GetBytes(b), true, maxChunkLength))!;
+
+			await _clock.AdvanceAsync(TimeSpan.FromSeconds(2.0));
+			await AssertMetadata(lfs, logFile, str1.Length + str2.Length + str3.Length + str4.Length + str5.Length, 9);
+            await AssertChunk(lfs, logFile.Id, 4, 0, 0, str1.Length, 0);
+            await AssertChunk(lfs, logFile.Id, 4, 1, str1.Length, str2.Length + str3.Length + str4.Length, 1);
+            await AssertChunk(lfs, logFile.Id, 4, 2, offset, a.Length, 7);
+            await AssertChunk(lfs, logFile.Id, 4, 3, offset + a.Length, b.Length, 8);
+
+            Assert.AreEqual(str5, await ReadLogFile(lfs, logFile, offset, str5.Length));
+        }
+
+        [TestMethod]
+        public async Task GetLogFileTest()
+        {
+            await GetMongoServiceSingleton().Database.DropCollectionAsync("LogFiles");
+            Assert.AreEqual(0, (await _logFileService.GetLogFilesAsync()).Count);
+
+			// Will implicitly test GetLogFileAsync(), AddCachedLogFile()
+			JobId jobId = JobId.GenerateNewId();
+            SessionId sessionId = SessionId.GenerateNewId();
+            ILogFile a = await _logFileService.CreateLogFileAsync(jobId, null, sessionId, LogType.Text, useNewStorageBackend: false, logId: null, CancellationToken.None);
+            ILogFile b = (await _logFileService.GetCachedLogFileAsync(a.Id, CancellationToken.None))!;
+            Assert.AreEqual(a.JobId, b.JobId);
+            Assert.AreEqual(a.SessionId, b.SessionId);
+            Assert.AreEqual(a.Type, b.Type);
+
+            ILogFile? notFound = await _logFileService.GetCachedLogFileAsync(LogId.GenerateNewId(), CancellationToken.None);
+            Assert.IsNull(notFound);
+
+            await _logFileService.CreateLogFileAsync(JobId.GenerateNewId(), null, SessionId.GenerateNewId(), LogType.Text, useNewStorageBackend: false, logId: null, cancellationToken: CancellationToken.None);
+            Assert.AreEqual(2, (await _logFileService.GetLogFilesAsync()).Count);
+        }
+
+        [TestMethod]
+        public async Task AuthorizeForSession()
+        {
+			JobId jobId = JobId.GenerateNewId();
+            SessionId sessionId = SessionId.GenerateNewId();
+            ILogFile logFile = await _logFileService.CreateLogFileAsync(jobId, null, sessionId, LogType.Text, useNewStorageBackend: false, logId: null, cancellationToken: CancellationToken.None);
+            ILogFile logFileNoSession = await _logFileService.CreateLogFileAsync(jobId, null, null, LogType.Text, useNewStorageBackend: false, logId: null, cancellationToken: CancellationToken.None);
+
+			ClaimsPrincipal hasClaim = new ClaimsPrincipal(new ClaimsIdentity(new List<Claim>
+            {
+                new Claim(HordeClaimTypes.AgentSessionId, sessionId.ToString()),
+            }, "TestAuthType"));
+            ClaimsPrincipal hasNoClaim = new ClaimsPrincipal(new ClaimsIdentity(new List<Claim>
+            {
+                new Claim(HordeClaimTypes.AgentSessionId, "invalid-session-id"),
+            }, "TestAuthType"));
+
+            Assert.IsTrue(LogFileService.AuthorizeForSession(logFile, hasClaim));
+            Assert.IsFalse(LogFileService.AuthorizeForSession(logFile, hasNoClaim));
+            Assert.IsFalse(LogFileService.AuthorizeForSession(logFileNoSession, hasClaim));
+        }
+
+		[TestMethod]
+		public async Task ChunkSplitting()
+		{
+			JobId jobId = JobId.GenerateNewId();
+			ILogFile logFile = await _logFileService.CreateLogFileAsync(jobId, null, null, LogType.Text, useNewStorageBackend: false, logId: null, cancellationToken: CancellationToken.None);
+
+			long offset = 0;
+
+			const int MaxChunkSize = 4;
+			const int MaxSubChunkLineCount = 128;
+
+			byte[] line1 = Encoding.UTF8.GetBytes("hello world\n");
+			await _logFileService.WriteLogDataAsync(logFile, offset, 0, line1, false, MaxChunkSize, MaxSubChunkLineCount, CancellationToken.None);
+			offset += line1.Length;
+
+			byte[] line2 = Encoding.UTF8.GetBytes("ab\n");
+			await _logFileService.WriteLogDataAsync(logFile, offset, 1, line2, false, MaxChunkSize, MaxSubChunkLineCount, CancellationToken.None);
+			offset += line2.Length;
+
+			byte[] line3 = Encoding.UTF8.GetBytes("a\n");
+			await _logFileService.WriteLogDataAsync(logFile, offset, 2, line3, false, MaxChunkSize, MaxSubChunkLineCount, CancellationToken.None);
+			offset += line3.Length;
+
+			byte[] line4 = Encoding.UTF8.GetBytes("b\n");
+			await _logFileService.WriteLogDataAsync(logFile, offset, 3, line4, false, MaxChunkSize, MaxSubChunkLineCount, CancellationToken.None);
+			offset += line4.Length;
+
+			byte[] line5 = Encoding.UTF8.GetBytes("a\nb\n");
+			await _logFileService.WriteLogDataAsync(logFile, offset, 4, line5, false, MaxChunkSize, MaxSubChunkLineCount, CancellationToken.None);
+			// offset += line5.Length;
+
+			await _logFileService.FlushAsync();
+			logFile = (await _logFileService.GetLogFileAsync(logFile.Id, CancellationToken.None))!;
+			Assert.AreEqual(4, logFile.Chunks.Count);
+			Assert.AreEqual(6, logFile.MaxLineIndex);
+
+			Assert.AreEqual(0, logFile.Chunks[0].LineIndex);
+			Assert.AreEqual(1, logFile.Chunks[1].LineIndex);
+			Assert.AreEqual(2, logFile.Chunks[2].LineIndex);
+			Assert.AreEqual(4, logFile.Chunks[3].LineIndex);
+
+			Assert.AreEqual(12, logFile.Chunks[0].Length);
+			Assert.AreEqual(3, logFile.Chunks[1].Length);
+			Assert.AreEqual(4, logFile.Chunks[2].Length);
+			Assert.AreEqual(4, logFile.Chunks[3].Length);
+
+			Assert.AreEqual(0, logFile.Chunks.GetChunkForLine(0));
+			Assert.AreEqual(1, logFile.Chunks.GetChunkForLine(1));
+			Assert.AreEqual(2, logFile.Chunks.GetChunkForLine(2));
+			Assert.AreEqual(2, logFile.Chunks.GetChunkForLine(3));
+			Assert.AreEqual(3, logFile.Chunks.GetChunkForLine(4));
+			Assert.AreEqual(3, logFile.Chunks.GetChunkForLine(5));
+		}
+	}
+}
\ No newline at end of file
diff --git a/Engine/Source/Programs/Horde/Horde.Server.Tests/LogIndexingTests.cs b/Engine/Source/Programs/Horde/Horde.Server.Tests/LogIndexingTests.cs
new file mode 100644
--- /dev/null
+++ b/Engine/Source/Programs/Horde/Horde.Server.Tests/LogIndexingTests.cs
@@ -0,0 +1,246 @@
+// Copyright Epic Games, Inc. All Rights Reserved.
+
+using System;
+using System.Collections.Generic;
+using System.IO;
+using System.Linq;
+using System.Text;
+using System.Threading;
+using System.Threading.Tasks;
+using Horde.Server.Compute.Tests.Properties;
+using Horde.Server.Logs;
+using Horde.Server.Logs.Builder;
+using Horde.Server.Logs.Storage;
+using Horde.Server.Jobs;
+using Horde.Server.Utilities;
+using HordeCommon;
+using Microsoft.Extensions.Configuration;
+using Microsoft.Extensions.DependencyInjection;
+using Microsoft.Extensions.Logging;
+using Microsoft.VisualStudio.TestTools.UnitTesting;
+using Horde.Server.Logs.Data;
+using EpicGames.Horde.Logs;
+
+namespace Horde.Server.Tests
+{
+	[TestClass]
+	public class LogIndexingTests : DatabaseIntegrationTest
+	{
+		private readonly ILogFileService _logFileService;
+		private readonly NullLogStorage _nullLogStorage;
+		private readonly LocalLogStorage _logStorage;
+		private readonly byte[] _data = Resources.TextFile;
+
+		public LogIndexingTests()
+		{
+			LogFileCollection logFileCollection = new LogFileCollection(GetMongoServiceSingleton());
+
+			ServiceProvider serviceProvider = new ServiceCollection()
+				.AddLogging(builder => builder.AddConsole().SetMinimumLevel(LogLevel.Debug))
+				.BuildServiceProvider();
+
+			ILoggerFactory loggerFactory = serviceProvider.GetRequiredService<ILoggerFactory>();
+			ILogger<LogFileService> logger = loggerFactory.CreateLogger<LogFileService>();
+
+			// Just to satisfy the parameter, need to be fixed
+			IConfiguration config = new ConfigurationBuilder().Build();
+
+			ILogBuilder logBuilder = new LocalLogBuilder();
+			_nullLogStorage = new NullLogStorage();
+			_logStorage = new LocalLogStorage(20, _nullLogStorage);
+			TestOptions<ServerSettings> settingsOpts = new (new ServerSettings());
+			_logFileService = new LogFileService(logFileCollection, null!, logBuilder, _logStorage, new FakeClock(), null!, null!, settingsOpts, logger);
+		}
+
+		protected override void Dispose(bool disposing)
+		{
+			base.Dispose(disposing);
+
+			_logStorage.Dispose();
+			_nullLogStorage.Dispose();
+		}
+
+		[TestMethod]
+		public async Task IndexTests()
+		{
+			JobId jobId = JobId.GenerateNewId();
+			ILogFile logFile = await _logFileService.CreateLogFileAsync(jobId, null, null, LogType.Text, useNewStorageBackend: false);
+
+			// Write the test data to the log file in blocks
+			int offset = 0;
+			int lineIndex = 0;
+			while (offset < _data.Length)
+			{
+				int length = 0;
+				int lineCount = 0;
+
+				for(int idx = offset; idx < Math.Min(_data.Length, offset + 1883); idx++)
+				{
+					if(_data[idx] == '\n')
+					{
+						length = (idx + 1) - offset;
+						lineCount++;
+						break;
+					}
+				}
+
+				logFile = await WriteLogDataAsync(logFile, offset, lineIndex, _data.AsMemory(offset, length), false);
+
+				offset += length;
+				lineIndex += lineCount;
+			}
+
+			// Read the data back out and check it's the same
+			byte[] readData = new byte[_data.Length];
+			using (Stream stream = await _logFileService.OpenRawStreamAsync(logFile))
+			{
+				int readSize = await stream.ReadAsync(readData, 0, readData.Length);
+				Assert.AreEqual(readData.Length, readSize);
+
+				int equalSize = 0;
+				while(equalSize < _data.Length && _data[equalSize] == readData[equalSize])
+				{
+					equalSize++;
+				}
+
+				Assert.AreEqual(equalSize, readSize);
+			}
+
+			// Test some searches
+			await SearchLogDataTestAsync(logFile);
+
+			// Generate an index and test again
+			logFile = await WriteLogDataAsync(logFile, offset, lineIndex, Array.Empty<byte>(), true);
+			await SearchLogDataTestAsync(logFile);
+		}
+
+		[TestMethod]
+		public void TrieTests()
+		{
+			ReadOnlyTrieBuilder builder = new ReadOnlyTrieBuilder();
+
+			ulong[] values = { 1, 2, 3, 5, 8, 13, 21, 34, 55, 89, 144, 233, 377, 610, 987, 1597, 2584, 4181, 6765, 10946, 0xfedcba9876543210UL };
+			foreach (ulong value in values)
+			{
+				builder.Add(value);
+			}
+
+			ReadOnlyTrie trie = builder.Build();
+			Assert.IsTrue(Enumerable.SequenceEqual(trie.EnumerateRange(0, UInt64.MaxValue), values));
+			Assert.IsTrue(Enumerable.SequenceEqual(trie.EnumerateRange(0, 90), values.Where(x => x <= 90)));
+			Assert.IsTrue(Enumerable.SequenceEqual(trie.EnumerateRange(2, 89), values.Where(x => x >= 2 && x <= 89)));
+		}
+
+		[TestMethod]
+		public async Task PartialTokenTests()
+		{
+			JobId jobId = JobId.GenerateNewId();
+			ILogFile logFile = await _logFileService.CreateLogFileAsync(jobId, null, null, LogType.Text, useNewStorageBackend: false);
+
+			string[] lines =
+			{
+				"abcdefghi\n",
+				"jklmno123\n",
+				"pqrst99uv\n",
+				"wx\n"
+			};
+
+			int length = 0;
+			for (int lineIdx = 0; lineIdx < lines.Length; lineIdx++)
+			{
+				logFile = await WriteLogDataAsync(logFile, length, lineIdx, Encoding.UTF8.GetBytes(lines[lineIdx]), true);
+				length += lines[lineIdx].Length;
+			}
+
+			for (int lineIdx = 0; lineIdx < lines.Length; lineIdx++)
+			{
+				for (int strLen = 1; strLen < 7; strLen++)
+				{
+					for (int strOfs = 0; strOfs + strLen < lines[lineIdx].Length - 1; strOfs++)
+					{
+						string str = lines[lineIdx].Substring(strOfs, strLen);
+
+						SearchStats stats = new SearchStats();
+						List<int> results = await _logFileService.SearchLogDataAsync(logFile, str, 0, 5, stats, CancellationToken.None);
+						Assert.AreEqual(1, results.Count);
+						Assert.AreEqual(lineIdx, results[0]);
+
+						Assert.AreEqual(1, stats.NumScannedBlocks);
+						Assert.AreEqual(3, stats.NumSkippedBlocks);
+						Assert.AreEqual(0, stats.NumFalsePositiveBlocks);
+					}
+				}
+			}
+		}
+
+		[TestMethod]
+		public async Task AppendIndexTests()
+		{
+			JobId jobId = JobId.GenerateNewId();
+			ILogFile logFile = await _logFileService.CreateLogFileAsync(jobId, null, null, LogType.Text, useNewStorageBackend: false);
+
+			logFile = await WriteLogDataAsync(logFile, 0, 0, Encoding.UTF8.GetBytes("abc\n"), true);
+			logFile = await WriteLogDataAsync(logFile, 4, 1, Encoding.UTF8.GetBytes("def\n"), true);
+			logFile = await WriteLogDataAsync(logFile, 8, 2, Encoding.UTF8.GetBytes("ghi\n"), false);
+
+			await ((LogFileService)_logFileService).FlushPendingWritesAsync();
+
+			{
+				SearchStats stats = new SearchStats();
+				List<int> results = await _logFileService.SearchLogDataAsync(logFile, "abc", 0, 5, stats, CancellationToken.None);
+				Assert.AreEqual(1, results.Count);
+				Assert.AreEqual(0, results[0]);
+
+				Assert.AreEqual(2, stats.NumScannedBlocks); // abc + ghi (no index yet because it hasn't been flushed)
+				Assert.AreEqual(1, stats.NumSkippedBlocks); // def
+				Assert.AreEqual(0, stats.NumFalsePositiveBlocks);
+			}
+			{
+				SearchStats stats = new SearchStats();
+				List<int> results = await _logFileService.SearchLogDataAsync(logFile, "def", 0, 5, stats, CancellationToken.None);
+				Assert.AreEqual(1, results.Count);
+				Assert.AreEqual(1, results[0]);
+
+				Assert.AreEqual(2, stats.NumScannedBlocks); // def + ghi (no index yet because it hasn't been flushed)
+				Assert.AreEqual(1, stats.NumSkippedBlocks); // abc
+				Assert.AreEqual(0, stats.NumFalsePositiveBlocks);
+			}
+			{
+				SearchStats stats = new SearchStats();
+				List<int> results = await _logFileService.SearchLogDataAsync(logFile, "ghi", 0, 5, stats, CancellationToken.None);
+				Assert.AreEqual(1, results.Count);
+				Assert.AreEqual(2, results[0]);
+
+				Assert.AreEqual(1, stats.NumScannedBlocks); // ghi
+				Assert.AreEqual(2, stats.NumSkippedBlocks); // abc + def
+				Assert.AreEqual(0, stats.NumFalsePositiveBlocks);
+			}
+		}
+
+		async Task<ILogFile> WriteLogDataAsync(ILogFile logFile, long offset, int lineIndex, ReadOnlyMemory<byte> data, bool flush)
+		{
+			const int MaxChunkLength = 32 * 1024;
+			const int MaxSubChunkLineCount = 128;
+
+			ILogFile? newLogFile = await _logFileService.WriteLogDataAsync(logFile, offset, lineIndex, data, flush, MaxChunkLength, MaxSubChunkLineCount);
+			Assert.IsNotNull(newLogFile);
+			return newLogFile!;
+		}
+
+		async Task SearchLogDataTestAsync(ILogFile logFile)
+		{
+			await SearchLogDataTestAsync(logFile, "HISPANIOLA", 0, 4, new[] { 1503, 1520, 1525, 1595 });
+			await SearchLogDataTestAsync(logFile, "Hispaniola", 0, 4, new[] { 1503, 1520, 1525, 1595 });
+			await SearchLogDataTestAsync(logFile, "HizpaniolZ", 0, 4, Array.Empty<int>());
+			await SearchLogDataTestAsync(logFile, "Pieces of eight!", 0, 100, new[] { 2227, 2228, 5840, 5841, 7520 });
+			await SearchLogDataTestAsync(logFile, "NEWSLETTER", 0, 100, new[] { 7886 });
+		}
+
+		async Task SearchLogDataTestAsync(ILogFile logFile, string text, int firstLine, int count, int[] expectedLines)
+		{
+			SearchStats stats = new SearchStats();
+			List<int> lines = await _logFileService.SearchLogDataAsync(logFile, text, firstLine, count, stats, CancellationToken.None);
+			Assert.IsTrue(lines.SequenceEqual(expectedLines));
+		}
+	}
+}
diff --git a/Engine/Source/Programs/Horde/Horde.Server.Tests/NotificationServiceTests.cs b/Engine/Source/Programs/Horde/Horde.Server.Tests/NotificationServiceTests.cs
new file mode 100644
--- /dev/null
+++ b/Engine/Source/Programs/Horde/Horde.Server.Tests/NotificationServiceTests.cs
@@ -0,0 +1,205 @@
+// Copyright Epic Games, Inc. All Rights Reserved.
+
+using System;
+using System.Collections.Generic;
+using System.Threading.Tasks;
+using Horde.Server.Agents.Pools;
+using Horde.Server.Devices;
+using Horde.Server.Issues;
+using Horde.Server.Jobs;
+using Horde.Server.Jobs.Graphs;
+using Horde.Server.Notifications;
+using Horde.Server.Users;
+using Horde.Server.Streams;
+using Horde.Server.Utilities;
+using HordeCommon;
+using Microsoft.Extensions.DependencyInjection;
+using Microsoft.VisualStudio.TestTools.UnitTesting;
+using Moq;
+using Horde.Server.Logs;
+using Microsoft.Extensions.Caching.Memory;
+
+namespace Horde.Server.Tests
+{
+	public class FakeNotificationSink : INotificationSink
+	{
+		public List<JobScheduledNotification> JobScheduledNotifications { get; } = new();
+		public int JobScheduledCallCount { get; set; }
+		
+		public Task NotifyJobScheduledAsync(List<JobScheduledNotification> notifications)
+		{
+			JobScheduledNotifications.AddRange(notifications);
+			JobScheduledCallCount++;
+			return Task.CompletedTask;
+		}
+
+		public Task NotifyJobCompleteAsync(IJob job, IGraph graph, LabelOutcome outcome) { throw new NotImplementedException(); }
+		public Task NotifyJobCompleteAsync(IUser user, IJob job, IGraph graph, LabelOutcome outcome) { throw new NotImplementedException(); }
+		public Task NotifyJobStepCompleteAsync(IUser user, IJob job, IJobStepBatch batch, IJobStep step, INode node, List<ILogEventData> jobStepEventData) { throw new NotImplementedException(); }
+		public Task NotifyLabelCompleteAsync(IUser user, IJob job, ILabel label, int labelIdx, LabelOutcome outcome, List<(string, JobStepOutcome, Uri)> stepData) { throw new NotImplementedException(); }
+		public Task NotifyIssueUpdatedAsync(IIssue issue) { throw new NotImplementedException(); }
+		public Task NotifyConfigUpdateAsync(Exception? ex) => Task.CompletedTask;
+		public Task NotifyConfigUpdateFailureAsync(string errorMessage, string fileName, int? change = null, IUser? author = null, string? description = null) { throw new NotImplementedException(); }
+		public Task NotifyDeviceServiceAsync(string message, IDevice? device = null, IDevicePool? pool = null, StreamConfig? stream = null, IJob? job = null, IJobStep? step = null, INode? node = null, IUser? user = null) { throw new NotImplementedException(); }
+		public Task SendIssueReportAsync(IssueReportGroup report) => throw new NotImplementedException();
+	}
+
+	[TestClass]
+	public class NotificationServiceTests : TestSetup
+	{
+		protected override void ConfigureServices(IServiceCollection services)
+		{
+			base.ConfigureServices(services);
+
+			services.AddSingleton<FakeNotificationSink>();
+			services.AddSingleton<INotificationSink>(sp => sp.GetRequiredService<FakeNotificationSink>());
+		}
+
+		public static IJob CreateJob(StreamId streamId, int change, string name, IGraph graph)
+		{
+			JobId jobId = JobId.Parse("5ec16da1774cb4000107c2c1");
+
+			List<IJobStepBatch> batches = new List<IJobStepBatch>();
+			for (int groupIdx = 0; groupIdx < graph.Groups.Count; groupIdx++)
+			{
+				INodeGroup @group = graph.Groups[groupIdx];
+
+				List<IJobStep> steps = new List<IJobStep>();
+				for (int nodeIdx = 0; nodeIdx < @group.Nodes.Count; nodeIdx++)
+				{
+					SubResourceId stepId = new SubResourceId((ushort)((groupIdx * 100) + nodeIdx));
+
+					Mock<IJobStep> step = new Mock<IJobStep>(MockBehavior.Strict);
+					step.SetupGet(x => x.Id).Returns(stepId);
+					step.SetupGet(x => x.NodeIdx).Returns(nodeIdx);
+
+					steps.Add(step.Object);
+				}
+
+				SubResourceId batchId = new SubResourceId((ushort)(groupIdx * 100));
+
+				Mock<IJobStepBatch> batch = new Mock<IJobStepBatch>(MockBehavior.Strict);
+				batch.SetupGet(x => x.Id).Returns(batchId);
+				batch.SetupGet(x => x.GroupIdx).Returns(groupIdx);
+				batch.SetupGet(x => x.Steps).Returns(steps);
+				batches.Add(batch.Object);
+			}
+
+			Mock<IJob> job = new Mock<IJob>(MockBehavior.Strict);
+			job.SetupGet(x => x.Id).Returns(jobId);
+			job.SetupGet(x => x.Name).Returns(name);
+			job.SetupGet(x => x.StreamId).Returns(streamId);
+			job.SetupGet(x => x.Change).Returns(change);
+			job.SetupGet(x => x.Batches).Returns(batches);
+			return job.Object;
+		}
+		
+		[TestMethod]
+		public async Task NotifyJobScheduled()
+		{
+			FakeNotificationSink fakeSink = ServiceProvider.GetRequiredService<FakeNotificationSink>();
+
+			NotificationService service = (NotificationService)ServiceProvider.GetRequiredService<INotificationService>();
+			await service._ticker.StartAsync();
+
+			Fixture fixture = await CreateFixtureAsync();
+			IPool pool = await PoolService.CreatePoolAsync("BogusPool", properties: new Dictionary<string, string>());
+
+			Assert.AreEqual(0, fakeSink.JobScheduledNotifications.Count);
+			service.NotifyJobScheduled(pool, false, fixture.Job1, fixture.Graph, SubResourceId.Random());
+			service.NotifyJobScheduled(pool, false, fixture.Job2, fixture.Graph, SubResourceId.Random());
+			
+			// Currently no good way to wait for NotifyJobScheduled() to complete as the execution is completely async in background task (see ExecuteAsync)
+			await Task.Delay(1000);
+			await Clock.AdvanceAsync(service._notificationBatchInterval + TimeSpan.FromMinutes(5));
+			Assert.AreEqual(2, fakeSink.JobScheduledNotifications.Count);
+			Assert.AreEqual(1, fakeSink.JobScheduledCallCount);
+		}
+		
+		[TestMethod]
+		public async Task JobScheduledNotificationsAreDeduplicated()
+		{
+			FakeNotificationSink fakeSink = ServiceProvider.GetRequiredService<FakeNotificationSink>();
+			NotificationService service = (NotificationService)ServiceProvider.GetRequiredService<INotificationService>();
+			await service._ticker.StartAsync();
+			Fixture fixture = await CreateFixtureAsync();
+			IPool pool = await PoolService.CreatePoolAsync("BogusPool", properties: new Dictionary<string, string>());
+
+			service.NotifyJobScheduled(pool, false, fixture.Job1, fixture.Graph, SubResourceId.Random());
+			service.NotifyJobScheduled(pool, false, fixture.Job1, fixture.Graph, SubResourceId.Random());
+			service.NotifyJobScheduled(pool, false, fixture.Job1, fixture.Graph, SubResourceId.Random());
+			
+			// Currently no good way to wait for NotifyJobScheduled() to complete as the execution is completely async in background task (see ExecuteAsync)
+			await Task.Delay(1000);
+			await Clock.AdvanceAsync(service._notificationBatchInterval + TimeSpan.FromMinutes(5));
+			
+			// Only one job scheduled notification should have been sent, despite queuing three
+			Assert.AreEqual(1, fakeSink.JobScheduledNotifications.Count);
+			
+			// Clear the cache by compacting it 100%
+			MemoryCache cache = (MemoryCache)ServiceProvider.GetRequiredService<IMemoryCache>();
+			cache.Compact(1.0);
+			
+			// Notify of exactly the same job again
+			service.NotifyJobScheduled(pool, false, fixture.Job1, fixture.Graph, SubResourceId.Random());
+			
+			await Task.Delay(1000);
+			await Clock.AdvanceAsync(service._notificationBatchInterval + TimeSpan.FromMinutes(5));
+			Assert.AreEqual(2, fakeSink.JobScheduledNotifications.Count);
+		}
+
+		//public void NotifyLabelUpdate(IJob Job, IReadOnlyList<(LabelState, LabelOutcome)> OldLabelStates, IReadOnlyList<(LabelState, LabelOutcome)> NewLabelStates)
+		//{
+		//	// If job has any label trigger IDs, send label complete notifications if needed
+		//	if (Job.LabelIdxToTriggerId.Any())
+		//	{
+		//		EnqueueTask(() => SendAllLabelNotificationsAsync(Job, OldLabelStates, NewLabelStates));
+		//	}
+		//}
+
+		//private async Task SendAllLabelNotificationsAsync(IJob Job, IReadOnlyList<(LabelState State, LabelOutcome Outcome)> OldLabelStates, IReadOnlyList<(LabelState, LabelOutcome)> NewLabelStates)
+		//{
+		//	IStream? Stream = await StreamService.GetStreamAsync(Job.StreamId);
+		//	if (Stream == null)
+		//	{
+		//		return;
+		//	}
+
+		//	IGraph? Graph = await GraphCollection.GetAsync(Job.GraphHash);
+		//	if (Graph == null)
+		//	{
+		//		return;
+		//	}
+
+		//	IReadOnlyDictionary<NodeRef, IJobStep> StepForNode = Job.GetStepForNodeMap();
+		//	for (int LabelIdx = 0; LabelIdx < Graph.Labels.Count; ++LabelIdx)
+		//	{
+		//		(LabelState State, LabelOutcome Outcome) OldLabel = OldLabelStates[LabelIdx];
+		//		(LabelState State, LabelOutcome Outcome) NewLabel = NewLabelStates[LabelIdx];
+		//		if (OldLabel != NewLabel)
+		//		{
+		//			// If the state transitioned from Unspecified to Running, don't update unless the outcome also changed.
+		//			if (OldLabel.State == LabelState.Unspecified && NewLabel.State == LabelState.Running && OldLabel.Outcome == NewLabel.Outcome)
+		//			{
+		//				continue;
+		//			}
+
+		//			// If the label isn't complete, don't report on outcome changing to success, this will be reported when the label state becomes complete.
+		//			if (NewLabel.State != LabelState.Complete && NewLabel.Outcome == LabelOutcome.Success)
+		//			{
+		//				return;
+		//			}
+
+		//			bool fireTrigger = NewLabel.State == LabelState.Complete;
+		//			INotificationTrigger? Trigger = await GetNotificationTrigger(Job.LabelIdxToTriggerId[LabelIdx], fireTrigger);
+		//			if (Trigger == null)
+		//			{
+		//				continue;
+		//			}
+
+		//			await SendLabelUpdateNotificationsAsync(Job, Stream, Graph, StepForNode, Graph.Labels[LabelIdx], NewLabel.State, NewLabel.Outcome, Trigger);
+		//		}
+		//	}
+		//}
+	}
+}
diff --git a/Engine/Source/Programs/Horde/Horde.Server.Tests/PoolServiceTests.cs b/Engine/Source/Programs/Horde/Horde.Server.Tests/PoolServiceTests.cs
new file mode 100644
--- /dev/null
+++ b/Engine/Source/Programs/Horde/Horde.Server.Tests/PoolServiceTests.cs
@@ -0,0 +1,98 @@
+// Copyright Epic Games, Inc. All Rights Reserved.
+
+using System;
+using System.Collections.Generic;
+using System.Threading.Tasks;
+using Horde.Server.Agents;
+using Horde.Server.Agents.Pools;
+using Horde.Server.Utilities;
+using Microsoft.VisualStudio.TestTools.UnitTesting;
+
+namespace Horde.Server.Tests
+{
+	[TestClass]
+    public class PoolServiceTests : TestSetup
+    {
+        private readonly Dictionary<string, string> _fixtureProps = new Dictionary<string, string>
+        {
+            {"foo", "bar"},
+            {"lorem", "ipsum"}
+        };
+
+        public PoolServiceTests()
+        {
+        }
+
+        private async Task<IPool> CreatePoolFixture(string name)
+        {
+            return await PoolService.CreatePoolAsync(name, properties: _fixtureProps);
+        }
+
+        [TestMethod]
+        public async Task GetPoolTest()
+        {
+            Assert.IsNull(await PoolService.GetPoolAsync(new PoolId("this-does-not-exist")));
+
+            IPool newPool = await CreatePoolFixture("create-pool");
+			IPool? pool = await PoolService.GetPoolAsync(newPool.Id);
+			Assert.IsNotNull(pool);
+            Assert.AreEqual("create-pool", pool!.Id.ToString());
+            Assert.AreEqual("create-pool", pool.Name);
+            Assert.AreEqual(_fixtureProps.Count, pool.Properties.Count);
+            Assert.AreEqual(_fixtureProps["foo"], pool.Properties["foo"]);
+            Assert.AreEqual(_fixtureProps["lorem"], pool.Properties["lorem"]);
+        }
+        
+        [TestMethod]
+        public async Task GetPoolsTest()
+        {
+            await GetMongoServiceSingleton().Database.DropCollectionAsync("Pools");
+            
+            List<IPool> pools = await PoolService.GetPoolsAsync();
+            Assert.AreEqual(pools.Count, 0);
+
+			IPool pool0 = await CreatePoolFixture("multiple-pools-0");
+			IPool pool1 = await CreatePoolFixture("multiple-pools-1");
+            pools = await PoolService.GetPoolsAsync();
+            Assert.AreEqual(pools.Count, 2);
+            Assert.AreEqual(pools[0].Name, pool0.Name);
+            Assert.AreEqual(pools[1].Name, pool1.Name);
+        }
+        
+        [TestMethod]
+        public async Task DeletePoolTest()
+        {
+            Assert.IsFalse(await PoolService.DeletePoolAsync(new PoolId("this-does-not-exist")));
+            IPool pool = await CreatePoolFixture("pool-to-be-deleted");
+            Assert.IsTrue(await PoolService.DeletePoolAsync(pool.Id));
+        }
+        
+        [TestMethod]
+        public async Task UpdatePoolTest()
+        {
+			string uniqueSuffix = Guid.NewGuid().ToString("N");
+            IPool pool = await CreatePoolFixture($"update-pool-{uniqueSuffix}");
+            Dictionary<string, string?> updatedProps = new Dictionary<string, string?>
+            {
+                {"foo", "bar"},
+                {"lorem", null}, // This entry will get removed
+                {"cookies", "yumyum"},
+            };
+
+            IPool? updatedPool = await PoolService.UpdatePoolAsync(pool, $"update-pool-new-name-{uniqueSuffix}", newProperties: updatedProps);
+			Assert.IsNotNull(updatedPool);
+            Assert.AreEqual(pool.Id, updatedPool!.Id);
+            Assert.AreEqual($"update-pool-new-name-{uniqueSuffix}", updatedPool.Name);
+            Assert.AreEqual(updatedPool.Properties.Count, 2);
+            Assert.AreEqual(updatedProps["foo"], updatedPool.Properties["foo"]);
+            Assert.AreEqual(updatedProps["cookies"], updatedPool.Properties["cookies"]);
+        }
+        
+        [TestMethod]
+        public async Task UpdatePoolCollectionTest()
+        {
+	        IPool pool = await CreatePoolFixture("update-pool-2");
+	        await PoolCollection.TryUpdateAsync(pool, lastScaleUpTime: DateTime.UtcNow);
+        }
+    }
+}
\ No newline at end of file
diff --git a/Engine/Source/Programs/Horde/Horde.Server.Tests/RpcServiceTest.cs b/Engine/Source/Programs/Horde/Horde.Server.Tests/RpcServiceTest.cs
new file mode 100644
--- /dev/null
+++ b/Engine/Source/Programs/Horde/Horde.Server.Tests/RpcServiceTest.cs
@@ -0,0 +1,508 @@
+// Copyright Epic Games, Inc. All Rights Reserved.
+
+using System;
+using System.Collections.Generic;
+using System.IO;
+using System.Linq;
+using System.Reflection;
+using System.Security.Claims;
+using System.Threading;
+using System.Threading.Channels;
+using System.Threading.Tasks;
+using Google.Protobuf;
+using Grpc.Core;
+using Horde.Server.Utilities;
+using Horde.Server.Agents.Pools;
+using Horde.Server.Agents.Leases;
+using Horde.Server.Logs;
+using Horde.Server.Server;
+using HordeCommon.Rpc;
+using Microsoft.AspNetCore.Http;
+using Microsoft.AspNetCore.Http.Features;
+using Microsoft.Extensions.Hosting;
+using Microsoft.VisualStudio.TestTools.UnitTesting;
+using MongoDB.Bson;
+using AgentCapabilities = HordeCommon.Rpc.Messages.AgentCapabilities;
+using ISession = Microsoft.AspNetCore.Http.ISession;
+using Horde.Server.Jobs.Artifacts;
+using Horde.Server.Agents;
+using HordeCommon.Rpc.Messages;
+using Horde.Server.Agents.Sessions;
+
+namespace Horde.Server.Tests
+{
+	public class AppLifetimeStub : IHostApplicationLifetime
+	{
+		public CancellationToken ApplicationStarted { get; }
+		public CancellationToken ApplicationStopping { get; }
+		public CancellationToken ApplicationStopped { get; }
+
+		public AppLifetimeStub()
+		{
+			ApplicationStarted = new CancellationToken();
+			ApplicationStopping = new CancellationToken();
+			ApplicationStopped = new CancellationToken();
+		}
+
+		public void StopApplication()
+		{
+			throw new NotImplementedException();
+		}
+	}
+
+	[TestClass]
+	public class RpcServiceTest : TestSetup
+	{
+		private readonly ServerCallContext _adminContext = new ServerCallContextStub("app-horde-admins");
+
+		sealed class HttpContextStub : HttpContext
+		{
+			public override ConnectionInfo Connection { get; } = null!;
+			public override IFeatureCollection Features { get; } = null!;
+			public override IDictionary<object, object?> Items { get; set; } = null!;
+			public override HttpRequest Request { get; } = null!;
+			public override CancellationToken RequestAborted { get; set; }
+			public override IServiceProvider RequestServices { get; set; } = null!;
+			public override HttpResponse Response { get; } = null!;
+			public override ISession Session { get; set; } = null!;
+			public override string TraceIdentifier { get; set; } = null!;
+			public override ClaimsPrincipal User { get; set; }
+			public override WebSocketManager WebSockets { get; } = null!;
+
+			public HttpContextStub(string roleClaimType)
+			{
+				User = new ClaimsPrincipal(new ClaimsIdentity(new List<Claim>
+				{
+					new Claim(HordeClaimTypes.Role, roleClaimType),
+				}, "TestAuthType"));
+			}
+			
+			public HttpContextStub(ClaimsPrincipal user)
+			{
+				User = user;
+			}
+
+			public override void Abort()
+			{
+				throw new NotImplementedException();
+			}
+		}
+
+		class ServerCallContextStub : ServerCallContext
+		{
+			// Copied from ServerCallContextExtensions.cs in Grpc.Core
+			const string HttpContextKey = "__HttpContext";
+
+			public static ServerCallContext ForAdminWithAgentSessionId(string agentSessionId)
+			{
+				return new ServerCallContextStub(new ClaimsPrincipal(new ClaimsIdentity(new List<Claim>
+				{
+					new Claim(HordeClaimTypes.Role, "app-horde-admins"),
+					new Claim(HordeClaimTypes.AgentSessionId, agentSessionId),
+				}, "TestAuthType")));
+			}
+			
+			public static ServerCallContext ForAdmin()
+			{
+				return new ServerCallContextStub(new ClaimsPrincipal(new ClaimsIdentity(new List<Claim>
+				{
+					new Claim(HordeClaimTypes.Role, "app-horde-admins"),
+				}, "TestAuthType")));
+			}
+
+			public ServerCallContextStub(string roleClaimType)
+			{
+				// The GetHttpContext extension falls back to getting the HttpContext from UserState
+				// We can piggyback on that behavior during tests
+				UserState[HttpContextKey] = new HttpContextStub(roleClaimType);
+			}
+			
+			public ServerCallContextStub(ClaimsPrincipal user)
+			{
+				// The GetHttpContext extension falls back to getting the HttpContext from UserState
+				// We can piggyback on that behavior during tests
+				UserState[HttpContextKey] = new HttpContextStub(user);
+			}
+
+			protected override Task WriteResponseHeadersAsyncCore(Metadata responseHeaders)
+			{
+				throw new NotImplementedException();
+			}
+
+			protected override ContextPropagationToken CreatePropagationTokenCore(ContextPropagationOptions? options)
+			{
+				throw new NotImplementedException();
+			}
+
+			protected override string MethodCore { get; } = null!;
+			protected override string HostCore { get; } = null!;
+			protected override string PeerCore { get; } = null!;
+			protected override DateTime DeadlineCore { get; } = DateTime.Now.AddHours(24);
+			protected override Metadata RequestHeadersCore { get; } = null!;
+			protected override CancellationToken CancellationTokenCore { get; } = new CancellationToken();
+			protected override Metadata ResponseTrailersCore { get; } = null!;
+			protected override Status StatusCore { get; set; }
+			protected override WriteOptions? WriteOptionsCore { get; set; } = null!;
+			protected override AuthContext AuthContextCore { get; } = null!;
+		}
+		
+		class RpcServiceInvoker : CallInvoker
+		{
+			private readonly RpcService _rpcService;
+			private readonly ServerCallContext _serverCallContext;
+
+			public RpcServiceInvoker(RpcService rpcService, ServerCallContext serverCallContext)
+			{
+				_rpcService = rpcService;
+				_serverCallContext = serverCallContext;
+			}
+
+			public override TResponse BlockingUnaryCall<TRequest, TResponse>(Method<TRequest, TResponse> method, string? host, CallOptions options, TRequest request)
+			{
+				throw new NotImplementedException("Blocking calls are not supported! Method " + method.FullName);
+			}
+
+			public override AsyncUnaryCall<TResponse> AsyncUnaryCall<TRequest, TResponse>(Method<TRequest, TResponse> method, string? host, CallOptions options, TRequest request)
+			{
+				MethodInfo methodInfo = GetMethod(method.Name);
+				Task<TResponse> res = (methodInfo.Invoke(_rpcService, new object[] {request, _serverCallContext}) as Task<TResponse>)!;
+				return new AsyncUnaryCall<TResponse>(res, null!, null!, null!, null!, null!);
+			}
+
+			public override AsyncServerStreamingCall<TResponse> AsyncServerStreamingCall<TRequest, TResponse>(Method<TRequest, TResponse> method, string? host, CallOptions options,
+				TRequest request)
+			{
+				Console.WriteLine($"RpcServiceInvoker.AsyncServerStreamingCall(method={method.FullName} request={request})");
+				throw new NotImplementedException();
+			}
+
+			public override AsyncClientStreamingCall<TRequest, TResponse> AsyncClientStreamingCall<TRequest, TResponse>(Method<TRequest, TResponse> method, string? host, CallOptions options)
+			{
+				Console.WriteLine($"RpcServiceInvoker.AsyncClientStreamingCall(method={method.FullName})");
+				throw new NotImplementedException();
+			}
+
+			public override AsyncDuplexStreamingCall<TRequest, TResponse> AsyncDuplexStreamingCall<TRequest, TResponse>(Method<TRequest, TResponse> method, string? host, CallOptions options)
+			{
+				Console.WriteLine($"RpcServiceInvoker.AsyncDuplexStreamingCall(method={method.FullName})");
+
+				GrpcDuplexStreamHandler<TRequest> requestStream = new GrpcDuplexStreamHandler<TRequest>(_serverCallContext);
+				GrpcDuplexStreamHandler<TResponse> responseStream = new GrpcDuplexStreamHandler<TResponse>(_serverCallContext);
+
+				MethodInfo methodInfo = GetMethod(method.Name);
+				Task methodTask = (methodInfo.Invoke(_rpcService, new object[] {requestStream, responseStream, _serverCallContext}) as Task)!;
+				methodTask.ContinueWith(t => { Console.Error.WriteLine($"Uncaught exception in {method.Name}: {t}"); }, CancellationToken.None, TaskContinuationOptions.OnlyOnFaulted, TaskScheduler.Default);
+				
+				return new AsyncDuplexStreamingCall<TRequest, TResponse>(requestStream, responseStream, null!, null!, null!, null!);
+			}
+			
+			private MethodInfo GetMethod(string methodName)
+			{
+				MethodInfo? method = _rpcService.GetType().GetMethod(methodName);
+				if (method == null)
+				{
+					throw new ArgumentException($"Method {methodName} not found in RpcService");
+				}
+
+				return method;
+			}
+		}
+		
+		/// <summary>
+		/// Combines and cross-writes streams for a duplex streaming call in gRPC
+		/// </summary>
+		/// <typeparam name="T"></typeparam>
+		class GrpcDuplexStreamHandler<T> : IServerStreamWriter<T>, IClientStreamWriter<T>, IAsyncStreamReader<T> where T : class
+		{
+			private readonly ServerCallContext _serverCallContext;
+			private readonly Channel<T> _channel;
+
+			public WriteOptions? WriteOptions { get; set; }
+
+			public GrpcDuplexStreamHandler(ServerCallContext serverCallContext)
+			{
+				_channel = System.Threading.Channels.Channel.CreateUnbounded<T>();
+
+				_serverCallContext = serverCallContext;
+			}
+
+			public void Complete()
+			{
+				_channel.Writer.Complete();
+			}
+
+			public IAsyncEnumerable<T> ReadAllAsync()
+			{
+				return _channel.Reader.ReadAllAsync();
+			}
+
+			public async Task<T?> ReadNextAsync()
+			{
+				if (await _channel.Reader.WaitToReadAsync())
+				{
+					_channel.Reader.TryRead(out T? message);
+					return message;
+				}
+				else
+				{
+					return null;
+				}
+			}
+
+			public Task WriteAsync(T message)
+			{
+				if (_serverCallContext.CancellationToken.IsCancellationRequested)
+				{
+					return Task.FromCanceled(_serverCallContext.CancellationToken);
+				}
+
+				if (!_channel.Writer.TryWrite(message))
+				{
+					throw new InvalidOperationException("Unable to write message.");
+				}
+
+				return Task.CompletedTask;
+			}
+
+			public Task CompleteAsync()
+			{
+				Complete();
+				return Task.CompletedTask;
+			}
+
+			public async Task<bool> MoveNext(CancellationToken cancellationToken)
+			{
+				_serverCallContext.CancellationToken.ThrowIfCancellationRequested();
+				
+				if (await _channel.Reader.WaitToReadAsync(cancellationToken))
+				{
+					if (_channel.Reader.TryRead(out T? message))
+					{
+						Current = message;
+						return true;
+					}
+				}
+
+				Current = null!;
+				return false;
+			}
+
+			public T Current { get; private set; } = null!;
+		}
+
+		[TestMethod]
+		public async Task CreateSessionTest()
+		{
+			CreateSessionRequest req = new CreateSessionRequest();
+			await Assert.ThrowsExceptionAsync<StructuredRpcException>(() => RpcService.CreateSession(req, _adminContext));
+
+			req.Name = "MyName";
+			await Assert.ThrowsExceptionAsync<StructuredRpcException>(() => RpcService.CreateSession(req, _adminContext));
+
+			req.Capabilities = new AgentCapabilities();
+			CreateSessionResponse res = await RpcService.CreateSession(req, _adminContext);
+
+			Assert.AreEqual("MYNAME", res.AgentId);
+			// TODO: Check Token, ExpiryTime, SessionId 
+		}
+		
+		[TestMethod]
+		public async Task AgentJoinsPoolThroughProperties()
+		{
+			CreateSessionRequest req = new () { Name = "bogusAgentName", Capabilities = new AgentCapabilities() };
+			req.Capabilities.Properties.Add($"{KnownPropertyNames.RequestedPools}=fooPool,barPool");
+			CreateSessionResponse res = await RpcService.CreateSession(req, _adminContext);
+
+			IAgent agent = (await AgentService.GetAgentAsync(new AgentId(res.AgentId)))!;
+			CollectionAssert.AreEquivalent(new List<PoolId> { new ("fooPool"), new ("barPool") }, agent.GetPools().ToList());
+			
+			// Connect a second time, when the agent has already been created
+			req = new () { Name = "bogusAgentName", Capabilities = new AgentCapabilities() };
+			req.Capabilities.Properties.Add($"{KnownPropertyNames.RequestedPools}=bazPool");
+			res = await RpcService.CreateSession(req, _adminContext);
+
+			agent = (await AgentService.GetAgentAsync(new AgentId(res.AgentId)))!;
+			CollectionAssert.AreEquivalent(new List<PoolId> { new ("fooPool"), new ("barPool") , new ("bazPool") }, agent.GetPools().ToList());
+		}
+		
+		[TestMethod]
+		public async Task PropertiesFromAgentCapabilities()
+		{
+			CreateSessionRequest req = new () { Name = "bogusAgentName", Capabilities = new AgentCapabilities() };
+			req.Capabilities.Properties.Add("fooKey=barValue");
+			CreateSessionResponse res = await RpcService.CreateSession(req, _adminContext);
+			IAgent agent = (await AgentService.GetAgentAsync(new AgentId(res.AgentId)))!;
+			Assert.IsTrue(agent.Properties.Contains("fooKey=barValue"));
+		}
+		
+		[TestMethod]
+		public async Task PropertiesFromDeviceCapabilities()
+		{
+			CreateSessionRequest req = new () { Name = "bogusAgentName", Capabilities = new AgentCapabilities() };
+			req.Capabilities.Devices.Add(new DeviceCapabilities { Handle = "someHandle", Properties = { "foo=bar" } });
+			CreateSessionResponse res = await RpcService.CreateSession(req, _adminContext);
+			IAgent agent = (await AgentService.GetAgentAsync(new AgentId(res.AgentId)))!;
+			Assert.IsTrue(agent.Properties.Contains("foo=bar"));
+		}
+		
+		[TestMethod]
+		public async Task KnownPropertiesAreSetAsResources()
+		{
+			CreateSessionRequest req = new () { Name = "bogusAgentName", Capabilities = new AgentCapabilities() };
+			req.Capabilities.Devices.Add(new DeviceCapabilities { Handle = "someHandle", Properties = { $"{KnownPropertyNames.LogicalCores}=10" }});
+			CreateSessionResponse res = await RpcService.CreateSession(req, _adminContext);
+			IAgent agent = (await AgentService.GetAgentAsync(new AgentId(res.AgentId)))!;
+			Assert.AreEqual(10, agent.Resources[KnownPropertyNames.LogicalCores]);
+		}
+
+		[TestMethod]
+		public async Task UpdateSessionTest()
+		{
+			CreateSessionRequest createReq = new CreateSessionRequest
+			{
+				Name = "UpdateSessionTest1", Capabilities = new AgentCapabilities()
+			};
+			CreateSessionResponse createRes = await RpcService.CreateSession(createReq, _adminContext);
+			string agentId = createRes.AgentId;
+			string sessionId = createRes.SessionId;
+
+			TestAsyncStreamReader<UpdateSessionRequest> requestStream =
+				new TestAsyncStreamReader<UpdateSessionRequest>(_adminContext);
+			TestServerStreamWriter<UpdateSessionResponse> responseStream =
+				new TestServerStreamWriter<UpdateSessionResponse>(_adminContext);
+			Task call = RpcService.UpdateSession(requestStream, responseStream, _adminContext);
+
+			requestStream.AddMessage(new UpdateSessionRequest {AgentId = "does-not-exist", SessionId = sessionId});
+			StructuredRpcException re = await Assert.ThrowsExceptionAsync<StructuredRpcException>(() => call);
+			Assert.AreEqual(StatusCode.NotFound, re.StatusCode);
+			Assert.IsTrue(re.Message.Contains("Invalid agent name", StringComparison.OrdinalIgnoreCase));
+		}
+		
+		[TestMethod]
+		public async Task QueryServerSessionTest()
+		{
+			RpcService._longPollTimeout = TimeSpan.FromMilliseconds(200);
+
+			TestAsyncStreamReader<QueryServerStateRequest> requestStream =
+				new TestAsyncStreamReader<QueryServerStateRequest>(_adminContext);
+			TestServerStreamWriter<QueryServerStateResponse> responseStream =
+				new TestServerStreamWriter<QueryServerStateResponse>(_adminContext);
+			Task call = RpcService.QueryServerState(requestStream, responseStream, _adminContext);
+
+			requestStream.AddMessage(new QueryServerStateRequest {Name = "bogusAgentName"});
+			QueryServerStateResponse? res = await responseStream.ReadNextAsync();
+			Assert.IsNotNull(res);
+			
+			res = await responseStream.ReadNextAsync();
+			Assert.IsNotNull(res);
+
+			// Should timeout after LongPollTimeout specified above
+			await call;
+		}
+		
+		[TestMethod]
+		public async Task FinishBatchTest()
+		{
+			CreateSessionRequest createReq = new CreateSessionRequest
+			{
+				Name = "UpdateSessionTest1", Capabilities = new AgentCapabilities()
+			};
+			CreateSessionResponse createRes = await RpcService.CreateSession(createReq, _adminContext);
+			string agentId = createRes.AgentId;
+			string sessionId = createRes.SessionId;
+
+			TestAsyncStreamReader<UpdateSessionRequest> requestStream =
+				new TestAsyncStreamReader<UpdateSessionRequest>(_adminContext);
+			TestServerStreamWriter<UpdateSessionResponse> responseStream =
+				new TestServerStreamWriter<UpdateSessionResponse>(_adminContext);
+			Task call = RpcService.UpdateSession(requestStream, responseStream, _adminContext);
+
+			requestStream.AddMessage(new UpdateSessionRequest {AgentId = "does-not-exist", SessionId = sessionId});
+			StructuredRpcException re = await Assert.ThrowsExceptionAsync<StructuredRpcException>(() => call);
+			Assert.AreEqual(StatusCode.NotFound, re.StatusCode);
+			Assert.IsTrue(re.Message.Contains("Invalid agent name", StringComparison.OrdinalIgnoreCase));
+		}
+		
+		[TestMethod]
+		public async Task UploadArtifactTest()
+		{
+			Fixture fixture = await CreateFixtureAsync();
+
+			SessionId sessionId = SessionId.GenerateNewId();
+			ServerCallContext context = new ServerCallContextStub(new ClaimsPrincipal(new ClaimsIdentity(new List<Claim>
+			{
+				new Claim(HordeClaimTypes.Role, "app-horde-admins"),
+				new Claim(HordeClaimTypes.AgentSessionId, sessionId.ToString()),
+			}, "TestAuthType")));
+
+			string[] data = {"foo", "bar", "baz", "qux"};
+			string dataStr = String.Join("", data);
+			
+			UploadArtifactMetadata metadata = new UploadArtifactMetadata
+			{
+				JobId = fixture.Job1.Id.ToString(),
+				BatchId = fixture.Job1.Batches[0].Id.ToString(),
+				StepId = fixture.Job1.Batches[0].Steps[0].Id.ToString(),
+				Name = "testfile.txt",
+				MimeType = "text/plain",
+				Length = dataStr.Length
+			};
+
+			// Set the session ID on the job batch to pass auth later
+			Deref(await JobCollection.TryAssignLeaseAsync(fixture.Job1, 0, new PoolId("foo"),
+				new AgentId("test"), sessionId,
+				LeaseId.GenerateNewId(), LogId.GenerateNewId()));
+/*
+			TestAsyncStreamReader<UploadArtifactRequest> RequestStream = new TestAsyncStreamReader<UploadArtifactRequest>(Context);
+			Task<UploadArtifactResponse> Call = TestSetup.RpcService.UploadArtifact(RequestStream,  Context);
+			RequestStream.AddMessage(new UploadArtifactRequest { Metadata = Metadata });
+			RequestStream.AddMessage(new UploadArtifactRequest { Data = ByteString.CopyFromUtf8(Data[0]) });
+			RequestStream.AddMessage(new UploadArtifactRequest { Data = ByteString.CopyFromUtf8(Data[1]) });
+			RequestStream.AddMessage(new UploadArtifactRequest { Data = ByteString.CopyFromUtf8(Data[2]) });
+			// Only send three messages and not the last one.
+			// Aborting the upload here and retry in next code section below.
+			RequestStream.Complete();
+			await Task.Delay(500);
+*/			
+			TestAsyncStreamReader<UploadArtifactRequest> requestStream = new TestAsyncStreamReader<UploadArtifactRequest>(context);
+			Task<UploadArtifactResponse> call = RpcService.UploadArtifact(requestStream,  context);
+			requestStream.AddMessage(new UploadArtifactRequest { Metadata = metadata });
+			requestStream.AddMessage(new UploadArtifactRequest { Data = ByteString.CopyFromUtf8(data[0]) });
+			requestStream.AddMessage(new UploadArtifactRequest { Data = ByteString.CopyFromUtf8(data[1]) });
+			requestStream.AddMessage(new UploadArtifactRequest { Data = ByteString.CopyFromUtf8(data[2]) });
+			requestStream.AddMessage(new UploadArtifactRequest { Data = ByteString.CopyFromUtf8(data[3]) });
+			UploadArtifactResponse res = await call;
+
+			IArtifactV1? artifact = await ArtifactCollection.GetArtifactAsync(ObjectId.Parse(res.Id));
+			Assert.IsNotNull(artifact);
+			Stream stream = await ArtifactCollection.OpenArtifactReadStreamAsync(artifact!);
+			using StreamReader reader = new StreamReader(stream);
+			string text = await reader.ReadToEndAsync();
+			Assert.AreEqual(dataStr, text);
+		}
+		/*
+		[TestMethod]
+		public async Task UploadSoftwareAsync()
+		{
+			TestSetup TestSetup = await GetTestSetup();
+			
+			MemoryStream OutputStream = new MemoryStream();
+			using (ZipArchive ZipFile = new ZipArchive(OutputStream, ZipArchiveMode.Create, false))
+			{
+				string TempFilename = Path.GetTempFileName();
+				File.WriteAllText(TempFilename, "{\"Horde\": {\"Version\": \"myVersion\"}}");
+				ZipFile.CreateEntryFromFile(TempFilename, "appsettings.json");
+			}
+
+			ByteString Data = ByteString.CopyFrom(OutputStream.ToArray());
+			UploadSoftwareRequest Req = new UploadSoftwareRequest { Channel = "boguschannel", Data = Data };
+
+			UploadSoftwareResponse Res1 = await TestSetup.RpcService.UploadSoftware(Req, AdminContext);
+			Assert.AreEqual("r1", Res1.Version);
+			
+			UploadSoftwareResponse Res2 = await TestSetup.RpcService.UploadSoftware(Req, AdminContext);
+			Assert.AreEqual("r2", Res2.Version);
+		}
+		*/
+	}
+}
\ No newline at end of file
diff --git a/Engine/Source/Programs/Horde/Horde.Server.Tests/SchedulerTests.cs b/Engine/Source/Programs/Horde/Horde.Server.Tests/SchedulerTests.cs
new file mode 100644
--- /dev/null
+++ b/Engine/Source/Programs/Horde/Horde.Server.Tests/SchedulerTests.cs
@@ -0,0 +1,652 @@
+// Copyright Epic Games, Inc. All Rights Reserved.
+
+using System;
+using System.Collections.Generic;
+using System.Linq;
+using System.Threading.Tasks;
+using Horde.Server.Jobs;
+using Horde.Server.Logs;
+using Horde.Server.Server;
+using Horde.Server.Projects;
+using Horde.Server.Streams;
+using Horde.Server.Jobs.Templates;
+using Horde.Server.Jobs.Schedules;
+using Horde.Server.Users;
+using Horde.Server.Utilities;
+using HordeCommon;
+using Microsoft.Extensions.DependencyInjection;
+using Microsoft.VisualStudio.TestTools.UnitTesting;
+using Horde.Server.Jobs.Graphs;
+using Horde.Server.Perforce;
+
+namespace Horde.Server.Tests
+{
+	[TestClass]
+    public class SchedulerTests : TestSetup
+	{
+		ProjectId ProjectId { get; } = new ProjectId("ue5");
+		StreamId StreamId { get; } = new StreamId("ue5-main");
+		TemplateId TemplateId { get; } = new TemplateId("template1");
+
+		readonly ITemplate _template;
+		readonly HashSet<JobId> _initialJobIds;
+
+		public SchedulerTests()
+		{
+			IUser bob = UserCollection.FindOrAddUserByLoginAsync("Bob").Result;
+
+			GlobalConfig globalConfig = new GlobalConfig();
+			globalConfig.Projects.Add(new ProjectConfig { Id = ProjectId, Name = "UE4" });
+			SetConfig(globalConfig);
+
+			_template = TemplateCollection.GetOrAddAsync(new TemplateConfig { Name = "Test template" }).Result;
+
+			_initialJobIds = new HashSet<JobId>(JobCollection.FindAsync().Result.Select(x => x.Id));
+
+			PerforceService.Changes.Clear();
+			PerforceService.AddChange(StreamId, 100, bob, "", new[] { "code.cpp" });
+			PerforceService.AddChange(StreamId, 101, bob, "", new[] { "content.uasset" });
+			PerforceService.AddChange(StreamId, 102, bob, "", new[] { "content.uasset" });
+		}
+
+		async Task<IStream> SetScheduleAsync(ScheduleConfig schedule)
+		{
+			await ScheduleService.ResetAsync();
+
+			StreamConfig streamConfig = new StreamConfig();
+			streamConfig.Id = StreamId;
+			streamConfig.Name = "//UE5/Main";
+			streamConfig.Tabs.Add(new JobsTabConfig { Title = "foo", Templates = new List<TemplateId> { TemplateId } });
+			streamConfig.Templates.Add(new TemplateRefConfig { Id = TemplateId, Name = "Test", Schedule = schedule });
+
+			UpdateConfig(x => x.Projects[0].Streams = new List<StreamConfig> { streamConfig });
+
+			return await StreamCollection.GetAsync(streamConfig);
+		}
+
+		public async Task<List<IJob>> FileTestHelperAsync(params string[] files)
+		{
+			IUser bob = await UserCollection.FindOrAddUserByLoginAsync("Bob", "Bob");
+
+			PerforceService.Changes.Clear();
+			PerforceService.AddChange(StreamId, 100, bob, "", new[] { "code.cpp" });
+			PerforceService.AddChange(StreamId, 101, bob, "", new[] { "content.uasset" });
+			PerforceService.AddChange(StreamId, 102, bob, "", new[] { "content.uasset" });
+			PerforceService.AddChange(StreamId, 103, bob, "", new[] { "foo/code.cpp" });
+			PerforceService.AddChange(StreamId, 104, bob, "", new[] { "bar/code.cpp" });
+			PerforceService.AddChange(StreamId, 105, bob, "", new[] { "foo/bar/content.uasset" });
+			PerforceService.AddChange(StreamId, 106, bob, "", new[] { "bar/foo/content.uasset" });
+
+			DateTime startTime = new DateTime(2021, 1, 1, 12, 0, 0, DateTimeKind.Local); // Friday Jan 1, 2021 
+			Clock.UtcNow = startTime;
+
+			ScheduleConfig schedule = new ScheduleConfig();
+			schedule.Enabled = true;
+			schedule.MaxChanges = 10;
+			schedule.Patterns.Add(new SchedulePatternConfig { Interval = 1 });
+			schedule.Files = files.ToList();
+			await SetScheduleAsync(schedule);
+
+			await Clock.AdvanceAsync(TimeSpan.FromHours(1.25));
+			await ScheduleService.TickForTestingAsync();
+
+			return await GetNewJobs();
+		}
+
+		[TestMethod]
+		public async Task FileTest1Async()
+		{
+			List<IJob> jobs = await FileTestHelperAsync("....cpp");
+			Assert.AreEqual(3, jobs.Count);
+			Assert.AreEqual(100, jobs[0].Change);
+			Assert.AreEqual(103, jobs[1].Change);
+			Assert.AreEqual(104, jobs[2].Change);
+		}
+
+		[TestMethod]
+		public async Task FileTest2Async()
+		{
+			List<IJob> jobs = await FileTestHelperAsync("/foo/...");
+			Assert.AreEqual(2, jobs.Count);
+			Assert.AreEqual(103, jobs[0].Change);
+			Assert.AreEqual(105, jobs[1].Change);
+		}
+
+		[TestMethod]
+		public async Task FileTest3Async()
+		{
+			List<IJob> jobs = await FileTestHelperAsync("....uasset", "-/bar/...");
+			Assert.AreEqual(3, jobs.Count);
+			Assert.AreEqual(101, jobs[0].Change);
+			Assert.AreEqual(102, jobs[1].Change);
+			Assert.AreEqual(105, jobs[2].Change);
+		}
+
+		[TestMethod]
+		public void DayScheduleTest()
+		{
+			DateTime startTime = new DateTime(2021, 1, 1, 12, 0, 0, DateTimeKind.Utc); // Friday Jan 1, 2021 
+			Clock.UtcNow = startTime;
+
+			ScheduleConfig schedule = new ScheduleConfig { RequireSubmittedChange = false };
+			schedule.Patterns.Add(new SchedulePatternConfig(new List<DayOfWeek> { DayOfWeek.Friday, DayOfWeek.Sunday }, 13 * 60, null, null));
+
+			DateTime? nextTime = schedule.GetNextTriggerTimeUtc(startTime, TimeZoneInfo.Utc);
+			Assert.AreEqual(startTime + TimeSpan.FromHours(1.0), nextTime!.Value);
+
+			nextTime = schedule.GetNextTriggerTimeUtc(nextTime.Value, TimeZoneInfo.Utc);
+			Assert.AreEqual(startTime + TimeSpan.FromHours(1.0 + 24.0 * 2.0), nextTime!.Value);
+
+			nextTime = schedule.GetNextTriggerTimeUtc(nextTime.Value, TimeZoneInfo.Utc);
+			Assert.AreEqual(startTime + TimeSpan.FromHours(1.0 + 24.0 * 7.0), nextTime!.Value);
+
+			nextTime = schedule.GetNextTriggerTimeUtc(nextTime.Value, TimeZoneInfo.Utc);
+			Assert.AreEqual(startTime + TimeSpan.FromHours(1.0 + 24.0 * 9.0), nextTime!.Value);
+		}
+
+		[TestMethod]
+		public void MulticheduleTest()
+		{
+			DateTime startTime = new DateTime(2021, 1, 1, 12, 0, 0, DateTimeKind.Utc); // Friday Jan 1, 2021 
+			Clock.UtcNow = startTime;
+
+			ScheduleConfig schedule = new ScheduleConfig { RequireSubmittedChange = false };
+			schedule.Patterns.Add(new SchedulePatternConfig(null, 13 * 60, 14 * 60, 15));
+
+			DateTime? nextTime = schedule.GetNextTriggerTimeUtc(startTime, TimeZoneInfo.Utc);
+			Assert.AreEqual(startTime + TimeSpan.FromHours(1.0), nextTime!.Value);
+
+			nextTime = schedule.GetNextTriggerTimeUtc(nextTime.Value, TimeZoneInfo.Utc);
+			Assert.AreEqual(startTime + TimeSpan.FromHours(1.25), nextTime!.Value);
+
+			nextTime = schedule.GetNextTriggerTimeUtc(nextTime.Value, TimeZoneInfo.Utc);
+			Assert.AreEqual(startTime + TimeSpan.FromHours(1.5), nextTime!.Value);
+
+			nextTime = schedule.GetNextTriggerTimeUtc(nextTime.Value, TimeZoneInfo.Utc);
+			Assert.AreEqual(startTime + TimeSpan.FromHours(1.75), nextTime!.Value);
+
+			nextTime = schedule.GetNextTriggerTimeUtc(nextTime.Value, TimeZoneInfo.Utc);
+			Assert.AreEqual(startTime + TimeSpan.FromHours(2.0), nextTime!.Value);
+
+			nextTime = schedule.GetNextTriggerTimeUtc(nextTime.Value, TimeZoneInfo.Utc);
+			Assert.AreEqual(startTime + TimeSpan.FromHours(1.0 + 24.0), nextTime!.Value);
+		}
+
+		[TestMethod]
+		public void MultiPatternTest()
+		{
+			DateTime startTime = new DateTime(2021, 1, 1, 0, 0, 0, DateTimeKind.Utc); // Friday Jan 1, 2021 
+			Clock.UtcNow = startTime;
+
+			ScheduleConfig schedule = new ScheduleConfig { RequireSubmittedChange = false };
+			schedule.Patterns.Add(new SchedulePatternConfig(null, 11 * 60, 0, 0));
+			schedule.Patterns.Add(new SchedulePatternConfig(null, 19 * 60, 0, 0));
+
+			DateTime? nextTime = schedule.GetNextTriggerTimeUtc(startTime, TimeZoneInfo.Utc);
+			Assert.AreEqual(startTime + TimeSpan.FromHours(11), nextTime!.Value);
+
+			nextTime = schedule.GetNextTriggerTimeUtc(nextTime.Value, TimeZoneInfo.Utc);
+			Assert.AreEqual(startTime + TimeSpan.FromHours(19), nextTime!.Value);
+		}
+
+		[TestMethod]
+		public async Task NoSubmittedChangeTestAsync()
+		{
+			DateTime startTime = new DateTime(2021, 1, 1, 12, 0, 0, DateTimeKind.Local); // Friday Jan 1, 2021 
+			Clock.UtcNow = startTime;
+			Clock.TimeZone = TimeZoneInfo.Local;
+
+			ScheduleConfig schedule = new ScheduleConfig();
+			schedule.Enabled = true;
+			schedule.Patterns.Add(new SchedulePatternConfig { MinTime = 13 * 60, MaxTime = 14 * 60, Interval = 15 });
+//			Schedule.LastTriggerTime = StartTime;
+			await SetScheduleAsync(schedule);
+
+			// Initial tick
+			await ScheduleService.TickForTestingAsync();
+
+			List<IJob> jobs1 = await GetNewJobs();
+			Assert.AreEqual(0, jobs1.Count);
+
+			// Trigger a job
+			await Clock.AdvanceAsync(TimeSpan.FromHours(1.25));
+			await ScheduleService.TickForTestingAsync();
+
+			List<IJob> jobs2 = await GetNewJobs();
+			Assert.AreEqual(1, jobs2.Count);
+			Assert.AreEqual(102, jobs2[0].Change);
+			Assert.AreEqual(100, jobs2[0].CodeChange);
+		}
+
+		[TestMethod]
+		public async Task RequireSubmittedChangeTestAsync()
+		{
+			DateTime startTime = new DateTime(2021, 1, 1, 12, 0, 0, DateTimeKind.Local); // Friday Jan 1, 2021 
+			Clock.UtcNow = startTime;
+			Clock.TimeZone = TimeZoneInfo.Local;
+
+			ScheduleConfig schedule = new ScheduleConfig();
+			schedule.Enabled = true;
+			schedule.Patterns.Add(new SchedulePatternConfig { MinTime = 13 * 60, MaxTime = 14 * 60, Interval = 15 });
+			await SetScheduleAsync(schedule);
+
+			// Initial tick
+			await ScheduleService.TickForTestingAsync();
+
+			List<IJob> jobs1 = await GetNewJobs();
+			Assert.AreEqual(0, jobs1.Count);
+
+			// Trigger a job
+			await Clock.AdvanceAsync(TimeSpan.FromHours(1.25));
+			await ScheduleService.TickForTestingAsync();
+
+			List<IJob> jobs2 = await GetNewJobs();
+			Assert.AreEqual(1, jobs2.Count);
+			Assert.AreEqual(102, jobs2[0].Change);
+			Assert.AreEqual(100, jobs2[0].CodeChange);
+
+			StreamConfig? streamConfig;
+			GlobalConfig.CurrentValue.TryGetStream(StreamId, out streamConfig);
+
+			IStream stream2 = (await StreamCollection.GetAsync(streamConfig!))!;
+			ITemplateSchedule schedule2 = stream2.Templates.First().Value.Schedule!;
+			Assert.AreEqual(102, schedule2.LastTriggerChange);
+			Assert.AreEqual(Clock.UtcNow, schedule2.LastTriggerTimeUtc);
+
+			// Trigger another job
+			await Clock.AdvanceAsync(TimeSpan.FromHours(0.5));
+			await ScheduleService.TickForTestingAsync();
+
+			List<IJob> jobs3 = await GetNewJobs();
+			Assert.AreEqual(0, jobs3.Count);
+		}
+
+		[TestMethod]
+		public async Task MultipleJobsTestAsync()
+		{
+			DateTime startTime = new DateTime(2021, 1, 1, 12, 0, 0, DateTimeKind.Local); // Friday Jan 1, 2021 
+			Clock.UtcNow = startTime;
+			Clock.TimeZone = TimeZoneInfo.Local;
+
+			ScheduleConfig schedule = new ScheduleConfig();
+			schedule.Enabled = true;
+			schedule.Patterns.Add(new SchedulePatternConfig { MinTime = 13 * 60, MaxTime = 14 * 60, Interval = 15 });
+			schedule.MaxChanges = 2;
+			schedule.Commits.Add(CommitTag.Code);
+			await SetScheduleAsync(schedule);
+
+			// Initial tick
+			await ScheduleService.TickForTestingAsync();
+
+			List<IJob> jobs1 = await GetNewJobs();
+			Assert.AreEqual(0, jobs1.Count);
+
+			// Trigger some jobs
+			IUser bob = await UserCollection.FindOrAddUserByLoginAsync("Bob");
+			PerforceService.AddChange(StreamId, 103, bob, "", new string[] { "foo.cpp" });
+			PerforceService.AddChange(StreamId, 104, bob, "", new string[] { "foo.cpp" });
+			PerforceService.AddChange(StreamId, 105, bob, "", new string[] { "foo.uasset" });
+			PerforceService.AddChange(StreamId, 106, bob, "", new string[] { "foo.cpp" });
+
+			await Clock.AdvanceAsync(TimeSpan.FromHours(1.25));
+			await ScheduleService.TickForTestingAsync();
+
+			List<IJob> jobs2 = await GetNewJobs();
+			Assert.AreEqual(2, jobs2.Count);
+			Assert.AreEqual(104, jobs2[0].Change);
+			Assert.AreEqual(104, jobs2[0].CodeChange);
+			Assert.AreEqual(106, jobs2[1].Change);
+			Assert.AreEqual(106, jobs2[1].CodeChange);
+		}
+
+		[TestMethod]
+		public async Task SkipCiTestAsync()
+		{
+			DateTime startTime = new DateTime(2021, 1, 1, 12, 0, 0, DateTimeKind.Local); // Friday Jan 1, 2021 
+			Clock.UtcNow = startTime;
+			Clock.TimeZone = TimeZoneInfo.Local;
+
+			ScheduleConfig schedule = new ScheduleConfig();
+			schedule.Enabled = true;
+			schedule.Patterns.Add(new SchedulePatternConfig { MinTime = 13 * 60, MaxTime = 14 * 60, Interval = 15 });
+			schedule.MaxChanges = 2;
+			schedule.Commits.Add(CommitTag.Code);
+			await SetScheduleAsync(schedule);
+
+			// Initial tick
+			await ScheduleService.TickForTestingAsync();
+
+			List<IJob> jobs1 = await GetNewJobs();
+			Assert.AreEqual(0, jobs1.Count);
+
+			// Trigger some jobs
+			IUser bob = await UserCollection.FindOrAddUserByLoginAsync("Bob");
+			PerforceService.AddChange(StreamId, 103, bob, "", new string[] { "foo.cpp" });
+			PerforceService.AddChange(StreamId, 104, bob, "Don't build this change!\n#skipci", new string[] { "foo.cpp" });
+			PerforceService.AddChange(StreamId, 105, bob, "", new string[] { "foo.uasset" });
+			PerforceService.AddChange(StreamId, 106, bob, "", new string[] { "foo.cpp" });
+
+			await Clock.AdvanceAsync(TimeSpan.FromHours(1.25));
+			await ScheduleService.TickForTestingAsync();
+
+			List<IJob> jobs2 = await GetNewJobs();
+			Assert.AreEqual(2, jobs2.Count);
+			Assert.AreEqual(103, jobs2[0].Change);
+			Assert.AreEqual(103, jobs2[0].CodeChange);
+			Assert.AreEqual(106, jobs2[1].Change);
+			Assert.AreEqual(106, jobs2[1].CodeChange);
+		}
+
+		[TestMethod]
+		public async Task MaxActiveTestAsync()
+		{
+			DateTime startTime = new DateTime(2021, 1, 1, 12, 0, 0, DateTimeKind.Local); // Friday Jan 1, 2021 
+			Clock.UtcNow = startTime;
+			Clock.TimeZone = TimeZoneInfo.Local;
+
+			ScheduleConfig schedule = new ScheduleConfig();
+			schedule.Enabled = true;
+			schedule.RequireSubmittedChange = false;
+			schedule.Patterns.Add(new SchedulePatternConfig { MinTime = 13 * 60, MaxTime = 14 * 60, Interval = 15 });
+			schedule.MaxActive = 1;
+			await SetScheduleAsync(schedule);
+
+			// Trigger a job
+			await Clock.AdvanceAsync(TimeSpan.FromHours(1.25));
+			await ScheduleService.TickForTestingAsync();
+
+			List<IJob> jobs2 = await GetNewJobs();
+			Assert.AreEqual(1, jobs2.Count);
+			Assert.AreEqual(102, jobs2[0].Change);
+			Assert.AreEqual(100, jobs2[0].CodeChange);
+
+			// Test that another job does not trigger
+			await Clock.AdvanceAsync(TimeSpan.FromHours(0.5));
+			await ScheduleService.TickForTestingAsync();
+
+			List<IJob> jobs3 = await GetNewJobs();
+			Assert.AreEqual(0, jobs3.Count);
+
+			// Mark the original job as complete
+			await JobService.UpdateJobAsync(jobs2[0], abortedByUserId: KnownUsers.System);
+
+			// Test that another job does not trigger
+			await Clock.AdvanceAsync(TimeSpan.FromHours(0.5));
+			await ScheduleService.TickForTestingAsync();
+
+			List<IJob> jobs4 = await GetNewJobs();
+			Assert.AreEqual(1, jobs4.Count);
+			Assert.AreEqual(102, jobs4[0].Change);
+			Assert.AreEqual(100, jobs4[0].CodeChange);
+		}
+
+		[TestMethod]
+		public async Task CreateNewChangeTestAsync()
+		{
+			DateTime startTime = new DateTime(2021, 1, 1, 12, 0, 0, DateTimeKind.Local); // Friday Jan 1, 2021 
+			Clock.UtcNow = startTime;
+			Clock.TimeZone = TimeZoneInfo.Local;
+
+			ScheduleConfig schedule = new ScheduleConfig();
+			schedule.Enabled = true;
+			schedule.Patterns.Add(new SchedulePatternConfig { MinTime = 13 * 60, MaxTime = 14 * 60, Interval = 15 });
+			/*IStream stream = */await SetScheduleAsync(schedule);
+
+			// Trigger a job
+			await Clock.AdvanceAsync(TimeSpan.FromHours(1.25));
+			await ScheduleService.TickForTestingAsync();
+
+			List<IJob> jobs2 = await GetNewJobs();
+			Assert.AreEqual(1, jobs2.Count);
+			Assert.AreEqual(102, jobs2[0].Change);
+			Assert.AreEqual(100, jobs2[0].CodeChange);
+
+			// Check another job does not trigger due to the change above
+			await Clock.AdvanceAsync(TimeSpan.FromHours(1.25));
+			await ScheduleService.TickForTestingAsync();
+
+			List<IJob> jobs3 = await GetNewJobs();
+			Assert.AreEqual(0, jobs3.Count);
+		}
+
+		[TestMethod]
+		public async Task GateTestAsync()
+		{
+			DateTime startTime = new DateTime(2021, 1, 1, 12, 0, 0, DateTimeKind.Local); // Friday Jan 1, 2021 
+			Clock.UtcNow = startTime;
+
+			// Create two templates, the second dependent on the first
+			ITemplate? newTemplate1 = await TemplateCollection.GetOrAddAsync(new TemplateConfig { Name = "Test template 1" });
+			//TemplateRef newTemplateRef1 = new TemplateRef(newTemplate1);
+			TemplateId newTemplateRefId1 = new TemplateId("new-template-1");
+
+			ITemplate? newTemplate2 = await TemplateCollection.GetOrAddAsync(new TemplateConfig { Name = "Test template 2" });
+//			TemplateRef newTemplateRef2 = new TemplateRef(newTemplate2);
+//			newTemplateRef2.Schedule = new Schedule(Clock.UtcNow);
+//			newTemplateRef2.Schedule.Gate = new ScheduleGate(newTemplateRefId1, "TriggerNext");
+//			newTemplateRef2.Schedule.Patterns.Add(new SchedulePattern(null, 0, null, 10));
+//			newTemplateRef2.Schedule.LastTriggerTime = startTime;
+			TemplateId newTemplateRefId2 = new TemplateId("new-template-2");
+
+			StreamConfig config = new StreamConfig();
+			config.Id = StreamId;
+			config.Name = "//UE5/Main";
+			config.Tabs.Add(new JobsTabConfig { Title = "foo", Templates = new List<TemplateId> { newTemplateRefId1, newTemplateRefId2 } });
+			config.Templates = new() { new TemplateRefConfig { Id = newTemplateRefId1 }, new TemplateRefConfig { Id = newTemplateRefId2 }  };
+			UpdateConfig(x => x.Projects[0].Streams = new List<StreamConfig> { config });
+
+			IStream stream = await StreamCollection.GetAsync(config);
+
+			// Create the TriggerNext step and mark it as complete
+			IGraph graphA = await GraphCollection.AddAsync(newTemplate1, null);
+			NewGroup groupA = new NewGroup("win", new List<NewNode> { new NewNode("TriggerNext") });
+			graphA = await GraphCollection.AppendAsync(graphA, new List<NewGroup> { groupA });
+
+			// Tick the schedule and make sure it doesn't trigger
+			await ScheduleService.TickForTestingAsync();
+			List<IJob> jobs2 = await GetNewJobs();
+			Assert.AreEqual(0, jobs2.Count);
+
+			// Create a job and fail it
+			CreateJobOptions options1 = new CreateJobOptions();
+			options1.PreflightChange = 999;
+			options1.Arguments.Add("-Target=TriggerNext");
+
+			IJob job1 = await JobService.CreateJobAsync(null, config, newTemplateRefId1, newTemplate1.Hash, graphA, "Hello", 1234, 1233, options1);
+			SubResourceId batchId1 = job1.Batches[0].Id;
+			SubResourceId stepId1 = job1.Batches[0].Steps[0].Id;
+			job1 = Deref(await JobService.UpdateBatchAsync(job1, batchId1, config, LogId.GenerateNewId(), JobStepBatchState.Running));
+			job1 = Deref(await JobService.UpdateStepAsync(job1, batchId1, stepId1, config, JobStepState.Completed, JobStepOutcome.Failure));
+			job1 = Deref(await JobService.UpdateBatchAsync(job1, batchId1, config, LogId.GenerateNewId(), JobStepBatchState.Complete));
+			Assert.IsNotNull(job1);
+			await GetNewJobs();
+
+			// Tick the schedule and make sure it doesn't trigger
+			await Clock.AdvanceAsync(TimeSpan.FromMinutes(30.0));
+			await ScheduleService.TickForTestingAsync();
+			List<IJob> jobs3 = await GetNewJobs();
+			Assert.AreEqual(0, jobs3.Count);
+
+			// Create a job and make it succeed
+			CreateJobOptions options2 = new CreateJobOptions();
+			options2.PreflightChange = 999;
+			options2.Arguments.Add("-Target=TriggerNext");
+
+			IJob job2 = await JobService.CreateJobAsync(null, config, newTemplateRefId1, newTemplate1.Hash, graphA, "Hello", 1234, 1233, options2);
+			SubResourceId batchId2 = job2.Batches[0].Id;
+			SubResourceId stepId2 = job2.Batches[0].Steps[0].Id;
+			job2 = Deref(await JobService.UpdateBatchAsync(job2, batchId2, config, LogId.GenerateNewId(), JobStepBatchState.Running));
+			job2 = Deref(await JobService.UpdateStepAsync(job2, batchId2, stepId2, config, JobStepState.Completed, JobStepOutcome.Success));
+			Assert.IsNotNull(job2);
+
+			// Tick the schedule and make sure it does trigger
+			await ScheduleService.TickForTestingAsync();
+			List<IJob> jobs4 = await GetNewJobs();
+			Assert.AreEqual(1, jobs4.Count);
+			Assert.AreEqual(1234, jobs4[0].Change);
+			Assert.AreEqual(1233, jobs4[0].CodeChange);
+		}
+
+		[TestMethod]
+		public async Task GateTest2Async()
+		{
+			IUser bob = await UserCollection.FindOrAddUserByLoginAsync("Bob");
+
+			DateTime startTime = new DateTime(2021, 1, 1, 12, 0, 0, DateTimeKind.Local); // Friday Jan 1, 2021 
+			Clock.UtcNow = startTime;
+
+			PerforceService.Changes.Clear();
+			PerforceService.AddChange(StreamId, 1230, bob, "", new[] { "code.cpp" });
+			PerforceService.AddChange(StreamId, 1231, bob, "", new[] { "content.uasset" });
+			PerforceService.AddChange(StreamId, 1232, bob, "", new[] { "content.uasset" });
+			PerforceService.AddChange(StreamId, 1233, bob, "", new[] { "code.cpp" });
+
+			// Create two templates, the second dependent on the first
+			TemplateId newTemplateRefId1 = new TemplateId("new-template-1");
+			TemplateRefConfig newTemplate1 = new TemplateRefConfig();
+			newTemplate1.Id = newTemplateRefId1;
+
+			TemplateId newTemplateRefId2 = new TemplateId("new-template-2");
+			TemplateRefConfig newTemplate2 = new TemplateRefConfig();
+			newTemplate2.Id = newTemplateRefId2;
+			newTemplate2.Name = "Test template 2";
+			newTemplate2.Schedule = new ScheduleConfig();
+			newTemplate2.Schedule.MaxChanges = 4;
+			newTemplate2.Schedule.Commits.Add(CommitTag.Code);
+			newTemplate2.Schedule.Gate = new ScheduleGateConfig { TemplateId = newTemplateRefId1, Target = "TriggerNext" };
+			newTemplate2.Schedule.Patterns.Add(new SchedulePatternConfig { Interval = 10 });// (null, 0, null, 10));
+//			NewTemplate2.Schedule.LastTriggerTime = StartTime;
+
+//			IStream? stream = await StreamService.GetStreamAsync(StreamId);
+
+			StreamConfig config = new StreamConfig();
+			config.Id = StreamId;
+			config.Name = "//UE5/Main";
+			config.Tabs.Add(new JobsTabConfig { Title = "foo", Templates = new List<TemplateId> { newTemplateRefId1, newTemplateRefId2 } });
+			config.Templates.Add(newTemplate1);
+			config.Templates.Add(newTemplate2);
+			UpdateConfig(x => x.Projects[0].Streams = new List<StreamConfig> { config });
+
+			IStream stream = await StreamCollection.GetAsync(config);
+
+//			stream = (await CreateOrReplaceStreamAsync(StreamId, stream, ProjectId, config))!;
+
+			ITemplate template1 = (await TemplateCollection.GetOrAddAsync(config.Templates[0]))!;
+			Assert.IsNotNull(template1);
+			ITemplate template2 = (await TemplateCollection.GetOrAddAsync(config.Templates[1]))!;
+			Assert.IsNotNull(template2);
+
+			// Create the graph
+			IGraph graphA = await GraphCollection.AddAsync(template1, null);
+			NewGroup groupA = new NewGroup("win", new List<NewNode> { new NewNode("TriggerNext") });
+			graphA = await GraphCollection.AppendAsync(graphA, new List<NewGroup> { groupA });
+
+			// Create successful jobs for all the changes we added above
+			for (int change = 1230; change <= 1233; change++)
+			{
+				int codeChange = (change < 1233) ? 1230 : 1233;
+
+				CreateJobOptions options1 = new CreateJobOptions();
+				options1.Arguments.Add("-Target=TriggerNext");
+
+				IJob job1 = await JobService.CreateJobAsync(null, config, newTemplateRefId1, _template.Hash, graphA, "Hello", change, codeChange, options1);
+				for (int batchIdx = 0; batchIdx < job1.Batches.Count; batchIdx++)
+				{
+					SubResourceId batchId1 = job1.Batches[batchIdx].Id;
+					job1 = Deref(await JobService.UpdateBatchAsync(job1, batchId1, config, LogId.GenerateNewId(), JobStepBatchState.Running));
+					for (int stepIdx = 0; stepIdx < job1.Batches[batchIdx].Steps.Count; stepIdx++)
+					{
+						SubResourceId stepId1 = job1.Batches[batchIdx].Steps[stepIdx].Id;
+						job1 = Deref(await JobService.UpdateStepAsync(job1, batchId1, stepId1, config, JobStepState.Completed, JobStepOutcome.Success, newLogId: LogId.GenerateNewId()));
+					}
+					job1 = Deref(await JobService.UpdateBatchAsync(job1, batchId1, config, LogId.GenerateNewId(), JobStepBatchState.Complete));
+				}
+			}
+			await GetNewJobs();
+
+			// Tick the schedule and make sure it doesn't trigger
+			await Clock.AdvanceAsync(TimeSpan.FromMinutes(30.0));
+			await ScheduleService.TriggerAsync(StreamId, newTemplateRefId2, Clock.UtcNow, default);
+			List<IJob> jobs3 = await GetNewJobs();
+			Assert.AreEqual(2, jobs3.Count);
+			Assert.AreEqual(1230, jobs3[0].Change);
+			Assert.AreEqual(1233, jobs3[1].Change);
+		}
+
+		[TestMethod]
+		public async Task UpdateConfigAsync()
+		{
+			DateTime startTime = new DateTime(2021, 1, 1, 12, 0, 0, DateTimeKind.Local); // Friday Jan 1, 2021 
+			Clock.UtcNow = startTime;
+			Clock.TimeZone = TimeZoneInfo.Local;
+
+			ScheduleConfig schedule = new ScheduleConfig();
+			schedule.Enabled = true;
+			schedule.RequireSubmittedChange = false;
+			schedule.Patterns.Add(new SchedulePatternConfig { MinTime = 13 * 60, MaxTime = 14 * 60, Interval = 15 });
+			schedule.MaxActive = 2;
+			await SetScheduleAsync(schedule);
+
+			await Clock.AdvanceAsync(TimeSpan.FromHours(1.25));
+			await ScheduleService.TickForTestingAsync();
+
+			List<IJob> jobs1 = await GetNewJobs();
+			Assert.AreEqual(1, jobs1.Count);
+			Assert.AreEqual(102, jobs1[0].Change);
+			Assert.AreEqual(100, jobs1[0].CodeChange);
+
+			// Make sure the job is registered
+			IStream? stream1 = await StreamCollection.GetAsync(GlobalConfig.CurrentValue.Streams[0]);
+			ITemplateRef templateRef1 = stream1!.Templates.First().Value;
+			Assert.AreEqual(1, templateRef1.Schedule!.ActiveJobs.Count);
+			Assert.AreEqual(jobs1[0].Id, templateRef1.Schedule!.ActiveJobs[0]);
+
+			// Test that another job does not trigger
+			await SetScheduleAsync(schedule);
+
+			// Make sure the job is still registered
+			IStream? stream2 = await StreamCollection.GetAsync(GlobalConfig.CurrentValue.Streams[0]);
+			ITemplateRef templateRef2 = stream2!.Templates.First().Value;
+			Assert.AreEqual(1, templateRef2.Schedule!.ActiveJobs.Count);
+			Assert.AreEqual(jobs1[0].Id, templateRef2.Schedule!.ActiveJobs[0]);
+		}
+
+		[TestMethod]
+		public async Task StreamPausing()
+		{
+			DateTime startTime = new DateTime(2021, 1, 1, 12, 0, 0, DateTimeKind.Utc); // Friday Jan 1, 2021 
+			Clock.UtcNow = startTime;
+
+			ScheduleConfig schedule = new ScheduleConfig();
+			schedule.Enabled = true;
+			schedule.Patterns.Add(new SchedulePatternConfig { MinTime = 13 * 60, MaxTime = 14 * 60, Interval = 15 });
+			IStream stream = await SetScheduleAsync(schedule);
+
+			IStreamCollection streamCollection = ServiceProvider.GetRequiredService<IStreamCollection>();
+			await streamCollection.TryUpdatePauseStateAsync(stream, newPausedUntil: startTime.AddHours(5), newPauseComment: "testing");
+
+			// Try trigger a job. No job should be scheduled as the stream is paused
+			await Clock.AdvanceAsync(TimeSpan.FromHours(1.25));
+			await ScheduleService.TickForTestingAsync();
+			List<IJob> jobs2 = await GetNewJobs();
+			Assert.AreEqual(0, jobs2.Count);
+
+			// Advance time beyond the pause period. A build should now trigger
+			await Clock.AdvanceAsync(TimeSpan.FromHours(5.25));
+			await ScheduleService.TickForTestingAsync();
+
+			List<IJob> jobs3 = await GetNewJobs();
+			Assert.AreEqual(1, jobs3.Count);
+			Assert.AreEqual(102, jobs3[0].Change);
+			Assert.AreEqual(100, jobs3[0].CodeChange);
+		}
+
+		async Task<List<IJob>> GetNewJobs()
+		{
+			List<IJob> jobs = await JobCollection.FindAsync();
+			jobs.RemoveAll(x => _initialJobIds.Contains(x.Id));
+			_initialJobIds.UnionWith(jobs.Select(x => x.Id));
+			return jobs.OrderBy(x => x.Change).ToList();
+		}
+	}
+}
diff --git a/Engine/Source/Programs/Horde/Horde.Server.Tests/ServiceAccountAuthTest.cs b/Engine/Source/Programs/Horde/Horde.Server.Tests/ServiceAccountAuthTest.cs
new file mode 100644
--- /dev/null
+++ b/Engine/Source/Programs/Horde/Horde.Server.Tests/ServiceAccountAuthTest.cs
@@ -0,0 +1,111 @@
+﻿// Copyright Epic Games, Inc. All Rights Reserved.
+
+using System.Collections.Generic;
+using System.Linq;
+using System.Security.Claims;
+using System.Threading.Tasks;
+using Horde.Server.Authentication;
+using Horde.Server.Users;
+using Microsoft.AspNetCore.Authentication;
+using Microsoft.AspNetCore.Http;
+using Microsoft.Extensions.DependencyInjection;
+using Microsoft.Extensions.Logging;
+using Microsoft.Extensions.Primitives;
+using Microsoft.Extensions.WebEncoders.Testing;
+using Microsoft.Net.Http.Headers;
+using Microsoft.VisualStudio.TestTools.UnitTesting;
+
+namespace Horde.Server.Tests
+{
+	[TestClass]
+	public class ServiceAccountAuthTest : DatabaseIntegrationTest
+	{
+		public ServiceAccountAuthTest()
+		{
+ 		}
+
+		protected override void ConfigureServices(IServiceCollection services)
+		{
+			base.ConfigureServices(services);
+
+			services.AddSingleton<ILoggerFactory, LoggerFactory>();
+			services.AddSingleton<ServiceAccountCollection>();
+			services.AddSingleton<IServiceAccountCollection>(sp => sp.GetRequiredService<ServiceAccountCollection>());
+		}
+
+		private async Task<ServiceAccountAuthHandler> GetAuthHandlerAsync(string? headerValue)
+		{
+			ServiceAccountAuthOptions options = new ServiceAccountAuthOptions();
+
+			ServiceAccountCollection serviceAccounts = ServiceProvider.GetRequiredService<ServiceAccountCollection>();
+			await serviceAccounts.AddAsync("mytoken", new List<string> { "myclaim###myvalue", "foo###bar" }, "mydesc");
+
+			ILoggerFactory loggerFactory = ServiceProvider.GetRequiredService<ILoggerFactory>();
+			ServiceAccountAuthHandler handler = new ServiceAccountAuthHandler(new TestOptionsMonitor<ServiceAccountAuthOptions>(options), loggerFactory, new UrlTestEncoder(), new SystemClock(), serviceAccounts);
+			AuthenticationScheme scheme = new AuthenticationScheme(ServiceAccountAuthHandler.AuthenticationScheme, "ServiceAccountAuth", handler.GetType());
+			
+			HttpContext httpContext = new DefaultHttpContext();
+			if (headerValue != null)
+			{
+				httpContext.Request.Headers.Add(HeaderNames.Authorization, new StringValues(headerValue));	
+			}
+			
+			await handler.InitializeAsync(scheme, httpContext);
+
+			return handler;
+		}
+
+		[TestMethod]
+		public async Task ValidToken()
+		{
+			ServiceAccountAuthHandler handler = await GetAuthHandlerAsync("ServiceAccount mytoken");
+			AuthenticateResult result = await handler.AuthenticateAsync();
+			Assert.IsTrue(result.Succeeded);
+			
+			handler = await GetAuthHandlerAsync("ServiceAccount   mytoken        ");
+			result = await handler.AuthenticateAsync();
+			Assert.IsTrue(result.Succeeded);
+		}
+		
+		[TestMethod]
+		public async Task InvalidToken()
+		{
+			ServiceAccountAuthHandler handler = await GetAuthHandlerAsync("ServiceAccount doesNotExist");
+			AuthenticateResult result = await handler.AuthenticateAsync();
+			Assert.IsFalse(result.Succeeded);
+			
+			// Valid token but bad prefix
+			handler = await GetAuthHandlerAsync("SriceAcct mytoken");
+			result = await handler.AuthenticateAsync();
+			Assert.IsFalse(result.Succeeded);
+		}
+		
+		[TestMethod]
+		public async Task NoResult()
+		{
+			// Valid token but bad prefix
+			ServiceAccountAuthHandler handler = await GetAuthHandlerAsync("Bogus mytoken");
+			AuthenticateResult result = await handler.AuthenticateAsync();
+			Assert.IsFalse(result.Succeeded);
+			Assert.IsTrue(result.None);
+			
+			// No Authorization header at all
+			handler = await GetAuthHandlerAsync(null);
+			result = await handler.AuthenticateAsync();
+			Assert.IsFalse(result.Succeeded);
+			Assert.IsTrue(result.None);
+		}
+		
+		[TestMethod]
+		public async Task Claims()
+		{
+			ServiceAccountAuthHandler handler = await GetAuthHandlerAsync("ServiceAccount mytoken");
+			AuthenticateResult result = await handler.AuthenticateAsync();
+			Assert.IsTrue(result.Succeeded);
+			Assert.AreEqual(3, result.Ticket!.Principal.Claims.Count());
+			Assert.AreEqual(ServiceAccountAuthHandler.AuthenticationScheme, result.Ticket.Principal.FindFirst(ClaimTypes.Name)!.Value);
+			Assert.AreEqual("myvalue", result.Ticket!.Principal.FindFirst("myclaim")!.Value);
+			Assert.AreEqual("bar", result.Ticket!.Principal.FindFirst("foo")!.Value);
+		}
+	}
+}
\ No newline at end of file
diff --git a/Engine/Source/Programs/Horde/Horde.Server.Tests/ServiceAccountCollectionTests.cs b/Engine/Source/Programs/Horde/Horde.Server.Tests/ServiceAccountCollectionTests.cs
new file mode 100644
--- /dev/null
+++ b/Engine/Source/Programs/Horde/Horde.Server.Tests/ServiceAccountCollectionTests.cs
@@ -0,0 +1,73 @@
+﻿// Copyright Epic Games, Inc. All Rights Reserved.
+
+using System.Collections.Generic;
+using System.Threading.Tasks;
+using Horde.Server.Server;
+using Horde.Server.Users;
+using Microsoft.VisualStudio.TestTools.UnitTesting;
+using MongoDB.Bson;
+
+namespace Horde.Server.Tests
+{
+	[TestClass]
+	public class ServiceAccountCollectionTests : DatabaseIntegrationTest
+	{
+		private readonly IServiceAccountCollection _serviceAccounts;
+		private readonly IServiceAccount _serviceAccount;
+		
+		public ServiceAccountCollectionTests()
+		{
+			MongoService mongoService = GetMongoServiceSingleton();
+			_serviceAccounts = new ServiceAccountCollection(mongoService);
+			_serviceAccount = _serviceAccounts.AddAsync(ObjectId.GenerateNewId().ToString(), new List<string> {"myclaim###myvalue"}, "mydesc").Result;
+		}
+
+		[TestMethod]
+		public async Task Add()
+		{
+			IServiceAccount sa = await _serviceAccounts.AddAsync("addToken", new List<string> {"myclaim###myvalue"}, "mydesc");
+			Assert.AreEqual("addToken", sa.SecretToken);
+			Assert.AreEqual(1, sa.GetClaims().Count);
+			Assert.AreEqual("myvalue", sa.GetClaims()[0].Value);
+			Assert.IsTrue(sa.Enabled);
+			Assert.AreEqual("mydesc", sa.Description);
+		}
+		
+		[TestMethod]
+		public async Task Get()
+		{
+			IServiceAccount sa = (await _serviceAccounts.GetAsync(_serviceAccount.Id))!;
+			Assert.AreEqual(_serviceAccount, sa);
+		}
+		
+		[TestMethod]
+		public async Task GetBySecretToken()
+		{
+			IServiceAccount sa = (await _serviceAccounts.GetBySecretTokenAsync(_serviceAccount.SecretToken))!;
+			Assert.AreEqual(_serviceAccount, sa);
+		}
+		
+		[TestMethod]
+		public async Task Update()
+		{
+			List<string> newClaims = new List<string> {{"newclaim1###newvalue1"}, {"newclaim2###newvalue2"}};
+			await _serviceAccounts.UpdateAsync(_serviceAccount.Id, "newtoken", newClaims, false, "newdesc");
+			IServiceAccount getSa = (await _serviceAccounts.GetAsync(_serviceAccount.Id))!;
+			
+			Assert.AreEqual("newtoken", getSa.SecretToken);
+			Assert.AreEqual(2, getSa.GetClaims().Count);
+			Assert.AreEqual("newvalue1", getSa.GetClaims()[0].Value);
+			Assert.AreEqual("newvalue2", getSa.GetClaims()[1].Value);
+			Assert.AreEqual(false, getSa.Enabled);
+			Assert.AreEqual("newdesc", getSa.Description);
+		}
+		
+		[TestMethod]
+		public async Task Delete()
+		{
+			await _serviceAccounts.DeleteAsync(_serviceAccount.Id);
+			IServiceAccount? result = await _serviceAccounts.GetAsync(_serviceAccount.Id);
+			Assert.IsNull(result);
+		}
+	}
+}
\ No newline at end of file
diff --git a/Engine/Source/Programs/Horde/Horde.Server.Tests/StreamServiceTest.cs b/Engine/Source/Programs/Horde/Horde.Server.Tests/StreamServiceTest.cs
new file mode 100644
--- /dev/null
+++ b/Engine/Source/Programs/Horde/Horde.Server.Tests/StreamServiceTest.cs
@@ -0,0 +1,39 @@
+// Copyright Epic Games, Inc. All Rights Reserved.
+
+using System;
+using System.Globalization;
+using System.Threading.Tasks;
+using Horde.Server.Streams;
+using Microsoft.VisualStudio.TestTools.UnitTesting;
+
+namespace Horde.Server.Tests
+{
+	[TestClass]
+    public class StreamServiceTests : TestSetup
+    {
+        [TestMethod]
+        public async Task Pausing()
+        {
+			Fixture fixture = await CreateFixtureAsync();
+
+	        IStream stream = (await StreamCollection.GetAsync(fixture!.StreamConfig!))!;
+	        Assert.IsFalse(stream.IsPaused(DateTime.UtcNow));
+	        Assert.IsNull(stream.PausedUntil);
+	        Assert.IsNull(stream.PauseComment);
+
+	        DateTime pausedUntil = DateTime.UtcNow.AddHours(1);
+	        await StreamCollection.UpdatePauseStateAsync(stream, newPausedUntil: pausedUntil, newPauseComment: "mycomment");
+	        stream = (await StreamCollection.GetAsync(fixture!.StreamConfig!))!;
+	        // Comparing by string to avoid comparing exact milliseconds as those are not persisted in MongoDB fields
+	        Assert.IsTrue(stream.IsPaused(DateTime.UtcNow));
+	        Assert.AreEqual(pausedUntil.ToString(CultureInfo.InvariantCulture), stream.PausedUntil!.Value.ToString(CultureInfo.InvariantCulture));
+	        Assert.AreEqual("mycomment", stream.PauseComment);
+	        
+	        await StreamCollection.UpdatePauseStateAsync(stream, newPausedUntil: null, newPauseComment: null);
+	        stream = (await StreamCollection.GetAsync(fixture!.StreamConfig!))!;
+	        Assert.IsFalse(stream.IsPaused(DateTime.UtcNow));
+	        Assert.IsNull(stream.PausedUntil);
+	        Assert.IsNull(stream.PauseComment);
+        }
+    }
+}
