Horde: Rename the Horde.Build project to Horde.Server.

diff --git a/Engine/Source/Programs/Horde/Horde.Server/Agents/Fleet/JobQueueStrategy.cs b/Engine/Source/Programs/Horde/Horde.Server/Agents/Fleet/JobQueueStrategy.cs
new file mode 100644
--- /dev/null
+++ b/Engine/Source/Programs/Horde/Horde.Server/Agents/Fleet/JobQueueStrategy.cs
@@ -0,0 +1,236 @@
+// Copyright Epic Games, Inc. All Rights Reserved.
+
+using System;
+using System.Collections.Generic;
+using System.Linq;
+using System.Text;
+using System.Text.Json.Serialization;
+using System.Threading.Tasks;
+using Horde.Server.Agents.Pools;
+using Horde.Server.Jobs;
+using Horde.Server.Jobs.Graphs;
+using Horde.Server.Server;
+using Horde.Server.Streams;
+using Horde.Server.Utilities;
+using HordeCommon;
+using Microsoft.Extensions.Caching.Memory;
+using Microsoft.Extensions.Options;
+using OpenTracing;
+using OpenTracing.Util;
+
+namespace Horde.Server.Agents.Fleet
+{
+	/// <summary>
+	/// Job queue sizing settings for a pool
+	/// </summary>
+	public class JobQueueSettings
+	{
+		/// <summary>
+		/// Factor translating queue size to additional agents to grow the pool with
+		/// The result is always rounded up to nearest integer. 
+		/// Example: if there are 20 jobs in queue, a factor 0.25 will result in 5 new agents being added (20 * 0.25)
+		/// </summary>
+		public double ScaleOutFactor { get; set;  } = 0.25;
+		
+		/// <summary>
+		/// Factor by which to shrink the pool size with when queue is empty
+		/// The result is always rounded up to nearest integer.
+		/// Example: when the queue size is zero, a default value of 0.9 will shrink the pool by 10% (current agent count * 0.9)
+		/// </summary>
+		public double ScaleInFactor { get; set; } = 0.9;
+
+		/// <summary>
+		/// How far back in time to look for job batches (that potentially are in the queue)
+		/// </summary>
+		public int SamplePeriodMin { get; set; } = 60 * 24 * 5; // 5 days
+
+		/// <summary>
+		/// Time spent in ready state before considered truly waiting for an agent
+		///
+		/// A job batch can be in ready state before getting picked up and executed.
+		/// This threshold will help ensure only batches that have been waiting longer than this value will be considered.
+		/// </summary>
+		public int ReadyTimeThresholdSec { get; set; } = 45;
+
+		/// <summary>
+		/// Constructor used for JSON serialization
+		/// </summary>
+		[JsonConstructor]
+		public JobQueueSettings()
+		{
+		}
+		
+		/// <summary>
+		/// Constructor
+		/// </summary>
+		/// <param name="scaleOutFactor"></param>
+		/// <param name="scaleInFactor"></param>
+		public JobQueueSettings(double? scaleOutFactor = null, double? scaleInFactor = null)
+		{
+			ScaleOutFactor = scaleOutFactor.GetValueOrDefault(ScaleOutFactor);
+			ScaleInFactor = scaleInFactor.GetValueOrDefault(ScaleInFactor);
+		}
+		
+		/// <inheritdoc />
+		public override string ToString()
+		{
+			StringBuilder sb = new (50);
+			sb.AppendFormat("{0}={1} ", nameof(ScaleOutFactor), ScaleOutFactor);
+			sb.AppendFormat("{0}={1} ", nameof(ScaleInFactor), ScaleInFactor);
+			sb.AppendFormat("{0}={1} ", nameof(SamplePeriodMin), SamplePeriodMin);
+			sb.AppendFormat("{0}={1} ", nameof(ReadyTimeThresholdSec), ReadyTimeThresholdSec);
+			return sb.ToString();
+		}
+	}
+	
+	/// <summary>
+	/// Calculate pool size by observing the number of jobs in waiting state
+	///
+	/// Allows for more proactive scaling compared to LeaseUtilizationStrategy.
+	/// <see cref="LeaseUtilizationStrategy"/> 
+	/// </summary>
+	public class JobQueueStrategy : IPoolSizeStrategy
+	{
+		private const string CacheKey = nameof(JobQueueStrategy);
+		internal JobQueueSettings Settings { get; }
+		
+		private readonly IJobCollection _jobs;
+		private readonly IGraphCollection _graphs;
+		private readonly IStreamCollection _streamCollection;
+		private readonly IClock _clock;
+		private readonly IMemoryCache _cache;
+		private readonly IOptionsMonitor<GlobalConfig> _globalConfig;
+
+		/// <summary>
+		/// Constructor
+		/// </summary>
+		/// <param name="jobs"></param>
+		/// <param name="graphs"></param>
+		/// <param name="streamCollection"></param>
+		/// <param name="clock"></param>
+		/// <param name="cache"></param>
+		/// <param name="globalConfig"></param>
+		/// <param name="settings"></param>
+		public JobQueueStrategy(IJobCollection jobs, IGraphCollection graphs, IStreamCollection streamCollection, IClock clock, IMemoryCache cache, IOptionsMonitor<GlobalConfig> globalConfig, JobQueueSettings? settings = null)
+		{
+			_jobs = jobs;
+			_graphs = graphs;
+			_streamCollection = streamCollection;
+			_clock = clock;
+			_cache = cache;
+			_globalConfig = globalConfig;
+			Settings = settings ?? new JobQueueSettings();
+		}
+
+		/// <inheritdoc/>
+		public string Name { get; } = "JobQueue";
+
+		/// <summary>
+		/// Extract all job step batches from a job, with their associated pool 
+		/// </summary>
+		/// <param name="job">Job to extract from</param>
+		/// <param name="streams">Cached lookup table of streams</param>
+		/// <returns></returns>
+		private async Task<List<(IJob Job, IJobStepBatch Batch, PoolId PoolId)>> GetJobBatchesWithPools(IJob job, Dictionary<StreamId, StreamConfig> streams)
+		{
+			IGraph graph = await _graphs.GetAsync(job.GraphHash);
+
+			List<(IJob Job, IJobStepBatch Batch, PoolId PoolId)> jobBatches = new();
+			foreach (IJobStepBatch batch in job.Batches)
+			{
+				if (batch.State != JobStepBatchState.Ready)
+				{
+					continue;
+				}
+				
+				TimeSpan? waitTime = _clock.UtcNow - batch.ReadyTimeUtc;
+				if (waitTime == null)
+				{
+					continue;
+				}
+				if (waitTime.Value < TimeSpan.FromSeconds(Settings.ReadyTimeThresholdSec))
+				{
+					continue;
+				}
+
+				if (!streams.TryGetValue(job.StreamId, out StreamConfig? streamConfig))
+				{
+					continue;
+				}
+
+				string batchAgentType = graph.Groups[batch.GroupIdx].AgentType;
+				if (!streamConfig.AgentTypes.TryGetValue(batchAgentType, out AgentConfig? agentType))
+				{
+					continue;
+				}
+				
+				jobBatches.Add((job, batch, agentType.Pool));
+			}
+
+			return jobBatches;
+		}
+
+		internal async Task<Dictionary<PoolId, int>> GetPoolQueueSizesAsync(DateTimeOffset jobsCreatedAfter)
+		{
+			using IScope scope = GlobalTracer.Instance.BuildSpan("JobQueueStrategy.GetPoolQueueSizes").StartActive();
+
+			Dictionary<StreamId, StreamConfig> streams = _globalConfig.CurrentValue.Streams.ToDictionary(x => x.Id, x => (StreamConfig)x);
+			List<IJob> recentJobs = await _jobs.FindAsync(minCreateTime: jobsCreatedAfter);
+
+			List<(IJob Job, IJobStepBatch Batch, PoolId PoolId)> jobBatches = new();
+			foreach (IJob job in recentJobs)
+			{
+				jobBatches.AddRange(await GetJobBatchesWithPools(job, streams));
+			}
+
+			List<(PoolId PoolId, int QueueSize)> poolsWithQueueSize = jobBatches.GroupBy(t => t.PoolId).Select(t => (t.Key, t.Count())).ToList();
+
+			scope.Span.SetTag("NumPools", poolsWithQueueSize.Count);
+			return poolsWithQueueSize.ToDictionary(x => x.PoolId, x => x.QueueSize);
+		}
+
+		/// <inheritdoc/>
+		public async Task<PoolSizeResult> CalculatePoolSizeAsync(IPool pool, List<IAgent> agents)
+		{
+			using IScope scope = GlobalTracer.Instance
+				.BuildSpan("JobQueueStrategy.CalculatePoolSize")
+				.WithTag(Datadog.Trace.OpenTracing.DatadogTags.ResourceName, pool.Id.ToString())
+				.WithTag("CurrentAgentCount", agents.Count)
+				.StartActive();
+			
+			DateTimeOffset minCreateTime = _clock.UtcNow - TimeSpan.FromMinutes(Settings.SamplePeriodMin);
+
+			// Cache pool queue sizes for a short while for faster runs when many pools are scaled
+			if (!_cache.TryGetValue(CacheKey, out Dictionary<PoolId, int> poolQueueSizes))
+			{
+				// Pool sizes haven't been cached, update them (might happen from multiple tasks but that is fine)
+				poolQueueSizes = await GetPoolQueueSizesAsync(minCreateTime);
+				_cache.Set(CacheKey, poolQueueSizes, TimeSpan.FromSeconds(60));
+			}
+
+			poolQueueSizes.TryGetValue(pool.Id, out int queueSize);
+			
+			Dictionary<string, object> status = new()
+			{
+				["Name"] = GetType().Name,
+				["QueueSize"] = queueSize,
+				["ScaleOutFactor"] = Settings.ScaleOutFactor,
+				["ScaleInFactor"] = Settings.ScaleInFactor,
+				["SamplePeriodMin"] = Settings.SamplePeriodMin,
+				["ReadyTimeThresholdSec"] = Settings.ReadyTimeThresholdSec,
+			};
+			
+			if (queueSize > 0)
+			{
+				int additionalAgentCount = (int)Math.Ceiling(queueSize * Settings.ScaleOutFactor);
+				int desiredAgentCount = agents.Count + additionalAgentCount;
+				return new PoolSizeResult(agents.Count, desiredAgentCount, status);
+			}
+			else
+			{
+				int desiredAgentCount = (int)(agents.Count * Settings.ScaleInFactor);
+				return new PoolSizeResult(agents.Count, desiredAgentCount, status);
+			}
+		}
+	}
+}
diff --git a/Engine/Source/Programs/Horde/Horde.Server/Agents/IAgent.cs b/Engine/Source/Programs/Horde/Horde.Server/Agents/IAgent.cs
new file mode 100644
--- /dev/null
+++ b/Engine/Source/Programs/Horde/Horde.Server/Agents/IAgent.cs
@@ -0,0 +1,869 @@
+// Copyright Epic Games, Inc. All Rights Reserved.
+
+using System;
+using System.Collections.Generic;
+using System.Diagnostics.CodeAnalysis;
+using System.Linq;
+using System.Security.Cryptography;
+using System.Threading.Tasks;
+using EpicGames.Core;
+using EpicGames.Horde.Common;
+using EpicGames.Horde.Compute;
+using Google.Protobuf.WellKnownTypes;
+using Horde.Server.Acls;
+using Horde.Server.Agents.Leases;
+using Horde.Server.Agents.Pools;
+using Horde.Server.Agents.Sessions;
+using Horde.Server.Agents.Software;
+using Horde.Server.Logs;
+using Horde.Server.Perforce;
+using Horde.Server.Server;
+using Horde.Server.Streams;
+using Horde.Server.Tools;
+using Horde.Server.Utilities;
+using HordeCommon;
+using HordeCommon.Rpc.Tasks;
+using MongoDB.Bson;
+using MongoDB.Bson.Serialization.Attributes;
+
+namespace Horde.Server.Agents
+{
+	/// <summary>
+	/// Information about a workspace synced to an agent
+	/// </summary>
+	public class AgentWorkspace
+	{
+		/// <summary>
+		/// Name of the Perforce cluster to use
+		/// </summary>
+		public string? Cluster { get; set; }
+
+		/// <summary>
+		/// User to log into Perforce with (eg. buildmachine)
+		/// </summary>
+		public string? UserName { get; set; }
+
+		/// <summary>
+		/// Identifier to distinguish this workspace from other workspaces
+		/// </summary>
+		public string Identifier { get; set; }
+
+		/// <summary>
+		/// The stream to sync
+		/// </summary>
+		public string Stream { get; set; }
+
+		/// <summary>
+		/// Custom view for the workspace
+		/// </summary>
+		public List<string>? View { get; set; }
+
+		/// <summary>
+		/// Whether to use an incremental workspace
+		/// </summary>
+		public bool Incremental { get; set; }
+		
+		/// <summary>
+		/// Method to use when syncing/materializing data from Perforce
+		/// </summary>
+		public string? Method { get; set; }
+
+		/// <summary>
+		/// Constructor
+		/// </summary>
+		/// <param name="cluster">Name of the Perforce cluster</param>
+		/// <param name="userName">User to log into Perforce with (eg. buildmachine)</param>
+		/// <param name="identifier">Identifier to distinguish this workspace from other workspaces</param>
+		/// <param name="stream">The stream to sync</param>
+		/// <param name="view">Custom view for the workspace</param>
+		/// <param name="incremental">Whether to use an incremental workspace</param>
+		/// <param name="method">Method to use when syncing/materializing data from Perforce</param>
+		public AgentWorkspace(string? cluster, string? userName, string identifier, string stream, List<string>? view, bool incremental, string? method)
+		{
+			if (!String.IsNullOrEmpty(cluster))
+			{
+				Cluster = cluster;
+			}
+			if (!String.IsNullOrEmpty(userName))
+			{
+				UserName = userName;
+			}
+			Identifier = identifier;
+			Stream = stream;
+			View = view;
+			Incremental = incremental;
+			Method = method;
+		}
+
+		/// <summary>
+		/// Constructor
+		/// </summary>
+		/// <param name="workspace">RPC message to construct from</param>
+		public AgentWorkspace(HordeCommon.Rpc.Messages.AgentWorkspace workspace)
+			: this(workspace.ConfiguredCluster, workspace.ConfiguredUserName, workspace.Identifier, workspace.Stream, (workspace.View.Count > 0) ? workspace.View.ToList() : null, workspace.Incremental, workspace.Method)
+		{
+		}
+
+		/// <summary>
+		/// Gets a digest of the settings for this workspace
+		/// </summary>
+		/// <returns>Digest for the workspace settings</returns>
+		public string GetDigest()
+		{
+#pragma warning disable CA5351 // Do Not Use Broken Cryptographic Algorithms
+			using (MD5 hasher = MD5.Create())
+			{
+				byte[] data = BsonExtensionMethods.ToBson(this);
+				return BitConverter.ToString(hasher.ComputeHash(data)).Replace("-", "", StringComparison.Ordinal);
+			}
+#pragma warning restore CA5351
+		}
+
+		/// <inheritdoc/>
+		public override bool Equals(object? obj)
+		{
+			AgentWorkspace? other = obj as AgentWorkspace;
+			if (other == null)
+			{
+				return false;
+			}
+			if (Cluster != other.Cluster || UserName != other.UserName || Identifier != other.Identifier || Stream != other.Stream || Incremental != other.Incremental)
+			{
+				return false;
+			}
+			if (!Enumerable.SequenceEqual(View ?? new List<string>(), other.View ?? new List<string>()))
+			{
+				return false;
+			}
+			return true;
+		}
+
+		/// <inheritdoc/>
+		public override int GetHashCode()
+		{
+			return HashCode.Combine(Cluster, UserName, Identifier, Stream, Incremental); // Ignore 'View' for now
+		}
+
+		/// <summary>
+		/// Checks if two workspace sets are equivalent, ignoring order
+		/// </summary>
+		/// <param name="workspacesA">First list of workspaces</param>
+		/// <param name="workspacesB">Second list of workspaces</param>
+		/// <returns>True if the sets are equivalent</returns>
+		public static bool SetEquals(IReadOnlyList<AgentWorkspace> workspacesA, IReadOnlyList<AgentWorkspace> workspacesB)
+		{
+			HashSet<AgentWorkspace> workspacesSetA = new HashSet<AgentWorkspace>(workspacesA);
+			return workspacesSetA.SetEquals(workspacesB);
+		}
+
+		/// <summary>
+		/// Converts this workspace to an RPC message
+		/// </summary>
+		/// <param name="server">The Perforce server</param>
+		/// <param name="credentials">Credentials for the server</param>
+		/// <returns>The RPC message</returns>
+		public HordeCommon.Rpc.Messages.AgentWorkspace ToRpcMessage(IPerforceServer server, PerforceCredentials? credentials)
+		{
+			// Construct the message
+			HordeCommon.Rpc.Messages.AgentWorkspace result = new ()
+			{
+				ConfiguredCluster = Cluster, ConfiguredUserName = UserName,
+				ServerAndPort = server.ServerAndPort,
+				UserName = credentials?.UserName ?? UserName,
+				Password = credentials?.Password,
+				Identifier = Identifier,
+				Stream = Stream,
+				Incremental = Incremental,
+				Method = Method ?? String.Empty
+			};
+			
+			if (View != null)
+			{
+				result.View.AddRange(View);
+			}
+			
+			return result;
+		}
+	}
+
+	/// <summary>
+	/// Document describing an active lease
+	/// </summary>
+	public class AgentLease
+	{
+		/// <summary>
+		/// Name of this lease
+		/// </summary>
+		[BsonRequired]
+		public LeaseId Id { get; set; }
+
+		/// <summary>
+		/// Name of this lease
+		/// </summary>
+		public string Name { get; set; }
+
+		/// <summary>
+		/// The current state of the lease
+		/// </summary>
+		public LeaseState State { get; set; }
+
+		/// <summary>
+		/// The stream for this lease
+		/// </summary>
+		public StreamId? StreamId { get; set; }
+
+		/// <summary>
+		/// The pool for this lease
+		/// </summary>
+		public PoolId? PoolId { get; set; }
+
+		/// <summary>
+		/// Optional log for this lease
+		/// </summary>
+		public LogId? LogId { get; set; }
+
+		/// <summary>
+		/// Time at which the lease started
+		/// </summary>
+		[BsonRequired]
+		public DateTime StartTime { get; set; }
+
+		/// <summary>
+		/// Time at which the lease should be terminated
+		/// </summary>
+		public DateTime? ExpiryTime { get; set; }
+
+		/// <summary>
+		/// Flag indicating whether this lease has been accepted by the agent
+		/// </summary>
+		public bool Active { get; set; }
+
+		/// <summary>
+		/// Resources used by this lease
+		/// </summary>
+		public IReadOnlyDictionary<string, int>? Resources { get; set; }
+
+		/// <summary>
+		/// Whether the lease requires exclusive access to the agent
+		/// </summary>
+		public bool Exclusive { get; set; }
+
+		/// <summary>
+		/// For leases in the pending state, encodes an "any" protobuf containing the payload for the agent to execute the lease.
+		/// </summary>
+		public byte[]? Payload { get; set; }
+
+		/// <summary>
+		/// Private constructor
+		/// </summary>
+		[BsonConstructor]
+		private AgentLease()
+		{
+			Name = String.Empty;
+		}
+
+		/// <summary>
+		/// Constructor
+		/// </summary>
+		/// <param name="id">Identifier for the lease</param>
+		/// <param name="name">Name of this lease</param>
+		/// <param name="streamId"></param>
+		/// <param name="poolId"></param>
+		/// <param name="logId">Unique id for the log</param>
+		/// <param name="state">State for the lease</param>
+		/// <param name="resources">Resources required for this lease</param>
+		/// <param name="exclusive">Whether to reserve the entire device</param>
+		/// <param name="payload">Encoded "any" protobuf describing the contents of the payload</param>
+		public AgentLease(LeaseId id, string name, StreamId? streamId, PoolId? poolId, LogId? logId, LeaseState state, IReadOnlyDictionary<string, int>? resources, bool exclusive, byte[]? payload)
+		{
+			Id = id;
+			Name = name;
+			StreamId = streamId;
+			PoolId = poolId;
+			LogId = logId;
+			State = state;
+			Resources = resources;
+			Exclusive = exclusive;
+			Payload = payload;
+			StartTime = DateTime.UtcNow;
+		}
+
+		/// <summary>
+		/// Determines if this is a conform lease
+		/// </summary>
+		/// <returns>True if this is a conform lease</returns>
+		public bool IsConformLease()
+		{
+			if (Payload != null)
+			{
+				Any basePayload = Any.Parser.ParseFrom(Payload);
+				if (basePayload.Is(ConformTask.Descriptor))
+				{
+					return true;
+				}
+			}
+			return false;
+		}
+
+		/// <summary>
+		/// Converts this lease to an RPC message
+		/// </summary>
+		/// <returns>RPC message</returns>
+		public HordeCommon.Rpc.Messages.Lease ToRpcMessage()
+		{
+			HordeCommon.Rpc.Messages.Lease lease = new HordeCommon.Rpc.Messages.Lease();
+			lease.Id = Id.ToString();
+			lease.Payload = Google.Protobuf.WellKnownTypes.Any.Parser.ParseFrom(Payload); 
+			lease.State = State;
+			return lease;
+		}
+	}
+
+	/// <summary>
+	/// Well-known property names for agents
+	/// </summary>
+	static class KnownPropertyNames
+	{
+		/// <summary>
+		/// The agent id
+		/// </summary>
+		public const string Id = "Id";
+
+		/// <summary>
+		/// The operating system (Linux, MacOS, Windows)
+		/// </summary>
+		public const string OsFamily = "OSFamily";
+
+		/// <summary>
+		/// Pools that this agent belongs to
+		/// </summary>
+		public const string Pool = "Pool";
+
+		/// <summary>
+		/// Pools requested by the agent to join when registering with server
+		/// </summary>
+		public const string RequestedPools = "RequestedPools";
+		
+		/// <summary>
+		/// Number of logical cores
+		/// </summary>
+		public const string LogicalCores = "LogicalCores";
+
+		/// <summary>
+		/// Amount of RAM, in GB
+		/// </summary>
+		public const string Ram = "RAM";
+		
+		/// <summary>
+		/// AWS: Instance type
+		/// </summary>
+		public const string AwsInstanceType = "aws-instance-type";
+	}
+
+	/// <summary>
+	/// Mirrors an Agent document in the database
+	/// </summary>
+	public interface IAgent
+	{
+		/// <summary>
+		/// Randomly generated unique id for this agent.
+		/// </summary>
+		public AgentId Id { get; }
+
+		/// <summary>
+		/// The current session id, if it's online
+		/// </summary>
+		public SessionId? SessionId { get; }
+
+		/// <summary>
+		/// Time at which the current session expires. 
+		/// </summary>
+		public DateTime? SessionExpiresAt { get; }
+
+		/// <summary>
+		/// Current status of this agent
+		/// </summary>
+		public AgentStatus Status { get; }
+		
+		/// <summary>
+		/// Time at which last status change took place.
+		/// </summary>
+		public DateTime? LastStatusChange { get; }
+
+		/// <summary>
+		/// Whether the agent is enabled
+		/// </summary>
+		public bool Enabled { get; }
+
+		/// <summary>
+		/// Whether the agent should be included on the dashboard. This is set to true for ephemeral agents once they are no longer online, or agents that are explicitly deleted.
+		/// </summary>
+		public bool Deleted { get; }
+
+		/// <summary>
+		/// Version of the software running on this agent
+		/// </summary>
+		public string? Version { get; }
+
+		/// <summary>
+		/// Arbitrary comment for the agent (useful for disable reasons etc)
+		/// </summary>
+		public string? Comment { get; }
+
+		/// <summary>
+		/// List of properties for this agent
+		/// </summary>
+		public IReadOnlyList<string> Properties { get; }
+
+		/// <summary>
+		/// List of resources available to the agent
+		/// </summary>
+		public IReadOnlyDictionary<string, int> Resources { get; }
+
+		/// <summary>
+		/// Last upgrade that was attempted
+		/// </summary>
+		public string? LastUpgradeVersion { get; }
+
+		/// <summary>
+		/// Time that which the last upgrade was attempted
+		/// </summary>
+		public DateTime? LastUpgradeTime { get; }
+
+		/// <summary>
+		/// Dynamically applied pools
+		/// </summary>
+		public IReadOnlyList<PoolId> DynamicPools { get; }
+
+		/// <summary>
+		/// List of manually assigned pools for agent
+		/// </summary>
+		public IReadOnlyList<PoolId> ExplicitPools { get; }
+
+		/// <summary>
+		/// Whether a conform is requested
+		/// </summary>
+		public bool RequestConform { get; }
+
+		/// <summary>
+		/// Whether a full conform is requested
+		/// </summary>
+		public bool RequestFullConform { get; }
+
+		/// <summary>
+		/// Whether a machine restart is requested
+		/// </summary>
+		public bool RequestRestart { get; }
+
+		/// <summary>
+		/// Whether the machine should be shutdown
+		/// </summary>
+		public bool RequestShutdown { get; }
+
+		/// <summary>
+		/// The reason for the last agent shutdown
+		/// </summary>
+		public string? LastShutdownReason { get; }
+
+		/// <summary>
+		/// List of workspaces currently synced to this machine
+		/// </summary>
+		public IReadOnlyList<AgentWorkspace> Workspaces { get; }
+
+		/// <summary>
+		/// Time at which the last conform job ran
+		/// </summary>
+		public DateTime LastConformTime { get; }
+
+		/// <summary>
+		/// Number of times a conform job has failed
+		/// </summary>
+		public int? ConformAttemptCount { get; }
+
+		/// <summary>
+		/// Array of active leases.
+		/// </summary>
+		public IReadOnlyList<AgentLease> Leases { get; }
+
+		/// <summary>
+		/// Last time that the agent was modified
+		/// </summary>
+		public DateTime UpdateTime { get; }
+
+		/// <summary>
+		/// Update counter for this document. Any updates should compare-and-swap based on the value of this counter, or increment it in the case of server-side updates.
+		/// </summary>
+		public uint UpdateIndex { get; }
+	}
+
+	/// <summary>
+	/// Extension methods for IAgent
+	/// </summary>
+	public static class AgentExtensions
+	{
+		/// <summary>
+		/// Default tool id for agent software
+		/// </summary>
+		public static ToolId DefaultAgentSoftwareToolId { get; } = new ToolId("horde-agent");
+
+		/// <summary>
+		/// Gets the tool id for the software the given agent should be running
+		/// </summary>
+		/// <param name="agent">Agent to check</param>
+		/// <param name="globalConfig">Current global config</param>
+		/// <returns>Identifier for the tool that the agent should be using</returns>
+		public static ToolId GetSoftwareToolId(this IAgent agent, GlobalConfig globalConfig)
+		{
+			ToolId toolId = DefaultAgentSoftwareToolId;
+			foreach (AgentSoftwareConfig softwareConfig in globalConfig.Software)
+			{
+				if (softwareConfig.Condition != null && agent.SatisfiesCondition(softwareConfig.Condition))
+				{
+					toolId = softwareConfig.ToolId;
+					break;
+				}
+			}
+			return toolId;
+		}
+
+		/// <summary>
+		/// Determines whether this agent is online
+		/// </summary>
+		/// <returns></returns>
+		public static bool IsSessionValid(this IAgent agent, DateTime utcNow)
+		{
+			return agent.SessionId.HasValue && agent.SessionExpiresAt.HasValue && utcNow < agent.SessionExpiresAt.Value;
+		}
+
+		/// <summary>
+		/// Tests whether an agent is in the given pool
+		/// </summary>
+		/// <param name="agent"></param>
+		/// <param name="poolId"></param>
+		/// <returns></returns>
+		public static bool IsInPool(this IAgent agent, PoolId poolId)
+		{
+			return agent.DynamicPools.Contains(poolId) || agent.ExplicitPools.Contains(poolId);
+		}
+
+		/// <summary>
+		/// Get all the pools for each agent
+		/// </summary>
+		/// <param name="agent">The agent to query</param>
+		/// <returns></returns>
+		public static IEnumerable<PoolId> GetPools(this IAgent agent)
+		{
+			foreach (PoolId poolId in agent.DynamicPools)
+			{
+				yield return poolId;
+			}
+			foreach (PoolId poolId in agent.ExplicitPools)
+			{
+				yield return poolId;
+			}
+		}
+
+		/// <summary>
+		/// Tests whether an agent has a particular property
+		/// </summary>
+		/// <param name="agent"></param>
+		/// <param name="property"></param>
+		/// <returns></returns>
+		public static bool HasProperty(this IAgent agent, string property)
+		{
+			return agent.Properties.BinarySearch(property, StringComparer.OrdinalIgnoreCase) >= 0;
+		}
+
+		/// <summary>
+		/// Finds property values from a sorted list of Name=Value pairs
+		/// </summary>
+		/// <param name="agent">The agent to query</param>
+		/// <param name="name">Name of the property to find</param>
+		/// <returns>Property values</returns>
+		public static IEnumerable<string> GetPropertyValues(this IAgent agent, string name)
+		{
+			if (name.Equals(KnownPropertyNames.Id, StringComparison.OrdinalIgnoreCase))
+			{
+				yield return agent.Id.ToString();
+			}
+			else if (name.Equals(KnownPropertyNames.Pool, StringComparison.OrdinalIgnoreCase))
+			{
+				foreach (PoolId poolId in agent.GetPools())
+				{
+					yield return poolId.ToString();
+				}
+			}
+			else
+			{
+				int index = agent.Properties.BinarySearch(name, StringComparer.OrdinalIgnoreCase);
+				if (index < 0)
+				{
+					index = ~index;
+					for (; index < agent.Properties.Count; index++)
+					{
+						string property = agent.Properties[index];
+						if (property.Length <= name.Length || !property.StartsWith(name, StringComparison.OrdinalIgnoreCase) || property[name.Length] != '=')
+						{
+							break;
+						}
+						yield return property.Substring(name.Length + 1);
+					}
+				}
+			}
+		}
+
+		/// <summary>
+		/// Evaluates a condition against an agent
+		/// </summary>
+		/// <param name="agent">The agent to evaluate</param>
+		/// <param name="condition">The condition to evaluate</param>
+		/// <returns>True if the agent satisfies the condition</returns>
+		public static bool SatisfiesCondition(this IAgent agent, Condition condition)
+		{
+			return condition.Evaluate(x => agent.GetPropertyValues(x));
+		}
+
+		/// <summary>
+		/// Determine whether it's possible to add a lease for the given resources
+		/// </summary>
+		/// <param name="agent">The agent to create a lease for</param>
+		/// <param name="requirements">Requirements for the lease</param>
+		/// <param name="assignedResources">Receives the allocated resources</param>
+		/// <returns>True if the new lease can be granted</returns>
+		public static bool MeetsRequirements(this IAgent agent, Requirements requirements, Dictionary<string, int> assignedResources)
+		{
+			return MeetsRequirements(agent, requirements.Condition, requirements.Resources, requirements.Exclusive, assignedResources);
+		}
+
+		/// <summary>
+		/// Determine whether it's possible to add a lease for the given resources
+		/// </summary>
+		/// <param name="agent">The agent to create a lease for</param>
+		/// <param name="exclusive">Whether t</param>
+		/// <param name="condition">Condition to satisfy</param>
+		/// <param name="resources">Resources required to execute</param>
+		/// <param name="assignedResources">Resources allocated to the task</param>
+		/// <returns>True if the new lease can be granted</returns>
+		public static bool MeetsRequirements(this IAgent agent, Condition? condition, Dictionary<string, ResourceRequirements>? resources, bool exclusive, Dictionary<string, int> assignedResources)
+		{
+			if (!agent.Enabled || agent.Status != AgentStatus.Ok)
+			{
+				return false;
+			}
+			if (agent.Leases.Any(x => x.Exclusive))
+			{
+				return false;
+			}
+			if (exclusive && agent.Leases.Any())
+			{
+				return false;
+			}
+			if (condition != null && !agent.SatisfiesCondition(condition))
+			{
+				return false;
+			}
+			if (resources != null)
+			{
+				foreach ((string name, ResourceRequirements resourceRequirements) in resources)
+				{
+					int remainingCount;
+					if (!agent.Resources.TryGetValue(name, out remainingCount))
+					{
+						return false;
+					}
+					foreach (AgentLease lease in agent.Leases)
+					{
+						if (lease.Resources != null)
+						{
+							int leaseCount;
+							lease.Resources.TryGetValue(name, out leaseCount);
+							remainingCount -= leaseCount;
+						}
+					}
+					if (remainingCount < resourceRequirements.Min)
+					{
+						return false;
+					}
+
+					int allocatedCount;
+					if (resourceRequirements.Max != null)
+					{
+						allocatedCount = Math.Min(resourceRequirements.Max.Value, remainingCount);
+					}
+					else
+					{
+						allocatedCount = resourceRequirements.Min;
+					}
+					assignedResources.Add(name, allocatedCount);
+				}
+			}
+			return true;
+		}
+
+		/// <summary>
+		/// Gets all the autosdk workspaces required for an agent
+		/// </summary>
+		/// <param name="agent"></param>
+		/// <param name="globalConfig"></param>
+		/// <param name="workspaces"></param>
+		/// <returns></returns>
+		public static HashSet<AgentWorkspace> GetAutoSdkWorkspaces(this IAgent agent, GlobalConfig globalConfig, List<AgentWorkspace> workspaces)
+		{
+			HashSet<AgentWorkspace> autoSdkWorkspaces = new HashSet<AgentWorkspace>();
+			foreach (string? clusterName in workspaces.Select(x => x.Cluster).Distinct())
+			{
+				PerforceCluster? cluster = globalConfig.FindPerforceCluster(clusterName);
+				if (cluster != null)
+				{
+					AgentWorkspace? autoSdkWorkspace = GetAutoSdkWorkspace(agent, cluster);
+					if (autoSdkWorkspace != null)
+					{
+						autoSdkWorkspaces.Add(autoSdkWorkspace);
+					}
+				}
+			}
+			return autoSdkWorkspaces;
+		}
+
+		/// <summary>
+		/// Get the AutoSDK workspace required for an agent
+		/// </summary>
+		/// <param name="agent"></param>
+		/// <param name="cluster">The perforce cluster to get a workspace for</param>
+		/// <returns></returns>
+		public static AgentWorkspace? GetAutoSdkWorkspace(this IAgent agent, PerforceCluster cluster)
+		{
+			foreach (AutoSdkWorkspace autoSdk in cluster.AutoSdk)
+			{
+				if (autoSdk.Stream != null && autoSdk.Properties.All(x => agent.Properties.Contains(x)))
+				{
+					return new AgentWorkspace(cluster.Name, autoSdk.UserName, autoSdk.Name ?? "AutoSDK", autoSdk.Stream!, null, true, null);
+				}
+			}
+			return null;
+		}
+
+		/// <summary>
+		/// Converts this workspace to an RPC message
+		/// </summary>
+		/// <param name="agent">The agent to get a workspace for</param>
+		/// <param name="workspace">The workspace definition</param>
+		/// <param name="cluster">The global state</param>
+		/// <param name="loadBalancer">The Perforce load balancer</param>
+		/// <param name="workspaceMessages">List of messages</param>
+		/// <returns>The RPC message</returns>
+		public static async Task<bool> TryAddWorkspaceMessage(this IAgent agent, AgentWorkspace workspace, PerforceCluster cluster, PerforceLoadBalancer loadBalancer, IList<HordeCommon.Rpc.Messages.AgentWorkspace> workspaceMessages)
+		{
+			// Find a matching server, trying to use a previously selected one if possible
+			string? baseServerAndPort;
+			string? serverAndPort;
+
+			HordeCommon.Rpc.Messages.AgentWorkspace? existingWorkspace = workspaceMessages.FirstOrDefault(x => x.ConfiguredCluster == workspace.Cluster);
+			if(existingWorkspace != null)
+			{
+				baseServerAndPort = existingWorkspace.BaseServerAndPort;
+				serverAndPort = existingWorkspace.ServerAndPort;
+			}
+			else
+			{
+				if (cluster == null)
+				{
+					return false;
+				}
+
+				IPerforceServer? server = await loadBalancer.SelectServerAsync(cluster, agent);
+				if (server == null)
+				{
+					return false;
+				}
+
+				baseServerAndPort = server.BaseServerAndPort;
+				serverAndPort = server.ServerAndPort;
+			}
+
+			// Find the matching credentials for the desired user
+			PerforceCredentials? credentials = null;
+			if (cluster != null)
+			{
+				if (workspace.UserName == null)
+				{
+					credentials = cluster.Credentials.FirstOrDefault();
+				}
+				else
+				{
+					credentials = cluster.Credentials.FirstOrDefault(x => String.Equals(x.UserName, workspace.UserName, StringComparison.OrdinalIgnoreCase));
+				}
+			}
+
+			// Construct the message
+			HordeCommon.Rpc.Messages.AgentWorkspace result = new ()
+			{
+				ConfiguredCluster = workspace.Cluster, ConfiguredUserName = workspace.UserName,
+				Cluster = cluster?.Name,
+				BaseServerAndPort = baseServerAndPort,
+				ServerAndPort = serverAndPort,
+				UserName = credentials?.UserName ?? workspace.UserName,
+				Password = credentials?.Password,
+				Identifier = workspace.Identifier,
+				Stream = workspace.Stream,
+				Incremental = workspace.Incremental,
+				Method = workspace.Method ?? String.Empty
+			};
+
+			if (workspace.View != null)
+			{
+				result.View.AddRange(workspace.View);
+			}
+			
+			workspaceMessages.Add(result);
+			return true;
+		}
+
+		/// <summary>
+		/// Tries to get an agent workspace definition from the given type name
+		/// </summary>
+		/// <param name="streamConfig">The stream object</param>
+		/// <param name="agentType">The agent type</param>
+		/// <param name="workspace">Receives the agent workspace definition</param>
+		/// <returns>True if the agent type was valid, and an agent workspace could be created</returns>
+		public static bool TryGetAgentWorkspace(this StreamConfig streamConfig, AgentConfig agentType, [NotNullWhen(true)] out (AgentWorkspace, bool)? workspace)
+		{
+			// Get the workspace settings
+			if (agentType.Workspace == null)
+			{
+				// Use the default settings (fast switching workspace, clean 
+				workspace = (new AgentWorkspace(null, null, streamConfig.GetDefaultWorkspaceIdentifier(), streamConfig.Name, null, false, null), true);
+				return true;
+			}
+			else
+			{
+				// Try to get the matching workspace type
+				WorkspaceConfig? workspaceType;
+				if (!streamConfig.WorkspaceTypes.TryGetValue(agentType.Workspace, out workspaceType))
+				{
+					workspace = null;
+					return false;
+				}
+
+				// Get the workspace identifier
+				string identifier;
+				if (workspaceType.Identifier != null)
+				{
+					identifier = workspaceType.Identifier;
+				}
+				else if (workspaceType.Incremental)
+				{
+					identifier = $"{streamConfig.GetEscapedName()}+{agentType.Workspace}";
+				}
+				else
+				{
+					identifier = streamConfig.GetDefaultWorkspaceIdentifier();
+				}
+
+				// Create the new workspace
+				workspace = (new AgentWorkspace(workspaceType.Cluster, workspaceType.UserName, identifier, workspaceType.Stream ?? streamConfig.Name, workspaceType.View, workspaceType.Incremental, workspaceType.Method), workspaceType.UseAutoSdk);
+				return true;
+			}
+		}
+	}
+}
diff --git a/Engine/Source/Programs/Horde/Horde.Server/Agents/IAgentCollection.cs b/Engine/Source/Programs/Horde/Horde.Server/Agents/IAgentCollection.cs
new file mode 100644
--- /dev/null
+++ b/Engine/Source/Programs/Horde/Horde.Server/Agents/IAgentCollection.cs
@@ -0,0 +1,159 @@
+// Copyright Epic Games, Inc. All Rights Reserved.
+
+using System;
+using System.Collections.Generic;
+using System.Threading.Tasks;
+using Horde.Server.Acls;
+using Horde.Server.Agents.Pools;
+using Horde.Server.Agents.Sessions;
+using Horde.Server.Auditing;
+using Horde.Server.Utilities;
+using HordeCommon;
+
+namespace Horde.Server.Agents
+{
+	/// <summary>
+	/// Interface for a collection of agent documents
+	/// </summary>
+	public interface IAgentCollection
+	{
+		/// <summary>
+		/// Adds a new agent with the given properties
+		/// </summary>
+		/// <param name="id">Id for the new agent</param>
+		/// <param name="enabled">Whether the agent is enabled or not</param>
+		/// <param name="pools">Pools for the agent</param>
+		Task<IAgent> AddAsync(AgentId id, bool enabled, List<PoolId>? pools = null);
+
+		/// <summary>
+		/// Deletes an agent
+		/// </summary>
+		/// <param name="agent">Deletes the agent</param>
+		/// <returns>Async task</returns>
+		Task<IAgent?> TryDeleteAsync(IAgent agent);
+
+		/// <summary>
+		/// Deletes an agent
+		/// </summary>
+		/// <param name="agentId">Unique id of the agent</param>
+		/// <returns>Async task</returns>
+		Task ForceDeleteAsync(AgentId agentId);
+
+		/// <summary>
+		/// Gets an agent by ID
+		/// </summary>
+		/// <param name="agentId">Unique id of the agent</param>
+		/// <returns>The agent document</returns>
+		Task<IAgent?> GetAsync(AgentId agentId);
+
+		/// <summary>
+		/// Finds all agents matching certain criteria
+		/// </summary>
+		/// <param name="poolId">The pool ID in string form containing the agent</param>
+		/// <param name="modifiedAfter">If set, only returns agents modified after this time</param>
+		/// <param name="status">Status to look for</param>
+		/// <param name="enabled">Enabled/disabled status to look for</param>
+		/// <param name="index">Index of the first result</param>
+		/// <param name="count">Number of results to return</param>
+		/// <returns>List of agents matching the given criteria</returns>
+		Task<List<IAgent>> FindAsync(PoolId? poolId = null, DateTime? modifiedAfter = null, AgentStatus? status = null, bool? enabled = null, int? index = null, int? count = null);
+
+		/// <summary>
+		/// Finds all the expried agents
+		/// </summary>
+		/// <param name="utcNow">The current time</param>
+		/// <param name="maxAgents">Maximum number of agents to return</param>
+		/// <returns>List of agents</returns>
+		Task<List<IAgent>> FindExpiredAsync(DateTime utcNow, int maxAgents);
+
+		/// <summary>
+		/// Update an agent's settings
+		/// </summary>
+		/// <param name="agent">Agent instance</param>
+		/// <param name="enabled">Whether the agent is enabled or not</param>
+		/// <param name="requestConform">Whether to request a conform job be run</param>
+		/// <param name="requestFullConform">Whether to request a full conform job be run</param>
+		/// <param name="requestRestart">Whether to request the machine be restarted</param>
+		/// <param name="requestShutdown">Whether to request the machine be shut down</param>
+		/// <param name="shutdownReason">The reason for shutting down agent, ex. Autoscaler/Manual/Unexpected</param>
+		/// <param name="pools">List of pools for the agent</param>
+		/// <param name="comment">New comment</param>
+		/// <returns>New agent state if update was successful</returns>
+		Task<IAgent?> TryUpdateSettingsAsync(IAgent agent, bool? enabled = null, bool? requestConform = null, bool? requestFullConform = null, bool? requestRestart = null, bool? requestShutdown = null, string? shutdownReason = null, List<PoolId>? pools = null, string? comment = null);
+
+		/// <summary>
+		/// Update the current workspaces for an agent.
+		/// </summary>
+		/// <param name="agent">The agent to update</param>
+		/// <param name="workspaces">Current list of workspaces</param>
+		/// <param name="requestConform">Whether the agent still needs to run another conform</param>
+		/// <returns>New agent state</returns>
+		Task<IAgent?> TryUpdateWorkspacesAsync(IAgent agent, List<AgentWorkspace> workspaces, bool requestConform);
+
+		/// <summary>
+		/// Sets the current session
+		/// </summary>
+		/// <param name="agent">The agent to update</param>
+		/// <param name="sessionId">New session id</param>
+		/// <param name="sessionExpiresAt">Expiry time for the new session</param>
+		/// <param name="status">Status of the agent</param>
+		/// <param name="properties">Properties for the current session</param>
+		/// <param name="resources">Resources for the agent</param>
+		/// <param name="pools">New list of pools for the agent</param>
+		/// <param name="dynamicPools">New list of dynamic pools for the agent</param>
+		/// <param name="lastStatusChange">Time to force status change timestamp to</param>
+		/// <param name="version">Current version of the agent software</param>
+		/// <returns>New agent state</returns>
+		Task<IAgent?> TryStartSessionAsync(IAgent agent, SessionId sessionId, DateTime sessionExpiresAt, AgentStatus status, IReadOnlyList<string> properties, IReadOnlyDictionary<string, int> resources, IReadOnlyList<PoolId> pools, IReadOnlyList<PoolId> dynamicPools, DateTime lastStatusChange, string? version);
+
+		/// <summary>
+		/// Attempt to update the agent state
+		/// </summary>
+		/// <param name="agent">Agent instance</param>
+		/// <param name="status">New status of the agent</param>
+		/// <param name="sessionExpiresAt">New expiry time for the current session</param>
+		/// <param name="properties">Properties for the current session</param>
+		/// <param name="resources">Resources for the agent</param>
+		/// <param name="dynamicPools">New list of dynamic pools for the agent</param>
+		/// <param name="leases">New set of leases</param>
+		/// <returns>True if the document was updated, false if another writer updated the document first</returns>
+		Task<IAgent?> TryUpdateSessionAsync(IAgent agent, AgentStatus? status, DateTime? sessionExpiresAt, IReadOnlyList<string>? properties, IReadOnlyDictionary<string, int>? resources, IReadOnlyList<PoolId>? dynamicPools, List<AgentLease>? leases);
+
+		/// <summary>
+		/// Terminates the current session
+		/// </summary>
+		/// <param name="agent">The agent to terminate the session for</param>
+		/// <returns>New agent state if it succeeded, otherwise null</returns>
+		Task<IAgent?> TryTerminateSessionAsync(IAgent agent);
+
+		/// <summary>
+		/// Attempts to add a lease to an agent
+		/// </summary>
+		/// <param name="agent">The agent to update</param>
+		/// <param name="newLease">The new lease document</param>
+		/// <returns>New agent state if it succeeded, otherwise null</returns>
+		Task<IAgent?> TryAddLeaseAsync(IAgent agent, AgentLease newLease);
+
+		/// <summary>
+		/// Attempts to cancel a lease
+		/// </summary>
+		/// <param name="agent">The agent to update</param>
+		/// <param name="leaseIdx">Index of the lease</param>
+		/// <returns>New agent state if it succeeded, otherwise null</returns>
+		Task<IAgent?> TryCancelLeaseAsync(IAgent agent, int leaseIdx);
+
+		/// <summary>
+		/// Gets the log channel for an agent
+		/// </summary>
+		/// <param name="agentId"></param>
+		/// <returns></returns>
+		IAuditLogChannel<AgentId> GetLogger(AgentId agentId);
+
+		/// <summary>
+		/// Subscribe to notifications on agent states being updated
+		/// </summary>
+		/// <param name="onUpdate">Callback for updates</param>
+		/// <returns>Disposable subscription object</returns>
+		Task<IAsyncDisposable> SubscribeToUpdateEventsAsync(Action<AgentId> onUpdate);
+	}
+}
diff --git a/Engine/Source/Programs/Horde/Horde.Server/Agents/Leases/LeaseCollection.cs b/Engine/Source/Programs/Horde/Horde.Server/Agents/Leases/LeaseCollection.cs
new file mode 100644
--- /dev/null
+++ b/Engine/Source/Programs/Horde/Horde.Server/Agents/Leases/LeaseCollection.cs
@@ -0,0 +1,191 @@
+// Copyright Epic Games, Inc. All Rights Reserved.
+
+using System;
+using System.Collections.Generic;
+using System.Linq;
+using System.Threading.Tasks;
+using Horde.Server.Agents.Pools;
+using Horde.Server.Agents.Sessions;
+using Horde.Server.Logs;
+using Horde.Server.Server;
+using Horde.Server.Streams;
+using Horde.Server.Utilities;
+using HordeCommon;
+using MongoDB.Bson.Serialization.Attributes;
+using MongoDB.Driver;
+
+namespace Horde.Server.Agents.Leases
+{
+	/// <summary>
+	/// Collection of lease documents
+	/// </summary>
+	public class LeaseCollection : ILeaseCollection
+	{
+		/// <summary>
+		/// Concrete implementation of a lease document
+		/// </summary>
+		class LeaseDocument : ILease
+		{
+			public LeaseId Id { get; set; }
+			public string Name { get; set; }
+			public AgentId AgentId { get; set; }
+			public SessionId SessionId { get; set; }
+			public StreamId? StreamId { get; set; }
+			public PoolId? PoolId { get; set; }
+			public LogId? LogId { get; set; }
+			public DateTime StartTime { get; set; }
+			public DateTime? FinishTime { get; set; }
+			public byte[] Payload { get; set; }
+			public LeaseOutcome Outcome { get; set; }
+
+			[BsonIgnoreIfNull]
+			public byte[]? Output { get; set; }
+
+			ReadOnlyMemory<byte> ILease.Payload => Payload;
+			ReadOnlyMemory<byte> ILease.Output => Output ?? Array.Empty<byte>();
+
+			[BsonConstructor]
+			private LeaseDocument()
+			{
+				Name = String.Empty;
+				Payload = null!;
+			}
+
+			public LeaseDocument(LeaseId id, string name, AgentId agentId, SessionId sessionId, StreamId? streamId, PoolId? poolId, LogId? logId, DateTime startTime, byte[] payload)
+			{
+				Id = id;
+				Name = name;
+				AgentId = agentId;
+				SessionId = sessionId;
+				StreamId = streamId;
+				PoolId = poolId;
+				LogId = logId;
+				StartTime = startTime;
+				Payload = payload;
+			}
+		}
+		
+		readonly IMongoCollection<LeaseDocument> _leases;
+		readonly MongoIndex<LeaseDocument> _finishTimeStartTimeCompoundIndex;
+
+		/// <summary>
+		/// Constructor
+		/// </summary>
+		/// <param name="mongoService">The database service instance</param>
+		public LeaseCollection(MongoService mongoService)
+		{
+			List<MongoIndex<LeaseDocument>> indexes = new List<MongoIndex<LeaseDocument>>();
+			indexes.Add(keys => keys.Ascending(x => x.AgentId));
+			indexes.Add(keys => keys.Ascending(x => x.SessionId));
+			indexes.Add(keys => keys.Ascending(x => x.StartTime));
+			indexes.Add(keys => keys.Ascending(x => x.FinishTime));
+			indexes.Add(_finishTimeStartTimeCompoundIndex = MongoIndex.Create<LeaseDocument>(keys => keys.Ascending(x => x.FinishTime).Descending(x => x.StartTime)));
+
+			_leases = mongoService.GetCollection<LeaseDocument>("Leases", indexes);
+		}
+
+		/// <inheritdoc/>
+		public async Task<ILease> AddAsync(LeaseId id, string name, AgentId agentId, SessionId sessionId, StreamId? streamId, PoolId? poolId, LogId? logId, DateTime startTime, byte[] payload)
+		{
+			LeaseDocument lease = new LeaseDocument(id, name, agentId, sessionId, streamId, poolId, logId, startTime, payload);
+			await _leases.ReplaceOneAsync(x => x.Id == id, lease, new ReplaceOptions { IsUpsert = true });
+			return lease;
+		}
+
+		/// <inheritdoc/>
+		public async Task DeleteAsync(LeaseId leaseId)
+		{
+			await _leases.DeleteOneAsync(x => x.Id == leaseId);
+		}
+
+		/// <inheritdoc/>
+		public async Task<ILease?> GetAsync(LeaseId leaseId)
+		{
+			return await _leases.Find(x => x.Id == leaseId).FirstOrDefaultAsync();
+		}
+
+		/// <inheritdoc/>
+		public async Task<List<ILease>> FindLeasesAsync(AgentId? agentId, SessionId? sessionId, DateTime? minTime, DateTime? maxTime, int? index, int? count, string? indexHint = null, bool consistentRead = true)
+		{
+			IMongoCollection<LeaseDocument> collection = consistentRead ? _leases : _leases.WithReadPreference(ReadPreference.SecondaryPreferred);
+			
+			FilterDefinitionBuilder<LeaseDocument> filterBuilder = Builders<LeaseDocument>.Filter;
+			FilterDefinition<LeaseDocument> filter = FilterDefinition<LeaseDocument>.Empty;
+			if (agentId != null)
+			{
+				filter &= filterBuilder.Eq(x => x.AgentId, agentId.Value);
+			}
+			if (sessionId != null)
+			{
+				filter &= filterBuilder.Eq(x => x.SessionId, sessionId.Value);
+			}
+			if (minTime != null)
+			{
+				filter &= filterBuilder.Or(filterBuilder.Eq(x => x.FinishTime!, null), filterBuilder.Gt(x => x.FinishTime!, minTime.Value));
+			}
+			if (maxTime != null)
+			{
+				filter &= filterBuilder.Lt(x => x.StartTime, maxTime.Value);
+			}
+
+			List<LeaseDocument> results = await collection.FindWithHint(filter, indexHint, x => x.SortByDescending(x => x.StartTime).Range(index, count).ToListAsync());
+			return results.ConvertAll<ILease>(x => x);
+		}
+		
+		/// <inheritdoc/>
+		public async Task<List<ILease>> FindLeasesByFinishTimeAsync(DateTime? minFinishTime, DateTime? maxFinishTime, int? index, int? count, string? indexHint, bool consistentRead)
+		{
+			IMongoCollection<LeaseDocument> collection = consistentRead ? _leases : _leases.WithReadPreference(ReadPreference.SecondaryPreferred);
+			FilterDefinitionBuilder<LeaseDocument> filterBuilder = Builders<LeaseDocument>.Filter;
+			FilterDefinition<LeaseDocument> filter = FilterDefinition<LeaseDocument>.Empty;
+
+			if (minFinishTime == null && maxFinishTime == null)
+			{
+				throw new ArgumentException($"Both {nameof(minFinishTime)} and {nameof(maxFinishTime)} cannot be null");
+			}
+
+			if (minFinishTime != null)
+			{
+				filter &= filterBuilder.Gt(x => x.FinishTime, minFinishTime.Value);
+			}
+			if (maxFinishTime != null)
+			{
+				filter &= filterBuilder.Lt(x => x.FinishTime, maxFinishTime.Value);
+			}
+
+			List<LeaseDocument> results = await collection.FindWithHint(filter, indexHint, x => x.SortByDescending(x => x.FinishTime).Range(index, count).ToListAsync());
+			return results.ConvertAll<ILease>(x => x);
+		}
+
+		/// <inheritdoc/>
+		public async Task<List<ILease>> FindLeasesAsync(DateTime? minTime, DateTime? maxTime)
+		{
+			return await FindLeasesAsync(null, null, minTime, maxTime, null, null, _finishTimeStartTimeCompoundIndex.Name, false);
+		}
+
+		/// <inheritdoc/>
+		public async Task<List<ILease>> FindActiveLeasesAsync(int? index = null, int? count = null)
+		{
+			List<LeaseDocument> results = await _leases.Find(x => x.FinishTime == null).Range(index, count).ToListAsync();
+			return results.ConvertAll<ILease>(x => x);
+		}
+
+		/// <inheritdoc/>
+		public async Task<bool> TrySetOutcomeAsync(LeaseId leaseId, DateTime finishTime, LeaseOutcome outcome, byte[]? output)
+		{
+			FilterDefinitionBuilder<LeaseDocument> filterBuilder = Builders<LeaseDocument>.Filter;
+			FilterDefinition<LeaseDocument> filter = filterBuilder.Eq(x => x.Id, leaseId) & filterBuilder.Eq(x => x.FinishTime, null);
+
+			UpdateDefinitionBuilder<LeaseDocument> updateBuilder = Builders<LeaseDocument>.Update;
+			UpdateDefinition<LeaseDocument> update = updateBuilder.Set(x => x.FinishTime, finishTime).Set(x => x.Outcome, outcome);
+
+			if(output != null && output.Length > 0)
+			{
+				update = update.Set(x => x.Output, output);
+			}
+
+			UpdateResult result = await _leases.UpdateOneAsync(filter, update);
+			return result.ModifiedCount > 0;
+		}
+	}
+}
diff --git a/Engine/Source/Programs/Horde/Horde.Server/Agents/Leases/LeaseId.cs b/Engine/Source/Programs/Horde/Horde.Server/Agents/Leases/LeaseId.cs
new file mode 100644
--- /dev/null
+++ b/Engine/Source/Programs/Horde/Horde.Server/Agents/Leases/LeaseId.cs
@@ -0,0 +1,58 @@
+// Copyright Epic Games, Inc. All Rights Reserved.
+
+using System.ComponentModel;
+using EpicGames.Core;
+using Horde.Server.Utilities;
+using MongoDB.Bson;
+
+namespace Horde.Server.Agents.Leases
+{
+	/// <summary>
+	/// Identifier for a user
+	/// </summary>
+	/// <param name="Id">Id to construct from</param>
+	[JsonSchemaString]
+	[TypeConverter(typeof(ObjectIdTypeConverter<LeaseId, LeaseIdConverter>))]
+	[ObjectIdConverter(typeof(LeaseIdConverter))]
+	public record struct LeaseId(ObjectId Id)
+	{
+		/// <summary>
+		/// Creates a new <see cref="LeaseId"/>
+		/// </summary>
+		public static LeaseId GenerateNewId() => new LeaseId(ObjectId.GenerateNewId());
+
+		/// <inheritdoc cref="ObjectId.TryParse(System.String, out ObjectId)"/>
+		public static bool TryParse(string text, out LeaseId leaseId)
+		{
+			ObjectId objectId;
+			if (ObjectId.TryParse(text, out objectId))
+			{
+				leaseId = new LeaseId(objectId);
+				return true;
+			}
+			else
+			{
+				leaseId = default;
+				return false;
+			}
+		}
+
+		/// <inheritdoc cref="ObjectId.Parse(System.String)"/>
+		public static LeaseId Parse(string text) => new LeaseId(ObjectId.Parse(text));
+
+		/// <inheritdoc/>
+		public override string ToString() => Id.ToString();
+	}
+
+	/// <summary>
+	/// Converter to and from <see cref="ObjectId"/> instances.
+	/// </summary>
+	class LeaseIdConverter : ObjectIdConverter<LeaseId>
+	{
+		/// <inheritdoc/>
+		public override LeaseId FromObjectId(ObjectId id) => new LeaseId(id);
+
+		/// <inheritdoc/>
+		public override ObjectId ToObjectId(LeaseId value) => value.Id;
+	}
+}
diff --git a/Engine/Source/Programs/Horde/Horde.Server/Agents/Leases/LeasesController.cs b/Engine/Source/Programs/Horde/Horde.Server/Agents/Leases/LeasesController.cs
new file mode 100644
--- /dev/null
+++ b/Engine/Source/Programs/Horde/Horde.Server/Agents/Leases/LeasesController.cs
@@ -0,0 +1,191 @@
+// Copyright Epic Games, Inc. All Rights Reserved.
+
+using System;
+using System.Collections.Generic;
+using System.Linq;
+using System.Threading.Tasks;
+using Horde.Server.Acls;
+using Horde.Server.Agents.Sessions;
+using Horde.Server.Server;
+using Horde.Server.Utilities;
+using Microsoft.AspNetCore.Authorization;
+using Microsoft.AspNetCore.Mvc;
+using Microsoft.Extensions.Options;
+using OpenTracing;
+using OpenTracing.Util;
+
+namespace Horde.Server.Agents.Leases
+{
+	/// <summary>
+	/// Controller for the /api/v1/leases endpoint
+	/// </summary>
+	[ApiController]
+	[Authorize]
+	[Route("[controller]")]
+	public class LeasesController : HordeControllerBase
+	{
+		readonly AgentService _agentService;
+		readonly IOptionsSnapshot<GlobalConfig> _globalConfig;
+
+		/// <summary>
+		/// Constructor
+		/// </summary>
+		public LeasesController(AgentService agentService, IOptionsSnapshot<GlobalConfig> globalConfig)
+		{
+			_agentService = agentService;
+			_globalConfig = globalConfig;
+		}
+
+		/// <summary>
+		/// Find all the leases for a particular agent
+		/// </summary>
+		/// <param name="agentId">Unique id of the agent to find</param>
+		/// <param name="sessionId">The session to query</param>
+		/// <param name="startTime">Start of the time window to consider</param>
+		/// <param name="finishTime">End of the time window to consider</param>
+		/// <param name="minFinishTime">Start of the time window to consider when querying by finish time (if set, other criteria are ignored)</param>
+		/// <param name="maxFinishTime">End of the time window to consider when querying by finish time (if set, other criteria are ignored)</param>
+		/// <param name="index">Index of the first result to return</param>
+		/// <param name="count">Number of results to return</param>
+		/// <param name="filter">Filter to apply to the results</param>
+		/// <returns>Sessions </returns>
+		[HttpGet]
+		[Route("/api/v1/leases")]
+		[ProducesResponseType(200, Type = typeof(List<GetAgentLeaseResponse>))]
+		public async Task<ActionResult<List<object>>> FindLeasesAsync(
+			[FromQuery] AgentId? agentId,
+			[FromQuery] SessionId? sessionId,
+			[FromQuery] DateTimeOffset? startTime,
+			[FromQuery] DateTimeOffset? finishTime,
+			[FromQuery] DateTimeOffset? minFinishTime,
+			[FromQuery] DateTimeOffset? maxFinishTime,
+			[FromQuery] int index = 0,
+			[FromQuery] int count = 1000,
+			[FromQuery] PropertyFilter? filter = null)
+		{
+			if (!_globalConfig.Value.Authorize(AclAction.ViewLeases, User))
+			{
+				return Forbid(AclAction.ViewLeases);
+			}
+
+			bool includeCosts = _globalConfig.Value.Authorize(AclAction.ViewCosts, User);
+
+			List<ILease> leases;
+			if (minFinishTime == null && maxFinishTime == null)
+			{
+				leases = await _agentService.FindLeasesAsync(agentId, sessionId, startTime?.UtcDateTime, finishTime?.UtcDateTime, index, count);
+			}
+			else
+			{
+				// Optimized path for queries made by finish time
+				leases = await _agentService.FindLeasesByFinishTimeAsync(minFinishTime?.UtcDateTime, maxFinishTime?.UtcDateTime, index, count);
+			}
+
+			List<object> responses = new List<object>();
+
+			using (IScope _ = GlobalTracer.Instance.BuildSpan($"GenerateResponses").StartActive())
+			{
+				Dictionary<AgentId, double?> cachedAgentRates = new Dictionary<AgentId, double?>();
+				foreach (ILease lease in leases)
+				{
+					double? agentRate = null;
+					if (includeCosts && !cachedAgentRates.TryGetValue(lease.AgentId, out agentRate))
+					{
+						agentRate = await _agentService.GetRateAsync(lease.AgentId);
+						cachedAgentRates.Add(lease.AgentId, agentRate);
+					}
+
+					Dictionary<string, string>? details = _agentService.GetPayloadDetails(lease.Payload);
+					responses.Add(PropertyFilter.Apply(new GetAgentLeaseResponse(lease, details, agentRate), filter));
+				}
+			}
+
+			return responses;
+		}
+
+		/// <summary>
+		/// Get info about a particular lease
+		/// </summary>
+		/// <param name="leaseId">Unique id of the particular lease</param>
+		/// <returns>Lease matching the given id</returns>
+		[HttpGet]
+		[Route("/api/v1/leases/{leaseId}")]
+		public async Task<ActionResult<GetAgentLeaseResponse>> GetLeaseAsync(LeaseId leaseId)
+		{
+			if (!_globalConfig.Value.Authorize(AclAction.ViewLeases, User))
+			{
+				return Forbid(AclAction.ViewLeases);
+			}
+
+			ILease? lease = await _agentService.GetLeaseAsync(leaseId);
+			if (lease == null)
+			{
+				return NotFound(leaseId);
+			}
+
+			IAgent? agent = await _agentService.GetAgentAsync(lease.AgentId);
+			if (agent == null)
+			{
+				return NotFound(lease.AgentId);
+			}
+
+			double? agentRate = null;
+			if (_globalConfig.Value.Authorize(AclAction.ViewCosts, User))
+			{
+				agentRate = await _agentService.GetRateAsync(agent.Id);
+			}
+
+			Dictionary<string, string>? details = _agentService.GetPayloadDetails(lease.Payload);
+			return new GetAgentLeaseResponse(lease, details, agentRate);
+		}
+
+		/// <summary>
+		/// Update a particular lease
+		/// </summary>
+		/// <param name="leaseId">Unique id of the particular lease</param>
+		/// <param name="request"></param>
+		/// <returns>Lease matching the given id</returns>
+		[HttpPut]
+		[Route("/api/v1/leases/{leaseId}")]
+		public async Task<ActionResult> UpdateLeaseAsync(LeaseId leaseId, [FromBody] UpdateLeaseRequest request)
+		{
+			if (!_globalConfig.Value.Authorize(AclAction.AdminWrite, User))
+			{
+				return Forbid(AclAction.AdminWrite);
+			}
+
+			// only update supported right now is abort
+			if (!request.Aborted.HasValue || !request.Aborted.Value)
+			{
+				return Ok();
+			}
+
+			ILease? lease = await _agentService.GetLeaseAsync(leaseId);
+			if (lease == null)
+			{
+				return NotFound(leaseId);
+			}
+
+			IAgent? agent = await _agentService.GetAgentAsync(lease.AgentId);
+			if (agent == null)
+			{
+				return NotFound(lease.AgentId);
+			}
+
+			AgentLease? agentLease = agent.Leases.FirstOrDefault(x => x.Id == leaseId);
+
+			if (agentLease == null)
+			{
+				return NotFound(agent.Id, leaseId);
+			}
+
+			if (!agentLease.IsConformLease())
+			{
+				return BadRequest("Lease abort only supported on conform leases for now, {LeaseId}", leaseId);
+			}
+
+			await _agentService.CancelLeaseAsync(agent, leaseId);
+			return Ok();
+		}
+	}
+}
diff --git a/Engine/Source/Programs/Horde/Horde.Server/Agents/Pools/IPoolCollection.cs b/Engine/Source/Programs/Horde/Horde.Server/Agents/Pools/IPoolCollection.cs
new file mode 100644
--- /dev/null
+++ b/Engine/Source/Programs/Horde/Horde.Server/Agents/Pools/IPoolCollection.cs
@@ -0,0 +1,174 @@
+// Copyright Epic Games, Inc. All Rights Reserved.
+
+using System;
+using System.Collections.Generic;
+using System.Threading.Tasks;
+using Amazon.Util.Internal;
+using EpicGames.Horde.Common;
+using Horde.Server.Agents.Fleet;
+using Horde.Server.Utilities;
+
+namespace Horde.Server.Agents.Pools
+{
+	/// <summary>
+	/// Exception when conflicting definitions for a pool are encountered
+	/// </summary>
+	public sealed class PoolConflictException : Exception
+	{
+		/// <summary>
+		/// Pool with conflicting definitions
+		/// </summary>
+		public PoolId PoolId { get; }
+
+		/// <summary>
+		/// The first revision string
+		/// </summary>
+		public string PrevRevision { get; }
+
+		/// <summary>
+		/// The second revision string
+		/// </summary>
+		public string NextRevision { get; }
+
+		/// <summary>
+		/// Constructor
+		/// </summary>
+		public PoolConflictException(PoolId poolId, string prevRevision, string nextRevision)
+			: base($"Duplicate definitions for {poolId} - first in {prevRevision}, now {nextRevision}")
+		{
+			PoolId = poolId;
+			PrevRevision = prevRevision;
+			NextRevision = nextRevision;
+		}
+	}
+
+	/// <summary>
+	/// Collection of pool documents
+	/// </summary>
+	public interface IPoolCollection
+	{
+		/// <summary>
+		/// Creates a new pool
+		/// </summary>
+		/// <param name="id">Unique id for the new pool</param>
+		/// <param name="name">Name of the new pool</param>
+		/// <param name="condition">Condition for agents to be included in this pool</param>
+		/// <param name="enableAutoscaling">Whether to enable autoscaling for this pool</param>
+		/// <param name="minAgents">Minimum number of agents in the pool</param>
+		/// <param name="numReserveAgents">Minimum number of idle agents to maintain</param>
+		/// <param name="conformInterval">Interval between conforms. Set to zero to disable.</param>
+		/// <param name="scaleOutCooldown">Cooldown time between scale-out events</param>
+		/// <param name="scaleInCooldown">Cooldown time between scale-in events</param>
+		/// <param name="sizeStrategies">Pool sizing strategies</param>
+		/// <param name="fleetManagers">Fleet managers</param>
+		/// <param name="sizeStrategy">Pool sizing strategy</param>
+		/// <param name="leaseUtilizationSettings">Settings for lease utilization strategy</param>
+		/// <param name="jobQueueSettings">Settings for job queue strategy</param>
+		/// <param name="computeQueueAwsMetricSettings">Settings for compute queue AWS metric strategy</param>
+		/// <param name="properties">Properties for the pool</param>
+		/// <returns>The new pool document</returns>
+		Task<IPool> AddAsync(
+			PoolId id,
+			string name,
+			Condition? condition = null,
+			bool? enableAutoscaling = null,
+			int? minAgents = null,
+			int? numReserveAgents = null,
+			TimeSpan? conformInterval = null,
+			TimeSpan? scaleOutCooldown = null,
+			TimeSpan? scaleInCooldown = null,
+			List<PoolSizeStrategyInfo>? sizeStrategies = null,
+			List<FleetManagerInfo>? fleetManagers = null,
+			PoolSizeStrategy? sizeStrategy = null,
+			LeaseUtilizationSettings? leaseUtilizationSettings = null,
+			JobQueueSettings? jobQueueSettings = null,
+			ComputeQueueAwsMetricSettings? computeQueueAwsMetricSettings = null,
+			IEnumerable<KeyValuePair<string, string>>? properties = null);
+
+		/// <summary>
+		/// Enumerates all the pools
+		/// </summary>
+		/// <returns>The pool documents</returns>
+		Task<List<IPool>> GetAsync();
+
+		/// <summary>
+		/// Gets a pool by ID
+		/// </summary>
+		/// <param name="id">Unique id of the pool</param>
+		/// <returns>The pool document</returns>
+		Task<IPool?> GetAsync(PoolId id);
+
+		/// <summary>
+		/// Gets a list of all valid pool ids
+		/// </summary>
+		/// <returns>List of pool ids</returns>
+		Task<List<PoolId>> GetPoolIdsAsync();
+
+		/// <summary>
+		/// Gets a pool by ID
+		/// </summary>
+		/// <param name="id">Unique id of the pool</param>
+		/// <returns>The pool document</returns>
+		Task<bool> DeleteAsync(PoolId id);
+
+		/// <summary>
+		/// Updates a pool
+		/// </summary>
+		/// <param name="pool">The pool to update</param>
+		/// <param name="newName">New name for the pool</param>
+		/// <param name="newCondition">New condition for the pool</param>
+		/// <param name="newEnableAutoscaling">New setting for whether to enable autoscaling</param>
+		/// <param name="newMinAgents">Minimum number of agents in the pool</param>
+		/// <param name="newNumReserveAgents">Minimum number of idle agents to maintain</param>
+		/// <param name="newWorkspaces">New workspaces for the pool</param>
+		/// <param name="newUseAutoSdk">New setting for whether to use autosdk</param>
+		/// <param name="newProperties">New properties for the pool</param>
+		/// <param name="conformInterval">Interval between conforms. Set to zero to disable.</param>
+		/// <param name="lastScaleUpTime">New time for last (auto) scale up</param>
+		/// <param name="lastScaleDownTime">New time for last (auto) scale down</param>
+		/// <param name="scaleOutCooldown">Cooldown time between scale-out events</param>
+		/// <param name="scaleInCooldown">Cooldown time between scale-in events</param>
+		/// <param name="lastScaleResult">Result from last scaling in/out attempt</param>
+		/// <param name="lastAgentCount">Last calculated agent count</param>
+		/// <param name="lastDesiredAgentCount">Last calculated desired agent count</param>
+		/// <param name="sizeStrategy">Pool sizing strategy</param>
+		/// <param name="newSizeStrategies">List of pool sizing strategies</param>
+		/// <param name="newFleetManagers">List of fleet managers</param>
+		/// <param name="leaseUtilizationSettings">Settings for lease utilization strategy</param>
+		/// <param name="jobQueueSettings">Settings for job queue strategy</param>
+		/// <param name="computeQueueAwsMetricSettings">Settings for compute queue AWS metric strategy</param>
+		/// <param name="useDefaultStrategy">Whether to use the default strategy</param>
+		/// <returns>Async task</returns>
+		Task<IPool?> TryUpdateAsync(
+			IPool pool,
+			string? newName = null,
+			Condition? newCondition = null,
+			bool? newEnableAutoscaling = null,
+			int? newMinAgents = null,
+			int? newNumReserveAgents = null,
+			List<AgentWorkspace>? newWorkspaces = null,
+			bool? newUseAutoSdk = null,
+			Dictionary<string, string?>? newProperties = null,
+			TimeSpan? conformInterval = null,
+			DateTime? lastScaleUpTime = null,
+			DateTime? lastScaleDownTime = null,
+			TimeSpan? scaleOutCooldown = null,
+			TimeSpan? scaleInCooldown = null,
+			ScaleResult? lastScaleResult = null,
+			int? lastAgentCount = null,
+			int? lastDesiredAgentCount = null,
+			PoolSizeStrategy? sizeStrategy = null,
+			List<PoolSizeStrategyInfo>? newSizeStrategies = null,
+			List<FleetManagerInfo>? newFleetManagers = null,
+			LeaseUtilizationSettings? leaseUtilizationSettings = null,
+			JobQueueSettings? jobQueueSettings = null,
+			ComputeQueueAwsMetricSettings? computeQueueAwsMetricSettings = null,
+			bool? useDefaultStrategy = null);
+
+		/// <summary>
+		/// Updates the list of pools from a config file
+		/// </summary>
+		/// <param name="poolConfigs">Configuration for the pools, and revision string for the config file containing them</param>
+		Task ConfigureAsync(IReadOnlyList<(PoolConfig Config, string Revision)> poolConfigs);
+	}
+}
diff --git a/Engine/Source/Programs/Horde/Horde.Server/Agents/Pools/PoolCollection.cs b/Engine/Source/Programs/Horde/Horde.Server/Agents/Pools/PoolCollection.cs
new file mode 100644
--- /dev/null
+++ b/Engine/Source/Programs/Horde/Horde.Server/Agents/Pools/PoolCollection.cs
@@ -0,0 +1,496 @@
+// Copyright Epic Games, Inc. All Rights Reserved.
+
+using System;
+using System.Collections.Generic;
+using System.Linq;
+using System.Threading.Tasks;
+using EpicGames.Horde.Common;
+using Horde.Server.Agents.Fleet;
+using Horde.Server.Server;
+using Horde.Server.Utilities;
+using Microsoft.Extensions.Options;
+using MongoDB.Bson;
+using MongoDB.Bson.Serialization.Attributes;
+using MongoDB.Driver;
+
+namespace Horde.Server.Agents.Pools
+{
+	/// <summary>
+	/// Collection of pool documents
+	/// </summary>
+	public class PoolCollection : IPoolCollection
+	{
+		class PoolDocument : IPool
+		{
+			[BsonId]
+			public PoolId Id { get; set; }
+
+			[BsonRequired]
+			public string Name { get; set; } = null!;
+
+			[BsonIgnoreIfNull]
+			public Condition? Condition { get; set; }
+
+			public List<AgentWorkspace> Workspaces { get; set; } = new List<AgentWorkspace>();
+
+			[BsonIgnoreIfDefault, BsonDefaultValue(true)]
+			public bool UseAutoSdk { get; set; } = true;
+
+			public Dictionary<string, string> Properties { get; set; } = new Dictionary<string, string>();
+
+			[BsonIgnoreIfDefault, BsonDefaultValue(true)]
+			public bool EnableAutoscaling { get; set; } = true;
+
+			[BsonIgnoreIfNull]
+			public int? MinAgents { get; set; }
+
+			[BsonIgnoreIfNull]
+			public int? NumReserveAgents { get; set; }
+
+			[BsonIgnoreIfNull]
+			public TimeSpan? ConformInterval { get; set; }
+
+			[BsonIgnoreIfNull]
+			public DateTime? LastScaleUpTime { get; set; }
+
+			[BsonIgnoreIfNull]
+			public DateTime? LastScaleDownTime { get; set; }
+
+			[BsonIgnoreIfNull]
+			public TimeSpan? ScaleOutCooldown { get; set; }
+
+			[BsonIgnoreIfNull]
+			public TimeSpan? ScaleInCooldown { get; set; }
+
+			[BsonIgnoreIfNull]
+			public ScaleResult? LastScaleResult { get; set; }
+
+			[BsonIgnoreIfNull]
+			public int? LastAgentCount { get; set; }
+			
+			[BsonIgnoreIfNull]
+			public int? LastDesiredAgentCount { get; set; }
+
+			[BsonIgnoreIfNull]
+			public PoolSizeStrategy? SizeStrategy { get; set; }
+
+			[BsonIgnoreIfNull]
+			public List<PoolSizeStrategyInfo>? SizeStrategies { get; set; }
+
+			[BsonIgnoreIfNull]
+			public List<FleetManagerInfo>? FleetManagers { get; set; }
+
+			[BsonIgnoreIfNull]
+			public LeaseUtilizationSettings? LeaseUtilizationSettings { get; set; }
+
+			[BsonIgnoreIfNull]
+			public JobQueueSettings? JobQueueSettings { get; set; }
+
+			[BsonIgnoreIfNull]
+			public ComputeQueueAwsMetricSettings? ComputeQueueAwsMetricSettings { get; set; }
+
+			[BsonIgnoreIfNull]
+			public string? Revision { get; set; }
+
+			public int UpdateIndex { get; set; }
+
+			// Read-only wrappers
+			IReadOnlyList<AgentWorkspace> IPool.Workspaces => Workspaces;
+			IReadOnlyDictionary<string, string> IPoolConfig.Properties => Properties;
+			IReadOnlyList<PoolSizeStrategyInfo> IPoolConfig.SizeStrategies => (IReadOnlyList<PoolSizeStrategyInfo>?)SizeStrategies ?? Array.Empty<PoolSizeStrategyInfo>();
+			IReadOnlyList<FleetManagerInfo> IPoolConfig.FleetManagers => (IReadOnlyList<FleetManagerInfo>?)FleetManagers ?? Array.Empty<FleetManagerInfo>();
+
+			public PoolDocument()
+			{
+			}
+
+			public PoolDocument(IPoolConfig other, string? otherRevision)
+			{
+				Id = other.Id;
+				Name = String.IsNullOrEmpty(other.Name)? other.Id.ToString() : other.Name;
+				Condition = other.Condition;
+				UseAutoSdk = other.UseAutoSdk;
+				Properties = new Dictionary<string, string>(other.Properties);
+				EnableAutoscaling = other.EnableAutoscaling;
+				MinAgents = other.MinAgents;
+				NumReserveAgents = other.NumReserveAgents;
+				ConformInterval = other.ConformInterval;
+				LastScaleUpTime = other.LastScaleUpTime;
+				LastScaleDownTime = other.LastScaleDownTime;
+				ScaleOutCooldown = other.ScaleOutCooldown;
+				ScaleInCooldown = other.ScaleInCooldown;
+				SizeStrategy = other.SizeStrategy;
+				if (other.SizeStrategies != null && other.SizeStrategies.Count > 0)
+				{
+					SizeStrategies = new List<PoolSizeStrategyInfo>(other.SizeStrategies);
+				}
+				if (other.FleetManagers != null && other.FleetManagers.Count > 0)
+				{
+					FleetManagers = new List<FleetManagerInfo>(other.FleetManagers);
+				}
+				LeaseUtilizationSettings = other.LeaseUtilizationSettings;
+				JobQueueSettings = other.JobQueueSettings;
+				ComputeQueueAwsMetricSettings = other.ComputeQueueAwsMetricSettings;
+				Revision = otherRevision;
+			}
+
+			public PoolDocument(IPool other)
+				: this(other, other.Revision)
+			{
+				Workspaces.AddRange(other.Workspaces);
+				UpdateIndex = other.UpdateIndex;
+			}
+		}
+
+		/// <summary>
+		/// Collection of pool documents
+		/// </summary>
+		readonly IMongoCollection<PoolDocument> _pools;
+
+		/// <summary>
+		/// Constructor
+		/// </summary>
+		/// <param name="mongoService">The database service instance</param>
+		public PoolCollection(MongoService mongoService)
+		{
+			_pools = mongoService.GetCollection<PoolDocument>("Pools");
+		}
+
+		/// <inheritdoc/>
+		public async Task<IPool> AddAsync(
+			PoolId id,
+			string name,
+			Condition? condition,
+			bool? enableAutoscaling,
+			int? minAgents,
+			int? numReserveAgents,
+			TimeSpan? conformInterval,
+			TimeSpan? scaleOutCooldown,
+			TimeSpan? scaleInCooldown,
+			List<PoolSizeStrategyInfo>? sizeStrategies,
+			List<FleetManagerInfo>? fleetManagers,
+			PoolSizeStrategy? sizeStrategy,
+			LeaseUtilizationSettings? leaseUtilizationSettings,
+			JobQueueSettings? jobQueueSettings,
+			ComputeQueueAwsMetricSettings? computeQueueAwsMetricSettings,
+			IEnumerable<KeyValuePair<string, string>>? properties)
+		{
+			PoolDocument pool = new PoolDocument();
+			pool.Id = id;
+			pool.Name = name;
+			pool.Condition = condition;
+			if (enableAutoscaling != null)
+			{
+				pool.EnableAutoscaling = enableAutoscaling.Value;
+			}
+			pool.MinAgents = minAgents;
+			pool.NumReserveAgents = numReserveAgents;
+			if (properties != null)
+			{
+				pool.Properties = new Dictionary<string, string>(properties);
+			}
+
+			pool.ConformInterval = conformInterval;
+			pool.ScaleOutCooldown = scaleOutCooldown;
+			pool.ScaleInCooldown = scaleInCooldown;
+			pool.SizeStrategy = sizeStrategy;
+			pool.LeaseUtilizationSettings = leaseUtilizationSettings;
+			pool.JobQueueSettings = jobQueueSettings;
+			pool.ComputeQueueAwsMetricSettings = computeQueueAwsMetricSettings;
+			if (sizeStrategies != null) { pool.SizeStrategies = sizeStrategies; }
+			if (fleetManagers != null) { pool.FleetManagers = fleetManagers; }
+			
+			await _pools.InsertOneAsync(pool);
+			return pool;
+		}
+
+		/// <inheritdoc/>
+		public async Task<List<IPool>> GetAsync()
+		{
+			List<PoolDocument> results = await _pools.Find(FilterDefinition<PoolDocument>.Empty).ToListAsync();
+			return results.ConvertAll<IPool>(x => x);
+		}
+
+		/// <inheritdoc/>
+		public async Task<IPool?> GetAsync(PoolId id)
+		{
+			return await _pools.Find(x => x.Id == id).FirstOrDefaultAsync();
+		}
+
+		/// <inheritdoc/>
+		public async Task<List<PoolId>> GetPoolIdsAsync()
+		{
+			ProjectionDefinition<PoolDocument, BsonDocument> projection = Builders<PoolDocument>.Projection.Include(x => x.Id);
+			List<BsonDocument> results = await _pools.Find(FilterDefinition<PoolDocument>.Empty).Project(projection).ToListAsync();
+			return results.ConvertAll(x => new PoolId(x.GetElement("_id").Value.AsString));
+		}
+
+		/// <inheritdoc/>
+		public async Task<bool> DeleteAsync(PoolId id)
+		{
+			FilterDefinition<PoolDocument> filter = Builders<PoolDocument>.Filter.Eq(x => x.Id, id);
+			DeleteResult result = await _pools.DeleteOneAsync(filter);
+			return result.DeletedCount > 0;
+		}
+
+		/// <summary>
+		/// Attempts to update a pool, upgrading it to the latest document schema if necessary
+		/// </summary>
+		/// <param name="poolToUpdate">Interface for the document to update</param>
+		/// <param name="transaction">The transaction to execute</param>
+		/// <returns>New pool definition, or null on failure.</returns>
+		async Task<IPool?> TryUpdateAsync(IPool poolToUpdate, TransactionBuilder<PoolDocument> transaction)
+		{
+			if (transaction.IsEmpty)
+			{
+				return poolToUpdate;
+			}
+
+			transaction.Set(x => x.UpdateIndex, poolToUpdate.UpdateIndex + 1);
+
+			PoolDocument? mutablePool = poolToUpdate as PoolDocument;
+			if (mutablePool != null)
+			{
+				UpdateResult result = await _pools.UpdateOneAsync(x => x.Id == poolToUpdate.Id && x.UpdateIndex == poolToUpdate.UpdateIndex, transaction.ToUpdateDefinition());
+				if (result.ModifiedCount == 0)
+				{
+					return null;
+				}
+
+				transaction.ApplyTo(mutablePool);
+			}
+			else
+			{
+				mutablePool = new PoolDocument(poolToUpdate);
+				transaction.ApplyTo(mutablePool);
+
+				ReplaceOneResult result = await _pools.ReplaceOneAsync<PoolDocument>(x => x.Id == poolToUpdate.Id && x.UpdateIndex == poolToUpdate.UpdateIndex, mutablePool);
+				if (result.ModifiedCount == 0)
+				{
+					return null;
+				}
+			}
+			return mutablePool;
+		}
+
+		/// <inheritdoc/>
+		public Task<IPool?> TryUpdateAsync(
+			IPool pool,
+			string? newName,
+			Condition? newCondition,
+			bool? newEnableAutoscaling,
+			int? newMinAgents,
+			int? newNumReserveAgents,
+			List<AgentWorkspace>? newWorkspaces,
+			bool? newUseAutoSdk,
+			Dictionary<string, string?>? newProperties,
+			TimeSpan? conformInterval,
+			DateTime? lastScaleUpTime,
+			DateTime? lastScaleDownTime,
+			TimeSpan? scaleOutCooldown,
+			TimeSpan? scaleInCooldown, 
+			ScaleResult? lastScaleResult, 
+			int? lastAgentCount, 
+			int? lastDesiredAgentCount, 
+			PoolSizeStrategy? sizeStrategy,
+			List<PoolSizeStrategyInfo>? newSizeStrategies,
+			List<FleetManagerInfo>? newFleetManagers,
+			LeaseUtilizationSettings? leaseUtilizationSettings,
+			JobQueueSettings? jobQueueSettings, 
+			ComputeQueueAwsMetricSettings? computeQueueAwsMetricSettings, 
+			bool? useDefaultStrategy)
+		{
+			TransactionBuilder<PoolDocument> transaction = new TransactionBuilder<PoolDocument>();
+			if (newName != null)
+			{
+				transaction.Set(x => x.Name, newName);
+			}
+			if (newCondition != null)
+			{
+				if (newCondition.IsEmpty())
+				{
+					transaction.Unset(x => x.Condition!);
+				}
+				else
+				{
+					transaction.Set(x => x.Condition, newCondition);
+				}
+			}
+			if (newEnableAutoscaling != null)
+			{
+				if (newEnableAutoscaling.Value)
+				{
+					transaction.Unset(x => x.EnableAutoscaling);
+				}
+				else
+				{
+					transaction.Set(x => x.EnableAutoscaling, newEnableAutoscaling.Value);
+				}
+			}
+			if (newMinAgents != null)
+			{
+				if (newMinAgents.Value < 0)
+				{
+					transaction.Unset(x => x.MinAgents!);
+				}
+				else
+				{
+					transaction.Set(x => x.MinAgents, newMinAgents.Value);
+				}
+			}
+			if (newNumReserveAgents != null)
+			{
+				if (newNumReserveAgents.Value < 0)
+				{
+					transaction.Unset(x => x.NumReserveAgents!);
+				}
+				else
+				{
+					transaction.Set(x => x.NumReserveAgents, newNumReserveAgents.Value);
+				}
+			}
+			if (newWorkspaces != null)
+			{
+				transaction.Set(x => x.Workspaces, newWorkspaces);
+			}
+			if (newUseAutoSdk != null)
+			{
+				transaction.Set(x => x.UseAutoSdk, newUseAutoSdk.Value);
+			}
+			if (newProperties != null)
+			{
+				transaction.UpdateDictionary(x => x.Properties, newProperties);
+			}
+			if (conformInterval != null)
+			{
+				if (conformInterval.Value < TimeSpan.Zero)
+				{
+					transaction.Unset(x => x.ConformInterval!);
+				}
+				else
+				{
+					transaction.Set(x => x.ConformInterval, conformInterval);
+				}
+			}
+			if (lastScaleUpTime != null)
+			{
+				transaction.Set(x => x.LastScaleUpTime, lastScaleUpTime);
+			}
+			if (lastScaleDownTime != null)
+			{
+				transaction.Set(x => x.LastScaleDownTime, lastScaleDownTime);
+			}
+			if (scaleOutCooldown != null)
+			{
+				transaction.Set(x => x.ScaleOutCooldown, scaleOutCooldown);
+			}
+			if (scaleInCooldown != null)
+			{
+				transaction.Set(x => x.ScaleInCooldown, scaleInCooldown);
+			}
+			if (lastScaleResult != null)
+			{
+				transaction.Set(x => x.LastScaleResult, lastScaleResult);
+			}
+			if (lastAgentCount != null)
+			{
+				transaction.Set(x => x.LastAgentCount, lastAgentCount);
+			}
+			if (lastDesiredAgentCount != null)
+			{
+				transaction.Set(x => x.LastDesiredAgentCount, lastDesiredAgentCount);
+			}
+			if (sizeStrategy != null)
+			{
+				transaction.Set(x => x.SizeStrategy, sizeStrategy);				
+			}
+			else if (useDefaultStrategy != null && useDefaultStrategy.Value)
+			{
+				transaction.Set(x => x.SizeStrategy, null);
+			}
+			if (newSizeStrategies != null)
+			{
+				transaction.Set(x => x.SizeStrategies, newSizeStrategies);				
+			}
+			if (newFleetManagers != null)
+			{
+				transaction.Set(x => x.FleetManagers, newFleetManagers);				
+			}
+			if (leaseUtilizationSettings != null)
+			{
+				transaction.Set(x => x.LeaseUtilizationSettings, leaseUtilizationSettings);
+			}
+			if (jobQueueSettings != null)
+			{
+				transaction.Set(x => x.JobQueueSettings, jobQueueSettings);
+			}
+			if (computeQueueAwsMetricSettings != null)
+			{
+				transaction.Set(x => x.ComputeQueueAwsMetricSettings, computeQueueAwsMetricSettings);
+			}
+			return TryUpdateAsync(pool, transaction);
+		}
+
+		/// <inheritdoc/>
+		public async Task ConfigureAsync(IReadOnlyList<(PoolConfig Config, string Revision)> poolConfigs)
+		{
+			List<PoolDocument> pools = await _pools.Find(FilterDefinition<PoolDocument>.Empty).ToListAsync();
+			Dictionary<PoolId, PoolDocument> idToPool = pools.ToDictionary(x => x.Id, x => x);
+
+			// Make sure we don't have any duplicates in the new pool configs before setting anything
+			Dictionary<PoolId, string> poolIdToRevision = new Dictionary<PoolId, string>();
+			foreach ((PoolConfig poolConfig, string revision) in poolConfigs)
+			{
+				if (poolIdToRevision.TryGetValue(poolConfig.Id, out string? prevRevision))
+				{
+					throw new PoolConflictException(poolConfig.Id, prevRevision, revision);
+				}
+				else
+				{
+					poolIdToRevision.Add(poolConfig.Id, revision);
+				}
+			}
+
+			// Update all the pools which are still present
+			foreach ((PoolConfig poolConfig, string revision) in poolConfigs)
+			{
+				idToPool.TryGetValue(poolConfig.Id, out PoolDocument? currentPool);
+				while (currentPool == null || currentPool.Revision != revision)
+				{
+					int updateIndex = currentPool?.UpdateIndex ?? 0;
+
+					PoolDocument newPool = new PoolDocument(poolConfig, revision);
+					if (currentPool != null)
+					{
+						newPool.Workspaces.AddRange(currentPool.Workspaces);
+						newPool.UpdateIndex = currentPool.UpdateIndex;
+					}
+
+					ReplaceOneResult result = await _pools.ReplaceOneAsync(x => x.Id == poolConfig.Id && x.UpdateIndex == updateIndex, newPool, new ReplaceOptions { IsUpsert = true });
+					if (result.MatchedCount > 0)
+					{
+						break;
+					}
+
+					currentPool = await _pools.Find(x => x.Id == poolConfig.Id).FirstOrDefaultAsync();
+				}
+				idToPool.Remove(poolConfig.Id);
+			}
+
+			// Remove anything that has been deleted
+			foreach (PoolDocument remainingPool in idToPool.Values)
+			{
+				PoolDocument? remainingPoolCopy = remainingPool;
+				while (remainingPoolCopy != null && !String.IsNullOrEmpty(remainingPoolCopy.Revision))
+				{
+					DeleteResult result = await _pools.DeleteOneAsync(x => x.Id == remainingPool.Id && x.UpdateIndex == remainingPoolCopy.UpdateIndex);
+					if (result.DeletedCount > 0)
+					{
+						break;
+					}
+					remainingPoolCopy = await _pools.Find(x => x.Id == remainingPool.Id).FirstOrDefaultAsync();
+				}
+			}
+		}
+	}
+}
diff --git a/Engine/Source/Programs/Horde/Horde.Server/Agents/Pools/PoolService.cs b/Engine/Source/Programs/Horde/Horde.Server/Agents/Pools/PoolService.cs
new file mode 100644
--- /dev/null
+++ b/Engine/Source/Programs/Horde/Horde.Server/Agents/Pools/PoolService.cs
@@ -0,0 +1,304 @@
+// Copyright Epic Games, Inc. All Rights Reserved.
+
+using System;
+using System.Collections.Generic;
+using System.Linq;
+using System.Threading;
+using System.Threading.Tasks;
+using EpicGames.Horde;
+using EpicGames.Horde.Common;
+using Horde.Server.Agents.Fleet;
+using Horde.Server.Server;
+using Horde.Server.Utilities;
+using HordeCommon;
+
+namespace Horde.Server.Agents.Pools
+{
+	/// <summary>
+	/// Wraps functionality for manipulating pools
+	/// </summary>
+	public class PoolService
+	{
+		/// <summary>
+		/// The globals service instance
+		/// </summary>
+		readonly GlobalsService _globalsService;
+
+		/// <summary>
+		/// Collection of pool documents
+		/// </summary>
+		readonly IPoolCollection _pools;
+
+		/// <summary>
+		/// Returns the current time
+		/// </summary>
+		readonly IClock _clock;
+
+		/// <summary>
+		/// Cached set of pools, along with the timestamp that it was obtained
+		/// </summary>
+		Tuple<DateTime, Dictionary<PoolId, IPool>>? _cachedPoolLookup;
+
+		/// <summary>
+		/// Constructor
+		/// </summary>
+		/// <param name="globalsService"></param>
+		/// <param name="pools">Collection of pool documents</param>
+		/// <param name="clock"></param>
+		public PoolService(GlobalsService globalsService, IPoolCollection pools, IClock clock)
+		{
+			_globalsService = globalsService;
+			_pools = pools;
+			_clock = clock;
+		}
+
+		/// <summary>
+		/// Creates a new pool
+		/// </summary>
+		/// <param name="name">Name of the new pool</param>
+		/// <param name="condition">Condition for agents to be automatically included in this pool</param>
+		/// <param name="enableAutoscaling">Whether to enable autoscaling for this pool</param>
+		/// <param name="minAgents">Minimum number of agents in the pool</param>
+		/// <param name="numReserveAgents">Minimum number of idle agents to maintain</param>
+		/// <param name="conformInterval">Interval between conforms. Set to zero to disable.</param>
+		/// <param name="scaleOutCooldown">Cooldown time between scale-out events</param>
+		/// <param name="scaleInCooldown">Cooldown time between scale-in events</param>
+		/// <param name="sizeStrategies">Pool sizing strategies</param>
+		/// <param name="fleetManagers">Fleet managers</param>
+		/// <param name="sizeStrategy">Pool sizing strategy</param>
+		/// <param name="leaseUtilizationSettings">Settings for lease utilization strategy</param>
+		/// <param name="jobQueueSettings">Settings for job queue strategy</param>
+		/// <param name="computeQueueAwsMetricSettings">Settings for compute queue AWS metric strategy</param>
+		/// <param name="properties">Properties for the new pool</param>
+		/// <returns>The new pool document</returns>
+		public Task<IPool> CreatePoolAsync(
+			string name,
+			Condition? condition = null,
+			bool? enableAutoscaling = null,
+			int? minAgents = null,
+			int? numReserveAgents = null,
+			TimeSpan? conformInterval = null,
+			TimeSpan? scaleOutCooldown = null,
+			TimeSpan? scaleInCooldown = null,
+			List<PoolSizeStrategyInfo>? sizeStrategies = null,
+			List<FleetManagerInfo>? fleetManagers = null,
+			PoolSizeStrategy? sizeStrategy = null,
+			LeaseUtilizationSettings? leaseUtilizationSettings = null,
+			JobQueueSettings? jobQueueSettings = null,
+			ComputeQueueAwsMetricSettings? computeQueueAwsMetricSettings = null,
+			Dictionary<string, string>? properties = null)
+		{
+			return _pools.AddAsync(
+				new PoolId(StringId.Sanitize(name)),
+				name,
+				condition,
+				enableAutoscaling,
+				minAgents,
+				numReserveAgents,
+				conformInterval,
+				scaleOutCooldown,
+				scaleInCooldown,
+				sizeStrategies,
+				fleetManagers,
+				sizeStrategy,
+				leaseUtilizationSettings,
+				jobQueueSettings,
+				computeQueueAwsMetricSettings,
+				properties);
+		}
+
+		/// <summary>
+		/// Deletes a pool
+		/// </summary>
+		/// <param name="poolId">Unique id of the pool</param>
+		/// <returns>Async task object</returns>
+		public Task<bool> DeletePoolAsync(PoolId poolId)
+		{
+			return _pools.DeleteAsync(poolId);
+		}
+
+		/// <summary>
+		/// Updates an existing pool
+		/// </summary>
+		/// <param name="pool">The pool to update</param>
+		/// <param name="newName">The new name for the pool</param>
+		/// <param name="newCondition">New requirements for the pool</param>
+		/// <param name="newEnableAutoscaling">Whether to enable autoscaling</param>
+		/// <param name="newMinAgents">Minimum number of agents in the pool</param>
+		/// <param name="newNumReserveAgents">Minimum number of idle agents to maintain</param>
+		/// <param name="newProperties">Properties on the pool to update. Any properties with a value of null will be removed.</param>
+		/// <param name="conformInterval">Interval between conforms. Set to zero to disable.</param>
+		/// <param name="scaleOutCooldown">Cooldown time between scale-out events</param>
+		/// <param name="scaleInCooldown">Cooldown time between scale-in events</param>
+		/// <param name="sizeStrategy">New pool sizing strategy for the pool</param>
+		/// <param name="newSizeStrategies">New pool sizing strategies for the pool</param>
+		/// <param name="newFleetManagers">New fleet managers for the pool</param>
+		/// <param name="leaseUtilizationSettings">Settings for lease utilization-based strategy</param>
+		/// <param name="jobQueueSettings">Settings for job queue-based strategy</param>
+		/// <param name="computeQueueAwsMetricSettings">Settings for compute queue AWS metric strategy</param>
+		/// <param name="useDefaultStrategy">Whether to use the default strategy</param>
+		/// <returns>Async task object</returns>
+		public async Task<IPool?> UpdatePoolAsync(
+			IPool? pool,
+			string? newName = null,
+			Condition? newCondition = null,
+			bool? newEnableAutoscaling = null,
+			int? newMinAgents = null,
+			int? newNumReserveAgents = null,
+			Dictionary<string, string?>? newProperties = null,
+			TimeSpan? conformInterval = null,
+			TimeSpan? scaleOutCooldown = null,
+			TimeSpan? scaleInCooldown = null,
+			PoolSizeStrategy? sizeStrategy = null,
+			List<PoolSizeStrategyInfo>? newSizeStrategies = null,
+			List<FleetManagerInfo>? newFleetManagers = null,
+			LeaseUtilizationSettings? leaseUtilizationSettings = null,
+			JobQueueSettings? jobQueueSettings = null,
+			ComputeQueueAwsMetricSettings? computeQueueAwsMetricSettings = null,
+			bool? useDefaultStrategy = null)
+		{
+			for (; pool != null; pool = await _pools.GetAsync(pool.Id))
+			{
+				IPool? newPool = await _pools.TryUpdateAsync(
+					pool,
+					newName,
+					newCondition,
+					newEnableAutoscaling,
+					newMinAgents,
+					newNumReserveAgents,
+					newProperties: newProperties,
+					conformInterval: conformInterval,
+					scaleOutCooldown: scaleOutCooldown,
+					scaleInCooldown: scaleInCooldown,
+					sizeStrategy: sizeStrategy,
+					newSizeStrategies: newSizeStrategies,
+					newFleetManagers: newFleetManagers,
+					leaseUtilizationSettings: leaseUtilizationSettings,
+					jobQueueSettings: jobQueueSettings,
+					computeQueueAwsMetricSettings: computeQueueAwsMetricSettings,
+					useDefaultStrategy: useDefaultStrategy);
+				
+				if (newPool != null)
+				{
+					return newPool;
+				}
+			}
+			return pool;
+		}
+
+		/// <summary>
+		/// Gets all the available pools
+		/// </summary>
+		/// <returns>List of pool documents</returns>
+		public Task<List<IPool>> GetPoolsAsync()
+		{
+			return _pools.GetAsync();
+		}
+
+		/// <summary>
+		/// Gets a pool by ID
+		/// </summary>
+		/// <param name="poolId">Unique id of the pool</param>
+		/// <returns>The pool document</returns>
+		public Task<IPool?> GetPoolAsync(PoolId poolId)
+		{
+			return _pools.GetAsync(poolId);
+		}
+
+		/// <summary>
+		/// Gets a cached pool definition
+		/// </summary>
+		/// <param name="poolId"></param>
+		/// <param name="validAtTime"></param>
+		/// <returns></returns>
+		public async Task<IPool?> GetCachedPoolAsync(PoolId poolId, DateTime validAtTime)
+		{
+			Dictionary<PoolId, IPool> poolMapping = await GetPoolLookupAsync(validAtTime);
+			poolMapping.TryGetValue(poolId, out IPool? pool);
+			return pool;
+		}
+
+		/// <summary>
+		/// Gets a cached pool definition
+		/// </summary>
+		/// <param name="agent"></param>
+		/// <param name="validAtTime"></param>
+		/// <returns></returns>
+		public async Task<List<IPool>> GetCachedPoolsAsync(IAgent agent, DateTime validAtTime)
+		{
+			Dictionary<PoolId, IPool> poolMapping = await GetPoolLookupAsync(validAtTime);
+
+			List<IPool> pools = new List<IPool>();
+			foreach(PoolId poolId in agent.GetPools())
+			{
+				if(poolMapping.TryGetValue(poolId, out IPool? pool))
+				{
+					pools.Add(pool);
+				}
+			}
+
+			return pools;
+		}
+
+		/// <summary>
+		/// Get a list of workspaces for the given agent
+		/// </summary>
+		/// <param name="agent">The agent to return workspaces for</param>
+		/// <param name="validAtTime">Absolute time at which we expect the results to be valid. Values may be cached as long as they are after this time.</param>
+		/// <param name="globalConfig">Current configuration</param>
+		/// <returns>List of workspaces</returns>
+		public async Task<HashSet<AgentWorkspace>> GetWorkspacesAsync(IAgent agent, DateTime validAtTime, GlobalConfig globalConfig)
+		{
+			bool addAutoSdkWorkspace = false;
+			HashSet<AgentWorkspace> workspaces = new HashSet<AgentWorkspace>();
+
+			Dictionary<PoolId, IPool> poolMapping = await GetPoolLookupAsync(validAtTime);
+			foreach (PoolId poolId in agent.GetPools())
+			{
+				IPool? pool;
+				if (poolMapping.TryGetValue(poolId, out pool))
+				{
+					workspaces.UnionWith(pool.Workspaces);
+					addAutoSdkWorkspace |= pool.UseAutoSdk;
+				}
+			}
+
+			if (addAutoSdkWorkspace)
+			{
+				workspaces.UnionWith(agent.GetAutoSdkWorkspaces(globalConfig, workspaces.ToList()));
+			}
+
+			return workspaces;
+		}
+
+		/// <summary>
+		/// Gets a mapping from pool identifiers to definitions
+		/// </summary>
+		/// <param name="validAtTime">Absolute time at which we expect the results to be valid. Values may be cached as long as they are after this time.</param>
+		/// <returns>Map of pool ids to pool documents</returns>
+		private async Task<Dictionary<PoolId, IPool>> GetPoolLookupAsync(DateTime validAtTime)
+		{
+			Tuple<DateTime, Dictionary<PoolId, IPool>>? cachedPoolLookupCopy = _cachedPoolLookup;
+			if (cachedPoolLookupCopy == null || cachedPoolLookupCopy.Item1 < validAtTime)
+			{
+				// Get a new list of cached pools
+				DateTime newCacheTime = _clock.UtcNow;
+				List<IPool> newPools = await _pools.GetAsync();
+				Tuple<DateTime, Dictionary<PoolId, IPool>> newCachedPoolLookup = Tuple.Create(newCacheTime, newPools.ToDictionary(x => x.Id, x => x));
+
+				// Try to swap it with the current version
+				while (cachedPoolLookupCopy == null || cachedPoolLookupCopy.Item1 < newCacheTime)
+				{
+					Tuple<DateTime, Dictionary<PoolId, IPool>>? originalValue = Interlocked.CompareExchange(ref _cachedPoolLookup, newCachedPoolLookup, cachedPoolLookupCopy);
+					if (originalValue == cachedPoolLookupCopy)
+					{
+						cachedPoolLookupCopy = newCachedPoolLookup;
+						break;
+					}
+					cachedPoolLookupCopy = originalValue;
+				}
+			}
+			return cachedPoolLookupCopy.Item2;
+		}
+	}
+}
