Horde: Rename the Horde.Build project to Horde.Server.

diff --git a/Engine/Source/Programs/Horde/Horde.Server/Agents/Pools/PoolUpdateService.cs b/Engine/Source/Programs/Horde/Horde.Server/Agents/Pools/PoolUpdateService.cs
new file mode 100644
--- /dev/null
+++ b/Engine/Source/Programs/Horde/Horde.Server/Agents/Pools/PoolUpdateService.cs
@@ -0,0 +1,140 @@
+// Copyright Epic Games, Inc. All Rights Reserved.
+
+using System;
+using System.Collections.Generic;
+using System.Linq;
+using System.Threading;
+using System.Threading.Tasks;
+using Horde.Server.Projects;
+using Horde.Server.Server;
+using Horde.Server.Streams;
+using Horde.Server.Utilities;
+using HordeCommon;
+using Microsoft.Extensions.Hosting;
+using Microsoft.Extensions.Logging;
+using Microsoft.Extensions.Options;
+
+namespace Horde.Server.Agents.Pools
+{
+	/// <summary>
+	/// Periodically updates pool documents to contain the correct workspaces
+	/// </summary>
+	public sealed class PoolUpdateService : IHostedService, IDisposable
+	{
+		readonly IPoolCollection _pools;
+		readonly IStreamCollection _streams;
+		readonly IOptionsMonitor<GlobalConfig> _globalConfig;
+		readonly ILogger<PoolService> _logger;
+		readonly ITicker _ticker;
+
+		/// <summary>
+		/// Constructor
+		/// </summary>
+		public PoolUpdateService(IPoolCollection pools, IStreamCollection streams, IClock clock, IOptionsMonitor<GlobalConfig> globalConfig, ILogger<PoolService> logger)
+		{
+			_pools = pools;
+			_streams = streams;
+			_globalConfig = globalConfig;
+			_logger = logger;
+			_ticker = clock.AddTicker($"{nameof(PoolUpdateService)}.{nameof(UpdatePoolsAsync)}", TimeSpan.FromSeconds(30.0), UpdatePoolsAsync, logger);
+		}
+
+		/// <inheritdoc/>
+		public Task StartAsync(CancellationToken cancellationToken) => _ticker.StartAsync();
+
+		/// <inheritdoc/>
+		public Task StopAsync(CancellationToken cancellationToken) => _ticker.StopAsync();
+
+		/// <inheritdoc/>
+		public void Dispose() => _ticker.Dispose();
+
+		/// <summary>
+		/// Execute the background task
+		/// </summary>
+		/// <param name="stoppingToken">Cancellation token for the async task</param>
+		/// <returns>Async task</returns>
+		async ValueTask UpdatePoolsAsync(CancellationToken stoppingToken)
+		{
+			// Capture the start time for this operation. We use this to attempt to sequence updates to agents, and prevent overriding another server's updates.
+			DateTime startTime = DateTime.UtcNow;
+
+			// Update the list
+			bool retryUpdate = true;
+			while (retryUpdate && !stoppingToken.IsCancellationRequested)
+			{
+				_logger.LogDebug("Updating pool->workspace map");
+
+				// Assume this will be the last iteration
+				retryUpdate = false;
+
+				// Capture the list of pools at the start of this update
+				List<IPool> currentPools = await _pools.GetAsync();
+
+				// Lookup table of pool id to workspaces
+				HashSet<PoolId> poolsWithAutoSdk = new HashSet<PoolId>();
+				Dictionary<PoolId, List<AgentWorkspace>> poolToAgentWorkspaces = new Dictionary<PoolId, List<AgentWorkspace>>();
+
+				// Capture the current config state
+				GlobalConfig globalConfig = _globalConfig.CurrentValue;
+
+				// Populate the workspace list from the current stream
+				foreach (StreamConfig streamConfig in globalConfig.Streams)
+				{
+					foreach (KeyValuePair<string, AgentConfig> agentTypePair in streamConfig.AgentTypes)
+					{
+						// Create the new agent workspace
+						(AgentWorkspace, bool)? result;
+						if (streamConfig.TryGetAgentWorkspace(agentTypePair.Value, out result))
+						{
+							(AgentWorkspace agentWorkspace, bool useAutoSdk) = result.Value;
+							AgentConfig agentType = agentTypePair.Value;
+
+							// Find or add a list of workspaces for this pool
+							List<AgentWorkspace>? agentWorkspaces;
+							if (!poolToAgentWorkspaces.TryGetValue(agentType.Pool, out agentWorkspaces))
+							{
+								agentWorkspaces = new List<AgentWorkspace>();
+								poolToAgentWorkspaces.Add(agentType.Pool, agentWorkspaces);
+							}
+
+							// Add it to the list
+							if (!agentWorkspaces.Contains(agentWorkspace))
+							{
+								agentWorkspaces.Add(agentWorkspace);
+							}
+							if (useAutoSdk)
+							{
+								poolsWithAutoSdk.Add(agentType.Pool);
+							}
+						}
+					}
+				}
+
+				// Update the list of workspaces for each pool
+				foreach (IPool currentPool in currentPools)
+				{
+					// Get the new list of workspaces for this pool
+					List<AgentWorkspace>? newWorkspaces;
+					if (!poolToAgentWorkspaces.TryGetValue(currentPool.Id, out newWorkspaces))
+					{
+						newWorkspaces = new List<AgentWorkspace>();
+					}
+
+					// Update the pools document
+					bool useAutoSdk = poolsWithAutoSdk.Contains(currentPool.Id);
+					if (!AgentWorkspace.SetEquals(currentPool.Workspaces, newWorkspaces) || currentPool.Workspaces.Count != newWorkspaces.Count || currentPool.UseAutoSdk != useAutoSdk)
+					{
+						_logger.LogInformation("New workspaces for pool {Pool}:{Workspaces}", currentPool.Id, String.Join("", newWorkspaces.Select(x => $"\n  Identifier=\"{x.Identifier}\", Stream={x.Stream}")));
+
+						IPool? result = await _pools.TryUpdateAsync(currentPool, newWorkspaces: newWorkspaces, newUseAutoSdk: useAutoSdk);
+						if (result == null)
+						{
+							_logger.LogInformation("Pool modified; will retry");
+							retryUpdate = true;
+						}
+					}
+				}
+			}
+		}
+	}
+}
diff --git a/Engine/Source/Programs/Horde/Horde.Server/Agents/Pools/PoolsController.cs b/Engine/Source/Programs/Horde/Horde.Server/Agents/Pools/PoolsController.cs
new file mode 100644
--- /dev/null
+++ b/Engine/Source/Programs/Horde/Horde.Server/Agents/Pools/PoolsController.cs
@@ -0,0 +1,205 @@
+// Copyright Epic Games, Inc. All Rights Reserved.
+
+using System;
+using System.Collections.Generic;
+using System.Linq;
+using System.Threading.Tasks;
+using Horde.Server.Acls;
+using Horde.Server.Agents.Fleet;
+using Horde.Server.Server;
+using Horde.Server.Utilities;
+using Microsoft.AspNetCore.Authorization;
+using Microsoft.AspNetCore.Mvc;
+using Microsoft.Extensions.Options;
+
+namespace Horde.Server.Agents.Pools
+{
+	/// <summary>
+	/// Controller for the /api/v1/pools endpoint
+	/// </summary>
+	[ApiController]
+	[Authorize]
+	[Route("[controller]")]
+	public class PoolsController : HordeControllerBase
+	{
+		private readonly PoolService _poolService;
+		private readonly IOptionsSnapshot<GlobalConfig> _globalConfig;
+
+		/// <summary>
+		/// Constructor
+		/// </summary>
+		public PoolsController(PoolService poolService, IOptionsSnapshot<GlobalConfig> globalConfig)
+		{
+			_poolService = poolService;
+			_globalConfig = globalConfig;
+		}
+
+		/// <summary>
+		/// Creates a new pool
+		/// </summary>
+		/// <param name="create">Parameters for the new pool.</param>
+		/// <returns>Http result code</returns>
+		[HttpPost]
+		[Route("/api/v1/pools")]
+		public async Task<ActionResult<CreatePoolResponse>> CreatePoolAsync([FromBody] CreatePoolRequest create)
+		{
+			if(!_globalConfig.Value.Authorize(AclAction.CreatePool, User))
+			{
+				return Forbid(AclAction.CreatePool);
+			}
+
+			List<PoolSizeStrategyInfo>? sizeStrategies = create.SizeStrategies?.Select(x => x.Convert()).ToList();
+			List<FleetManagerInfo>? fleetManagers = create.FleetManagers?.Select(x => x.Convert()).ToList();
+			LeaseUtilizationSettings? luSettings = create.LeaseUtilizationSettings?.Convert();
+			JobQueueSettings? jqSettings = create.JobQueueSettings?.Convert();
+			ComputeQueueAwsMetricSettings? cqamSettings = create.ComputeQueueAwsMetricSettings?.Convert();
+
+			TimeSpan? conformInterval = create.ConformInterval == null ? null : TimeSpan.FromHours(create.ConformInterval.Value);
+			TimeSpan? scaleOutCooldown = create.ScaleOutCooldown == null ? null : TimeSpan.FromSeconds(create.ScaleOutCooldown.Value);
+			TimeSpan? scaleInCooldown = create.ScaleInCooldown == null ? null : TimeSpan.FromSeconds(create.ScaleInCooldown.Value);
+
+			IPool newPool = await _poolService.CreatePoolAsync(
+				create.Name, create.Condition, create.EnableAutoscaling, create.MinAgents, create.NumReserveAgents, conformInterval,
+				scaleOutCooldown, scaleInCooldown, sizeStrategies, fleetManagers, create.SizeStrategy, luSettings, jqSettings, cqamSettings, create.Properties);
+			return new CreatePoolResponse(newPool.Id.ToString());
+		}
+
+		/// <summary>
+		/// Query all the pools
+		/// </summary>
+		/// <param name="filter">Filter for the properties to return</param>
+		/// <returns>Information about all the pools</returns>
+		[HttpGet]
+		[Route("/api/v1/pools")]
+		[ProducesResponseType(typeof(List<GetPoolResponse>), 200)]
+		public async Task<ActionResult<List<object>>> GetPoolsAsync([FromQuery] PropertyFilter? filter = null)
+		{
+			if (!_globalConfig.Value.Authorize(AclAction.ListPools, User))
+			{
+				return Forbid(AclAction.ListPools);
+			}
+
+			List<IPool> pools = await _poolService.GetPoolsAsync();
+
+			List<object> responses = new List<object>();
+			foreach (IPool pool in pools)
+			{
+				responses.Add(new GetPoolResponse(pool).ApplyFilter(filter));
+			}
+			return responses;
+		}
+
+		/// <summary>
+		/// Retrieve information about a specific pool
+		/// </summary>
+		/// <param name="poolId">Id of the pool to get information about</param>
+		/// <param name="filter">Filter to apply to the returned properties</param>
+		/// <returns>Information about the requested pool</returns>
+		[HttpGet]
+		[Route("/api/v1/pools/{poolId}")]
+		[ProducesResponseType(typeof(GetPoolResponse), 200)]
+		public async Task<ActionResult<object>> GetPoolAsync(string poolId, [FromQuery] PropertyFilter? filter = null)
+		{
+			if (!_globalConfig.Value.Authorize(AclAction.ViewPool, User))
+			{
+				return Forbid(AclAction.ViewPool);
+			}
+
+			PoolId poolIdValue = new PoolId(poolId);
+
+			IPool? pool = await _poolService.GetPoolAsync(poolIdValue);
+			if (pool == null)
+			{
+				return NotFound(poolIdValue);
+			}
+
+			return new GetPoolResponse(pool).ApplyFilter(filter);
+		}
+
+		/// <summary>
+		/// Update a pool's properties.
+		/// </summary>
+		/// <param name="poolId">Id of the pool to update</param>
+		/// <param name="update">Items on the pool to update</param>
+		/// <returns>Http result code</returns>
+		[HttpPut]
+		[Route("/api/v1/pools/{poolId}")]
+		public async Task<ActionResult> UpdatePoolAsync(string poolId, [FromBody] UpdatePoolRequest update)
+		{
+			if (!_globalConfig.Value.Authorize(AclAction.UpdatePool, User))
+			{
+				return Forbid(AclAction.UpdatePool);
+			}
+
+			PoolId poolIdValue = new PoolId(poolId);
+			
+			IPool? pool = await _poolService.GetPoolAsync(poolIdValue);
+			if(pool == null)
+			{
+				return NotFound(poolIdValue);
+			}
+
+			List<PoolSizeStrategyInfo>? newSizeStrategies = update.SizeStrategies?.Select(x => x.Convert()).ToList();
+			List<FleetManagerInfo>? newFleetManagers = update.FleetManagers?.Select(x => x.Convert()).ToList();
+			TimeSpan? conformInterval = update.ConformInterval == null ? null : TimeSpan.FromHours(update.ConformInterval.Value);
+			TimeSpan? scaleOutCooldown = update.ScaleOutCooldown == null ? null : TimeSpan.FromSeconds(update.ScaleOutCooldown.Value);
+			TimeSpan? scaleInCooldown = update.ScaleInCooldown == null ? null : TimeSpan.FromSeconds(update.ScaleInCooldown.Value);
+
+			await _poolService.UpdatePoolAsync(pool, update.Name, update.Condition, update.EnableAutoscaling,
+				update.MinAgents, update.NumReserveAgents, update.Properties, conformInterval, scaleOutCooldown, scaleInCooldown, update.SizeStrategy, newSizeStrategies, newFleetManagers,
+				update.LeaseUtilizationSettings?.Convert(), update.JobQueueSettings?.Convert(), update.ComputeQueueAwsMetricSettings?.Convert(), update.UseDefaultStrategy);
+			return new OkResult();
+		}
+
+		/// <summary>
+		/// Delete a pool
+		/// </summary>
+		/// <param name="poolId">Id of the pool to delete</param>
+		/// <returns>Http result code</returns>
+		[HttpDelete]
+		[Route("/api/v1/pools/{poolId}")]
+		public async Task<ActionResult> DeletePoolAsync(string poolId)
+		{
+			if (!_globalConfig.Value.Authorize(AclAction.DeletePool, User))
+			{
+				return Forbid(AclAction.DeletePool);
+			}
+
+			PoolId poolIdValue = new PoolId(poolId);
+			if(!await _poolService.DeletePoolAsync(poolIdValue))
+			{
+				return NotFound(poolIdValue);
+			}
+			return Ok();
+		}
+
+		/// <summary>
+		/// Batch update pool properties
+		/// </summary>
+		/// <param name="batchUpdates">List of pools to update</param>
+		/// <returns>Http result code</returns>
+		[HttpPut]
+		[Route("/api/v1/pools")]
+		public async Task<ActionResult> UpdatePoolAsync([FromBody] List<BatchUpdatePoolRequest> batchUpdates)
+		{
+			if (!_globalConfig.Value.Authorize(AclAction.UpdatePool, User))
+			{
+				return Forbid(AclAction.UpdatePool);
+			}
+
+			foreach (BatchUpdatePoolRequest update in batchUpdates)
+			{
+				PoolId poolIdValue = new PoolId(update.Id);
+
+				IPool? pool = await _poolService.GetPoolAsync(poolIdValue);
+				if (pool == null)
+				{
+					return NotFound(poolIdValue);
+				}
+
+				await _poolService.UpdatePoolAsync(pool, update.Name, newProperties: update.Properties);
+			}
+			return Ok();
+		}
+	}
+}
diff --git a/Engine/Source/Programs/Horde/Horde.Server/Agents/Sessions/ISessionCollection.cs b/Engine/Source/Programs/Horde/Horde.Server/Agents/Sessions/ISessionCollection.cs
new file mode 100644
--- /dev/null
+++ b/Engine/Source/Programs/Horde/Horde.Server/Agents/Sessions/ISessionCollection.cs
@@ -0,0 +1,69 @@
+// Copyright Epic Games, Inc. All Rights Reserved.
+
+using System;
+using System.Collections.Generic;
+using System.Threading.Tasks;
+using Horde.Server.Utilities;
+
+namespace Horde.Server.Agents.Sessions
+{
+	/// <summary>
+	/// Interface for a collection of session documents
+	/// </summary>
+	public interface ISessionCollection
+	{
+		/// <summary>
+		/// Adds a new session to the collection
+		/// </summary>
+		/// <param name="id">The session id</param>
+		/// <param name="agentId">The agent this session is for</param>
+		/// <param name="startTime">Start time of this session</param>
+		/// <param name="properties">Properties of this agent at the time the session started</param>
+		/// <param name="resources">Resources which the agent has</param>
+		/// <param name="version">Version of the agent software</param>
+		Task<ISession> AddAsync(SessionId id, AgentId agentId, DateTime startTime, IReadOnlyList<string>? properties, IReadOnlyDictionary<string, int>? resources, string? version);
+
+		/// <summary>
+		/// Gets information about a particular session
+		/// </summary>
+		/// <param name="sessionId">The unique session id</param>
+		/// <returns>The session information</returns>
+		Task<ISession?> GetAsync(SessionId sessionId);
+
+		/// <summary>
+		/// Find sessions for the given agent
+		/// </summary>
+		/// <param name="agentId">The unique agent id</param>
+		/// <param name="startTime">Start time to include in the search</param>
+		/// <param name="finishTime">Finish time to include in the search</param>
+		/// <param name="index">Index of the first result to return</param>
+		/// <param name="count">Number of results to return</param>
+		/// <returns>List of sessions matching the given criteria</returns>
+		Task<List<ISession>> FindAsync(AgentId agentId, DateTime? startTime, DateTime? finishTime, int index, int count);
+
+		/// <summary>
+		/// Finds any active sessions
+		/// </summary>
+		/// <param name="index">Index of the first result to return</param>
+		/// <param name="count">Number of results to return</param>
+		/// <returns>List of sessions</returns>
+		Task<List<ISession>> FindActiveSessionsAsync(int? index = null, int? count = null);
+
+		/// <summary>
+		/// Update a session from the collection
+		/// </summary>
+		/// <param name="sessionId">The session to update</param>
+		/// <param name="finishTime">Time at which the session finished</param>
+		/// <param name="properties">The agent properties</param>
+		/// <param name="resources">Resources which the agent has</param>
+		/// <returns>Async task</returns>
+		Task UpdateAsync(SessionId sessionId, DateTime finishTime, IReadOnlyList<string> properties, IReadOnlyDictionary<string, int> resources);
+
+		/// <summary>
+		/// Delete a session from the collection
+		/// </summary>
+		/// <param name="sessionId">The session id</param>
+		/// <returns>Async task</returns>
+		Task DeleteAsync(SessionId sessionId);
+	}
+}
diff --git a/Engine/Source/Programs/Horde/Horde.Server/Agents/Sessions/SessionCollection.cs b/Engine/Source/Programs/Horde/Horde.Server/Agents/Sessions/SessionCollection.cs
new file mode 100644
--- /dev/null
+++ b/Engine/Source/Programs/Horde/Horde.Server/Agents/Sessions/SessionCollection.cs
@@ -0,0 +1,136 @@
+// Copyright Epic Games, Inc. All Rights Reserved.
+
+using System;
+using System.Collections.Generic;
+using System.Linq;
+using System.Threading.Tasks;
+using Horde.Server.Server;
+using Horde.Server.Utilities;
+using MongoDB.Bson.Serialization.Attributes;
+using MongoDB.Driver;
+
+namespace Horde.Server.Agents.Sessions
+{
+	/// <summary>
+	/// Collection of session documents
+	/// </summary>
+	public class SessionCollection : ISessionCollection
+	{
+		/// <summary>
+		/// Concrete implementation of ISession
+		/// </summary>
+		class SessionDocument : ISession
+		{
+			[BsonRequired, BsonId]
+			public SessionId Id { get; set; }
+
+			[BsonRequired]
+			public AgentId AgentId { get; set; }
+
+			public DateTime StartTime { get; set; }
+			public DateTime? FinishTime { get; set; }
+			public List<string>? Properties { get; set; }
+			public Dictionary<string, int>? Resources { get; set; }
+			public string? Version { get; set; }
+
+			IReadOnlyList<string>? ISession.Properties => Properties;
+
+			[BsonConstructor]
+			private SessionDocument()
+			{
+			}
+
+			public SessionDocument(SessionId id, AgentId agentId, DateTime startTime, IReadOnlyList<string>? properties, IReadOnlyDictionary<string, int>? resources, string? version)
+			{
+				Id = id;
+				AgentId = agentId;
+				StartTime = startTime;
+				if (properties != null)
+				{
+					Properties = new List<string>(properties);
+				}
+				if (resources != null)
+				{
+					Resources = new Dictionary<string, int>(resources);
+				}
+				Version = version;
+			}
+		}
+
+		/// <summary>
+		/// Collection of session documents
+		/// </summary>
+		readonly IMongoCollection<SessionDocument> _sessions;
+
+		/// <summary>
+		/// Constructor
+		/// </summary>
+		/// <param name="mongoService">The database service</param>
+		public SessionCollection(MongoService mongoService)
+		{
+			List<MongoIndex<SessionDocument>> indexes = new List<MongoIndex<SessionDocument>>();
+			indexes.Add(keys => keys.Ascending(x => x.AgentId));
+			indexes.Add(keys => keys.Ascending(x => x.StartTime));
+			indexes.Add(keys => keys.Ascending(x => x.FinishTime));
+
+			_sessions = mongoService.GetCollection<SessionDocument>("Sessions", indexes);
+		}
+
+		/// <inheritdoc/>
+		public async Task<ISession> AddAsync(SessionId id, AgentId agentId, DateTime startTime, IReadOnlyList<string>? properties, IReadOnlyDictionary<string, int>? resources, string? version)
+		{
+			SessionDocument newSession = new SessionDocument(id, agentId, startTime, properties, resources, version);
+			await _sessions.InsertOneAsync(newSession);
+			return newSession;
+		}
+
+		/// <inheritdoc/>
+		public async Task<ISession?> GetAsync(SessionId sessionId)
+		{
+			return await _sessions.Find(x => x.Id == sessionId).FirstOrDefaultAsync();
+		}
+
+		/// <inheritdoc/>
+		public async Task<List<ISession>> FindAsync(AgentId agentId, DateTime? startTime, DateTime? finishTime, int index, int count)
+		{
+			FilterDefinitionBuilder<SessionDocument> filterBuilder = Builders<SessionDocument>.Filter;
+
+			FilterDefinition<SessionDocument> filter = filterBuilder.Eq(x => x.AgentId, agentId);
+			if (startTime != null)
+			{
+				filter &= filterBuilder.Gte(x => x.StartTime, startTime.Value);
+			}
+			if (finishTime != null)
+			{
+				filter &= filterBuilder.Or(filterBuilder.Eq(x => x.FinishTime, null), filterBuilder.Lte(x => x.FinishTime, finishTime.Value));
+			}
+
+			List<SessionDocument> results = await _sessions.Find(filter).SortByDescending(x => x.StartTime).Skip(index).Limit(count).ToListAsync();
+			return results.ConvertAll<ISession>(x => x);
+		}
+
+		/// <inheritdoc/>
+		public async Task<List<ISession>> FindActiveSessionsAsync(int? index, int? count)
+		{
+			List<SessionDocument> results = await _sessions.Find(x => x.FinishTime == null).Range(index, count).ToListAsync();
+			return results.ConvertAll<ISession>(x => x);
+		}
+
+		/// <inheritdoc/>
+		public Task UpdateAsync(SessionId sessionId, DateTime finishTime, IReadOnlyList<string> properties, IReadOnlyDictionary<string, int> resources)
+		{
+			UpdateDefinition<SessionDocument> update = Builders<SessionDocument>.Update
+				.Set(x => x.FinishTime, finishTime)
+				.Set(x => x.Properties, new List<string>(properties))
+				.Set(x => x.Resources, new Dictionary<string, int>(resources));
+
+			return _sessions.FindOneAndUpdateAsync(x => x.Id == sessionId, update);
+		}
+
+		/// <inheritdoc/>
+		public Task DeleteAsync(SessionId sessionId)
+		{
+			return _sessions.DeleteOneAsync(x => x.Id == sessionId);
+		}
+	}
+}
diff --git a/Engine/Source/Programs/Horde/Horde.Server/Agents/Sessions/SessionId.cs b/Engine/Source/Programs/Horde/Horde.Server/Agents/Sessions/SessionId.cs
new file mode 100644
--- /dev/null
+++ b/Engine/Source/Programs/Horde/Horde.Server/Agents/Sessions/SessionId.cs
@@ -0,0 +1,62 @@
+// Copyright Epic Games, Inc. All Rights Reserved.
+
+using System.ComponentModel;
+using EpicGames.Core;
+using Horde.Server.Utilities;
+using MongoDB.Bson;
+
+namespace Horde.Server.Agents.Sessions
+{
+	/// <summary>
+	/// Identifier for a session
+	/// </summary>
+	/// <param name="Id">Id to construct from</param>
+	[JsonSchemaString]
+	[TypeConverter(typeof(ObjectIdTypeConverter<SessionId, SessionIdConverter>))]
+	[ObjectIdConverter(typeof(SessionIdConverter))]
+	public record struct SessionId(ObjectId Id)
+	{
+		/// <summary>
+		/// Default empty value for session id
+		/// </summary>
+		public static SessionId Empty { get; } = default;
+
+		/// <summary>
+		/// Creates a new <see cref="SessionId"/>
+		/// </summary>
+		public static SessionId GenerateNewId() => new SessionId(ObjectId.GenerateNewId());
+
+		/// <inheritdoc cref="ObjectId.Parse(System.String)"/>
+		public static SessionId Parse(string text) => new SessionId(ObjectId.Parse(text));
+
+		/// <inheritdoc cref="ObjectId.TryParse(System.String, out ObjectId)"/>
+		public static bool TryParse(string text, out SessionId id)
+		{
+			if (ObjectId.TryParse(text, out ObjectId objectId))
+			{
+				id = new SessionId(objectId);
+				return true;
+			}
+			else
+			{
+				id = default;
+				return false;
+			}
+		}
+
+		/// <inheritdoc/>
+		public override string ToString() => Id.ToString();
+	}
+
+	/// <summary>
+	/// Converter to and from <see cref="ObjectId"/> instances.
+	/// </summary>
+	class SessionIdConverter : ObjectIdConverter<SessionId>
+	{
+		/// <inheritdoc/>
+		public override SessionId FromObjectId(ObjectId id) => new SessionId(id);
+
+		/// <inheritdoc/>
+		public override ObjectId ToObjectId(SessionId value) => value.Id;
+	}
+}
diff --git a/Engine/Source/Programs/Horde/Horde.Server/Agents/Telemetry/ITelemetryCollection.cs b/Engine/Source/Programs/Horde/Horde.Server/Agents/Telemetry/ITelemetryCollection.cs
new file mode 100644
--- /dev/null
+++ b/Engine/Source/Programs/Horde/Horde.Server/Agents/Telemetry/ITelemetryCollection.cs
@@ -0,0 +1,35 @@
+// Copyright Epic Games, Inc. All Rights Reserved.
+
+using System;
+using System.Collections.Generic;
+using System.Threading.Tasks;
+
+namespace Horde.Server.Agents.Telemetry
+{
+	/// <summary>
+	/// Collection of utilization collection
+	/// </summary>
+	public interface ITelemetryCollection
+	{
+		/// <summary>
+		/// Adds entries for the given utilization
+		/// </summary>
+		/// <param name="telemetry">Telemetry data to add</param>
+		/// <returns>Async task</returns>
+		Task AddUtilizationTelemetryAsync(IUtilizationTelemetry telemetry);
+
+		/// <summary>
+		/// Finds utilization data matching the given criteria
+		/// </summary>
+		/// <param name="startTimeUtc">Start time to query utilization for</param>
+		/// <param name="finishTimeUtc">Finish time to query utilization for</param>
+		/// <returns>The utilization data</returns>
+		Task<List<IUtilizationTelemetry>> GetUtilizationTelemetryAsync(DateTime startTimeUtc, DateTime finishTimeUtc);
+
+		/// <summary>
+		/// Finds the latest utilization data
+		/// </summary>
+		/// <returns>The utilization data</returns>
+		Task<IUtilizationTelemetry?> GetLatestUtilizationTelemetryAsync();
+	}
+}
diff --git a/Engine/Source/Programs/Horde/Horde.Server/Agents/Telemetry/Telemetry.cs b/Engine/Source/Programs/Horde/Horde.Server/Agents/Telemetry/Telemetry.cs
new file mode 100644
--- /dev/null
+++ b/Engine/Source/Programs/Horde/Horde.Server/Agents/Telemetry/Telemetry.cs
@@ -0,0 +1,226 @@
+// Copyright Epic Games, Inc. All Rights Reserved.
+
+using System;
+using System.Collections.Generic;
+using Horde.Server.Agents.Pools;
+using Horde.Server.Streams;
+using Horde.Server.Utilities;
+
+namespace Horde.Server.Agents.Telemetry
+{
+	/// <summary>
+	/// Information about the utilization of a pool
+	/// </summary>
+	public interface IUtilizationTelemetry
+	{
+		/// <summary>
+		/// Start time for the bucket
+		/// </summary>
+		public DateTime StartTime { get; }
+
+		/// <summary>
+		/// Finish time for this bucket
+		/// </summary>
+		public DateTime FinishTime { get; }
+
+		/// <summary>
+		/// Number of agents in total
+		/// </summary>
+		public int NumAgents { get; }
+
+		/// <summary>
+		/// Breakdown of utilization by pool
+		/// </summary>
+		public IReadOnlyList<IPoolUtilizationTelemetry> Pools { get; }
+
+		/// <summary>
+		/// Amount of time that agents were hibernating
+		/// </summary>
+		public double HibernatingTime { get; }
+
+		/// <summary>
+		/// Total time spent running administrative tasks (conform, etc...)
+		/// </summary>
+		public double AdminTime { get; }
+	}
+
+	/// <summary>
+	/// Concrete implementation of <see cref="IUtilizationTelemetry"/>
+	/// </summary>
+	public sealed class NewUtilizationTelemetry : IUtilizationTelemetry
+	{
+		/// <inheritdoc/>
+		public DateTime StartTime { get; set; }
+
+		/// <inheritdoc/>
+		public DateTime FinishTime { get; set; }
+
+		/// <inheritdoc/>
+		public int NumAgents { get; set; }
+
+		/// <inheritdoc/>
+		List<NewPoolUtilizationTelemetry> Pools { get; set; } = new List<NewPoolUtilizationTelemetry>();
+		IReadOnlyList<IPoolUtilizationTelemetry> IUtilizationTelemetry.Pools => Pools;
+
+		Dictionary<PoolId, NewPoolUtilizationTelemetry> PoolsLookup { get; set; } = new Dictionary<PoolId, NewPoolUtilizationTelemetry>();
+
+		/// <inheritdoc/>
+		public double HibernatingTime { get; set; }
+
+		/// <inheritdoc/>
+		public double AdminTime { get; set; }
+
+		/// <summary>
+		/// Constructor
+		/// </summary>
+		/// <param name="startTime"></param>
+		/// <param name="finishTime"></param>
+		public NewUtilizationTelemetry(DateTime startTime, DateTime finishTime)
+		{
+			StartTime = startTime;
+			FinishTime = finishTime;
+		}
+
+		/// <summary>
+		/// Adds a new pool
+		/// </summary>
+		/// <param name="poolId">The pool id</param>
+		/// <returns>Telemetry for the given pool</returns>
+		public NewPoolUtilizationTelemetry FindOrAddPool(PoolId poolId)
+		{
+			NewPoolUtilizationTelemetry? pool;
+			if (!PoolsLookup.TryGetValue(poolId, out pool))
+			{
+				pool = new NewPoolUtilizationTelemetry(poolId);
+				Pools.Add(pool);
+				PoolsLookup.Add(poolId, pool);
+			}
+			return pool;
+		}
+	}
+
+	/// <summary>
+	/// Information about the utilization of a pool
+	/// </summary>
+	public interface IPoolUtilizationTelemetry
+	{
+		/// <summary>
+		/// Pool containing the work to execute
+		/// </summary>
+		public PoolId PoolId { get; }
+
+		/// <summary>
+		/// Number of agents in this pool
+		/// </summary>
+		public int NumAgents { get; }
+
+		/// <summary>
+		/// The stream executing work. If this is null, the time accounts for the machine executing work in another pool.
+		/// </summary>
+		public IReadOnlyList<IStreamUtilizationTelemetry> Streams { get; }
+
+		/// <summary>
+		/// Amount of time spent running 
+		/// </summary>
+		public double AdminTime { get; }
+
+		/// <summary>
+		/// Amount of time that agents were hibernating
+		/// </summary>
+		public double HibernatingTime { get; }
+
+		/// <summary>
+		/// Amount of time spent by agents in this pool servicing other pools
+		/// </summary>
+		public double OtherTime { get; }
+	}
+
+	/// <summary>
+	/// Concrete implementation of <see cref="IPoolUtilizationTelemetry"/>
+	/// </summary>
+	public sealed class NewPoolUtilizationTelemetry : IPoolUtilizationTelemetry
+	{
+		/// <inheritdoc/>
+		public PoolId PoolId { get; set; }
+
+		/// <inheritdoc/>
+		public int NumAgents { get; set; }
+
+		/// <inheritdoc/>
+		List<NewPoolUtilizationTelemetryStream> Streams { get; set; } = new List<NewPoolUtilizationTelemetryStream>();
+		IReadOnlyList<IStreamUtilizationTelemetry> IPoolUtilizationTelemetry.Streams => Streams;
+
+		Dictionary<StreamId, NewPoolUtilizationTelemetryStream> StreamLookup { get; set; } = new Dictionary<StreamId, NewPoolUtilizationTelemetryStream>();
+
+		/// <inheritdoc/>
+		public double AdminTime { get; set; }
+
+		/// <inheritdoc/>
+		public double HibernatingTime { get; set; }
+
+		/// <inheritdoc/>
+		public double OtherTime { get; set; }
+
+		/// <summary>
+		/// Constructor
+		/// </summary>
+		/// <param name="poolId"></param>
+		public NewPoolUtilizationTelemetry(PoolId poolId)
+		{
+			PoolId = poolId;
+		}
+
+		/// <summary>
+		/// Adds a stream to the object
+		/// </summary>
+		/// <param name="streamId"></param>
+		public NewPoolUtilizationTelemetryStream FindOrAddStream(StreamId streamId)
+		{
+			NewPoolUtilizationTelemetryStream? stream;
+			if (!StreamLookup.TryGetValue(streamId, out stream))
+			{
+				stream = new NewPoolUtilizationTelemetryStream(streamId);
+				Streams.Add(stream);
+				StreamLookup.Add(streamId, stream);
+			}
+			return stream;
+		}
+	}
+
+	/// <summary>
+	/// Utilization of a pool for a particular stream
+	/// </summary>
+	public interface IStreamUtilizationTelemetry
+	{
+		/// <summary>
+		/// The stream id
+		/// </summary>
+		public StreamId StreamId { get; }
+
+		/// <summary>
+		/// Number of machine hours spent executing work for this stream
+		/// </summary>
+		public double Time { get; }
+	}
+
+	/// <summary>
+	/// Concrete implementation of <see cref="IStreamUtilizationTelemetry"/>
+	/// </summary>
+	public sealed class NewPoolUtilizationTelemetryStream : IStreamUtilizationTelemetry
+	{
+		/// <inheritdoc/>
+		public StreamId StreamId { get; set; }
+
+		/// <inheritdoc/>
+		public double Time { get; set; }
+
+		/// <summary>
+		/// Constructor
+		/// </summary>
+		/// <param name="streamId"></param>
+		public NewPoolUtilizationTelemetryStream(StreamId streamId)
+		{
+			StreamId = streamId;
+		}
+	}
+}
diff --git a/Engine/Source/Programs/Horde/Horde.Server/Agents/Telemetry/TelemetryCollection.cs b/Engine/Source/Programs/Horde/Horde.Server/Agents/Telemetry/TelemetryCollection.cs
new file mode 100644
--- /dev/null
+++ b/Engine/Source/Programs/Horde/Horde.Server/Agents/Telemetry/TelemetryCollection.cs
@@ -0,0 +1,130 @@
+// Copyright Epic Games, Inc. All Rights Reserved.
+
+using System;
+using System.Collections.Generic;
+using System.Linq;
+using System.Threading.Tasks;
+using EpicGames.Core;
+using Horde.Server.Agents.Pools;
+using Horde.Server.Server;
+using Horde.Server.Streams;
+using Horde.Server.Utilities;
+using MongoDB.Bson.Serialization.Attributes;
+using MongoDB.Driver;
+
+namespace Horde.Server.Agents.Telemetry
+{
+	/// <summary>
+	/// Collection of utilization data
+	/// </summary>
+	public class TelemetryCollection : ITelemetryCollection
+	{
+		class UtilizationDocument : IUtilizationTelemetry
+		{
+			public DateTime StartTime { get; set; }
+			public DateTime FinishTime { get; set; }
+
+			public int NumAgents { get; set; }
+			public List<PoolUtilizationDocument> Pools { get; set; }
+			IReadOnlyList<IPoolUtilizationTelemetry> IUtilizationTelemetry.Pools => Pools;
+
+			public double AdminTime { get; set; }
+			public double HibernatingTime { get; set; }
+			public int UpdateIndex { get; set; }
+
+			[BsonConstructor]
+			private UtilizationDocument()
+			{
+				Pools = new List<PoolUtilizationDocument>();
+			}
+
+			public UtilizationDocument(IUtilizationTelemetry other)
+			{
+				StartTime = other.StartTime;
+				FinishTime = other.FinishTime;
+				NumAgents = other.NumAgents;
+				Pools = other.Pools.ConvertAll(x => new PoolUtilizationDocument(x));
+				HibernatingTime = other.HibernatingTime;
+				AdminTime = other.AdminTime;
+			}
+		}
+
+		class PoolUtilizationDocument : IPoolUtilizationTelemetry
+		{
+			public PoolId PoolId { get; set; }
+			public int NumAgents { get; set; }
+
+			public List<StreamUtilizationDocument> Streams { get; set; } 
+			IReadOnlyList<IStreamUtilizationTelemetry> IPoolUtilizationTelemetry.Streams => Streams;
+
+			public double AdminTime { get; set; }
+			public double HibernatingTime { get; set; }
+			public double OtherTime { get; set; }
+
+			[BsonConstructor]
+			private PoolUtilizationDocument()
+			{
+				Streams = new List<StreamUtilizationDocument>();
+			}
+
+			public PoolUtilizationDocument(IPoolUtilizationTelemetry other)
+			{
+				PoolId = other.PoolId;
+				NumAgents = other.NumAgents;
+				Streams = other.Streams.ConvertAll(x => new StreamUtilizationDocument(x));
+				AdminTime = other.AdminTime;
+				HibernatingTime = other.HibernatingTime;
+				OtherTime = other.OtherTime;
+			}
+		}
+
+		class StreamUtilizationDocument : IStreamUtilizationTelemetry
+		{
+			public StreamId StreamId { get; set; }
+			public double Time { get; set; }
+
+			[BsonConstructor]
+			private StreamUtilizationDocument()
+			{
+			}
+
+			public StreamUtilizationDocument(IStreamUtilizationTelemetry other)
+			{
+				StreamId = other.StreamId;
+				Time = other.Time;
+			}
+		}
+
+		readonly IMongoCollection<UtilizationDocument> _utilization;
+
+		/// <summary>
+		/// Constructor
+		/// </summary>
+		/// <param name="database"></param>
+		public TelemetryCollection(MongoService database)
+		{
+			_utilization = database.GetCollection<UtilizationDocument>("Utilization", keys => keys.Ascending(x => x.FinishTime).Ascending(x => x.StartTime));
+		}
+
+		/// <inheritdoc/>
+		public async Task AddUtilizationTelemetryAsync(IUtilizationTelemetry newTelemetry)
+		{
+			await _utilization.InsertOneAsync(new UtilizationDocument(newTelemetry));
+		}
+
+		/// <inheritdoc/>
+		public async Task<List<IUtilizationTelemetry>> GetUtilizationTelemetryAsync(DateTime startTimeUtc, DateTime finishTimeUtc)
+		{
+			FilterDefinition<UtilizationDocument> filter = Builders<UtilizationDocument>.Filter.Gte(x => x.FinishTime, startTimeUtc) & Builders<UtilizationDocument>.Filter.Lte(x => x.StartTime, finishTimeUtc);
+
+			List<UtilizationDocument> documents = await _utilization.Find(filter).ToListAsync();
+			return documents.ConvertAll<IUtilizationTelemetry>(x => x);
+		}
+
+		/// <inheritdoc/>
+		public async Task<IUtilizationTelemetry?> GetLatestUtilizationTelemetryAsync()
+		{
+			return await _utilization.Find(FilterDefinition<UtilizationDocument>.Empty).SortByDescending(x => x.FinishTime).FirstOrDefaultAsync();
+		}
+	}
+}
diff --git a/Engine/Source/Programs/Horde/Horde.Server/Agents/Telemetry/TelemetryController.Messages.cs b/Engine/Source/Programs/Horde/Horde.Server/Agents/Telemetry/TelemetryController.Messages.cs
new file mode 100644
--- /dev/null
+++ b/Engine/Source/Programs/Horde/Horde.Server/Agents/Telemetry/TelemetryController.Messages.cs
@@ -0,0 +1,136 @@
+// Copyright Epic Games, Inc. All Rights Reserved.
+
+using System;
+using System.Collections.Generic;
+using EpicGames.Core;
+
+namespace Horde.Server.Agents.Telemetry
+{
+
+	/// <summary>
+	/// Represents one stream in one pool in one hour of telemetry
+	/// </summary>
+	public class UtilizationTelemetryStream
+	{
+		/// <summary>
+		/// Stream Id
+		/// </summary>
+		public string StreamId { get; set; }
+
+		/// <summary>
+		/// Total time
+		/// </summary>
+		public double Time { get; set; }
+
+		/// <summary>
+		/// Constructor
+		/// </summary>
+		/// <param name="stream">The stream</param>
+		public UtilizationTelemetryStream(IStreamUtilizationTelemetry stream)
+		{
+			StreamId = stream.StreamId.ToString();
+			Time = stream.Time;
+		}
+	}
+
+	/// <summary>
+	/// Representation of an hour of time
+	/// </summary>
+	public class UtilizationTelemetryPool
+	{
+		/// <summary>
+		/// Pool id
+		/// </summary>
+		public string PoolId { get; set; }
+
+		/// <summary>
+		/// Number of agents in this pool
+		/// </summary>
+		public int NumAgents { get; set; }
+
+		/// <summary>
+		/// Total time spent doing admin work
+		/// </summary>
+		public double AdminTime { get; set; }
+
+		/// <summary>
+		/// Time spent hibernating
+		/// </summary>
+		public double HibernatingTime { get; }
+
+		/// <summary>
+		/// Total time agents in this pool were doing work for other pools
+		/// </summary>
+		public double OtherTime { get; set; }
+
+		/// <summary>
+		/// List of streams
+		/// </summary>
+		public List<UtilizationTelemetryStream> Streams { get; set; }
+
+		/// <summary>
+		/// Constructor
+		/// </summary>
+		/// <param name="pool"></param>
+		public UtilizationTelemetryPool(IPoolUtilizationTelemetry pool)
+		{
+			PoolId = pool.PoolId.ToString();
+			NumAgents = pool.NumAgents;
+			AdminTime = pool.AdminTime;
+			HibernatingTime = pool.HibernatingTime;
+			OtherTime = pool.OtherTime;
+
+			Streams = pool.Streams.ConvertAll(stream => new UtilizationTelemetryStream(stream));
+		}
+	}
+
+	/// <summary>
+	/// Response data for a utilization request
+	/// </summary>
+	public class UtilizationTelemetryResponse
+	{
+		/// <summary>
+		/// Start hour
+		/// </summary>
+		public DateTimeOffset StartTime { get; set; }
+
+		/// <summary>
+		/// End hour
+		/// </summary>
+		public DateTimeOffset FinishTime { get; set; }
+
+		/// <summary>
+		/// List of pools
+		/// </summary>
+		public List<UtilizationTelemetryPool> Pools { get; set; }
+
+		/// <summary>
+		/// Total admin time
+		/// </summary>
+		public double AdminTime { get; set; }
+
+		/// <summary>
+		/// Total hibernating time
+		/// </summary>
+		public double HibernatingTime { get; set; }
+
+		/// <summary>
+		/// Total agents
+		/// </summary>
+		public int NumAgents { get; set; }
+
+		/// <summary>
+		/// Constructor
+		/// </summary>
+		public UtilizationTelemetryResponse(IUtilizationTelemetry telemetry)
+		{
+			StartTime = telemetry.StartTime;
+			FinishTime = telemetry.FinishTime;
+			AdminTime = telemetry.AdminTime;
+			HibernatingTime = telemetry.HibernatingTime;
+			NumAgents = telemetry.NumAgents;
+
+			Pools = telemetry.Pools.ConvertAll(pool => new UtilizationTelemetryPool(pool));
+		}
+	}
+}
diff --git a/Engine/Source/Programs/Horde/Horde.Server/Agents/Telemetry/TelemetryController.cs b/Engine/Source/Programs/Horde/Horde.Server/Agents/Telemetry/TelemetryController.cs
new file mode 100644
--- /dev/null
+++ b/Engine/Source/Programs/Horde/Horde.Server/Agents/Telemetry/TelemetryController.cs
@@ -0,0 +1,50 @@
+// Copyright Epic Games, Inc. All Rights Reserved.
+
+using System;
+using System.Collections.Generic;
+using System.Threading.Tasks;
+using Microsoft.AspNetCore.Mvc;
+
+namespace Horde.Server.Agents.Telemetry
+{
+	/// <summary>
+	/// Controller for the /api/v1/reports endpoint, used for the reports pages
+	/// </summary>
+	[ApiController]
+	[Route("[controller]")]
+	public sealed class ReportsController : ControllerBase
+	{
+		/// <summary>
+		/// the Telemetry collection singleton
+		/// </summary>
+		readonly ITelemetryCollection _telemetryCollection;
+
+		/// <summary>
+		/// Constructor
+		/// </summary>
+		/// <param name="telemetryCollection">The telemetry collection</param>
+		public ReportsController(ITelemetryCollection telemetryCollection)
+		{
+			_telemetryCollection = telemetryCollection;
+		}
+
+		/// <summary>
+		/// Gets a collection of utilization data including an endpoint and a range
+		/// </summary>
+		[HttpGet]
+		[Route("/api/v1/reports/utilization/{endDate}")]
+		public async Task<ActionResult<List<UtilizationTelemetryResponse>>> GetStreamUtilizationData(DateTime endDate, [FromQuery(Name = "Range")] int range, [FromQuery(Name = "TzOffset")] int? tzOffset)
+		{
+			// Logic here is a bit messy. The client is always passing in a date at midnight
+			// If user passes in 12/1/2020 into the date with range of 1, the range should be 12/1/2020:00:00:00 to 12/1/2020:23:59:59
+			// Range 2 should be 11/30/2020:00:00:00 to 12/1/2020:23:59:59
+			int offset = tzOffset ?? 0;
+			DateTimeOffset endDateOffset = new DateTimeOffset(endDate, TimeSpan.FromHours(offset)).Add(new TimeSpan(23, 59, 59));
+			DateTimeOffset startDateOffset = endDate.Subtract(new TimeSpan(range - 1, 0, 0, 0));
+
+			List<IUtilizationTelemetry> telemetry = await _telemetryCollection.GetUtilizationTelemetryAsync(startDateOffset.UtcDateTime, endDateOffset.UtcDateTime);
+
+			return telemetry.ConvertAll(telemetry => new UtilizationTelemetryResponse(telemetry));
+		}
+	}
+}
diff --git a/Engine/Source/Programs/Horde/Horde.Server/Agents/Telemetry/TelemetryService.cs b/Engine/Source/Programs/Horde/Horde.Server/Agents/Telemetry/TelemetryService.cs
new file mode 100644
--- /dev/null
+++ b/Engine/Source/Programs/Horde/Horde.Server/Agents/Telemetry/TelemetryService.cs
@@ -0,0 +1,141 @@
+// Copyright Epic Games, Inc. All Rights Reserved.
+
+using System;
+using System.Collections.Generic;
+using System.Linq;
+using System.Threading;
+using System.Threading.Tasks;
+using Horde.Server.Agents.Fleet;
+using Horde.Server.Agents.Leases;
+using Horde.Server.Agents.Pools;
+using Horde.Server.Utilities;
+using HordeCommon;
+using Microsoft.Extensions.Hosting;
+using Microsoft.Extensions.Logging;
+
+namespace Horde.Server.Agents.Telemetry
+{
+	/// <summary>
+	/// Service which updates telemetry periodically
+	/// </summary>
+	public sealed class TelemetryService : IHostedService, IDisposable
+	{
+		readonly ITelemetryCollection _telemetryCollection;
+		readonly IAgentCollection _agentCollection;
+		readonly ILeaseCollection _leaseCollection;
+		readonly IPoolCollection _poolCollection;
+		readonly IFleetManager _fleetManager;
+		readonly IClock _clock;
+		readonly ITicker _tick;
+		readonly ILogger _logger;
+
+		/// <summary>
+		/// Constructor
+		/// </summary>
+		public TelemetryService(ITelemetryCollection telemetryCollection, IAgentCollection agentCollection, ILeaseCollection leaseCollection, IPoolCollection poolCollection, IFleetManager fleetManager, IClock clock, ILogger<TelemetryService> logger)
+		{
+			_telemetryCollection = telemetryCollection;
+			_agentCollection = agentCollection;
+			_leaseCollection = leaseCollection;
+			_poolCollection = poolCollection;
+			_fleetManager = fleetManager;
+			_clock = clock;
+			_tick = clock.AddSharedTicker<TelemetryService>(TimeSpan.FromMinutes(10.0), TickLeaderAsync, logger);
+			_logger = logger;
+		}
+
+		/// <inheritdoc/>
+		public Task StartAsync(CancellationToken cancellationToken) => _tick.StartAsync();
+
+		/// <inheritdoc/>
+		public Task StopAsync(CancellationToken cancellationToken) => _tick.StopAsync();
+
+		/// <inheritdoc/>
+		public void Dispose() => _tick.Dispose();
+
+		/// <inheritdoc/>
+		async ValueTask TickLeaderAsync(CancellationToken stoppingToken)
+		{
+			DateTime currentTime = _clock.UtcNow;
+
+			// Find the last time that we need telemetry for
+			DateTime maxTime = currentTime.Date + TimeSpan.FromHours(currentTime.Hour);
+
+			// Get the latest telemetry data
+			IUtilizationTelemetry? latest = await _telemetryCollection.GetLatestUtilizationTelemetryAsync();
+			TimeSpan interval = TimeSpan.FromHours(1.0);
+			int count = (latest == null) ? (7 * 24) : (int)Math.Round((maxTime - latest.FinishTime) / interval);
+			DateTime minTime = maxTime - count * interval;
+
+			// Query all the current data
+			List<IAgent> agents = await _agentCollection.FindAsync();
+			List<IPool> pools = await _poolCollection.GetAsync();
+			List<ILease> leases = await _leaseCollection.FindLeasesAsync(minTime: minTime);
+
+			// Remove any agents which are offline
+			agents.RemoveAll(x => !x.Enabled || !x.IsSessionValid(currentTime));
+
+			// Find all the agents
+			Dictionary<AgentId, List<PoolId>> agentToPoolIds = agents.ToDictionary(x => x.Id, x => x.GetPools().ToList());
+
+			// Generate all the telemetry data
+			DateTime bucketMinTime = minTime;
+			for (int idx = 0; idx < count; idx++)
+			{
+				DateTime bucketMaxTime = bucketMinTime + interval;
+				_logger.LogInformation("Creating telemetry for {MinTime} to {MaxTime}", bucketMinTime, bucketMaxTime);
+
+				NewUtilizationTelemetry telemetry = new NewUtilizationTelemetry(bucketMinTime, bucketMaxTime);
+				telemetry.NumAgents = agents.Count;
+				foreach (IPool pool in pools)
+				{
+					if (pool.EnableAutoscaling)
+					{
+						int numStoppedInstances = await _fleetManager.GetNumStoppedInstancesAsync(pool, stoppingToken);
+						telemetry.NumAgents += numStoppedInstances;
+
+						NewPoolUtilizationTelemetry poolTelemetry = telemetry.FindOrAddPool(pool.Id);
+						poolTelemetry.NumAgents += numStoppedInstances;
+						poolTelemetry.HibernatingTime += interval.TotalHours * numStoppedInstances;
+					}
+				}
+				foreach (PoolId poolId in agentToPoolIds.Values.SelectMany(x => x))
+				{
+					telemetry.FindOrAddPool(poolId).NumAgents++;
+				}
+				foreach (ILease lease in leases)
+				{
+					if (lease.StartTime < bucketMaxTime && (!lease.FinishTime.HasValue || lease.FinishTime >= bucketMinTime))
+					{
+						List<PoolId>? leasePools;
+						if (agentToPoolIds.TryGetValue(lease.AgentId, out leasePools))
+						{
+							DateTime finishTime = lease.FinishTime ?? bucketMaxTime;
+							double time = new TimeSpan(Math.Min(finishTime.Ticks, bucketMaxTime.Ticks) - Math.Max(lease.StartTime.Ticks, bucketMinTime.Ticks)).TotalHours;
+
+							foreach (PoolId poolId in leasePools)
+							{
+								NewPoolUtilizationTelemetry poolTelemetry = telemetry.FindOrAddPool(poolId);
+								if (lease.PoolId == null || lease.StreamId == null)
+								{
+									poolTelemetry.AdminTime += time;
+								}
+								else if (poolId == lease.PoolId)
+								{
+									poolTelemetry.FindOrAddStream(lease.StreamId.Value).Time += time;
+								}
+								else
+								{
+									poolTelemetry.OtherTime += time;
+								}
+							}
+						}
+					}
+				}
+				await _telemetryCollection.AddUtilizationTelemetryAsync(telemetry);
+
+				bucketMinTime = bucketMaxTime;
+			}
+		}
+	}
+}
diff --git a/Engine/Source/Programs/Horde/Horde.Server/Auditing/AuditLog.cs b/Engine/Source/Programs/Horde/Horde.Server/Auditing/AuditLog.cs
new file mode 100644
--- /dev/null
+++ b/Engine/Source/Programs/Horde/Horde.Server/Auditing/AuditLog.cs
@@ -0,0 +1,276 @@
+// Copyright Epic Games, Inc. All Rights Reserved.
+
+using System;
+using System.Collections.Generic;
+using System.Text;
+using System.Threading.Channels;
+using System.Threading.Tasks;
+using EpicGames.Core;
+using Horde.Server.Server;
+using Horde.Server.Utilities;
+using Microsoft.Extensions.Logging;
+using MongoDB.Bson;
+using MongoDB.Bson.Serialization.Attributes;
+using MongoDB.Driver;
+
+namespace Horde.Server.Auditing
+{
+	class AuditLog<TSubject> : IAuditLog<TSubject>, IAsyncDisposable, IDisposable
+	{
+		class AuditLogMessage : IAuditLogMessage<TSubject>
+		{
+			public ObjectId Id { get; set; }
+
+			[BsonElement("s")]
+			public TSubject Subject { get; set; }
+
+			[BsonElement("t")]
+			public DateTime TimeUtc { get; set; }
+
+			[BsonElement("l")]
+			public LogLevel Level { get; set; }
+
+			[BsonElement("d")]
+			public string Data { get; set; }
+
+			public AuditLogMessage()
+			{
+				Subject = default!;
+				Data = String.Empty;
+			}
+
+			public AuditLogMessage(TSubject subject, DateTime timeUtc, LogLevel level, string data)
+			{
+				Id = ObjectId.GenerateNewId();
+				Subject = subject;
+				TimeUtc = timeUtc;
+				Level = level;
+				Data = data;
+			}
+		}
+
+		class AuditLogChannel : IAuditLogChannel<TSubject>
+		{
+			sealed class Scope : IDisposable
+			{
+				readonly AuditLogChannel _owner;
+
+				public LogEvent Message { get; }
+
+				public Scope(AuditLogChannel owner, LogEvent message)
+				{
+					_owner = owner;
+					Message = message;
+
+					_owner._scopes.Add(this);
+				}
+
+				public void Dispose() => _owner._scopes.Remove(this);
+			}
+
+			public readonly AuditLog<TSubject> Outer;
+			public TSubject Subject { get; }
+
+			readonly List<Scope> _scopes = new List<Scope>();
+
+			public AuditLogChannel(AuditLog<TSubject> outer, TSubject subject)
+			{
+				Outer = outer;
+				Subject = subject;
+			}
+
+			public IDisposable BeginScope<TState>(TState state) => new Scope(this, LogEvent.FromState(LogLevel.Information, default, state, null, (x, y) => x?.ToString() ?? String.Empty));
+
+			public bool IsEnabled(LogLevel logLevel) => true;
+
+			LogEvent CreateEvent<TState>(LogLevel logLevel, EventId eventId, TState state, Exception? exception, Func<TState, Exception?, string> formatter)
+			{
+				LogEvent logEvent = LogEvent.FromState(logLevel, eventId, state, exception, formatter);
+				if (_scopes.Count > 0)
+				{
+					StringBuilder message = new StringBuilder();
+					StringBuilder format = new StringBuilder();
+					Dictionary<string, object> properties = new Dictionary<string, object>(StringComparer.Ordinal);
+
+					for (int idx = 0; idx < _scopes.Count; idx++)
+					{
+						message.Append('[');
+						format.Append('[');
+						AppendMessage(_scopes[idx].Message, idx + 1, message, format, properties);
+						message.Append(']');
+						format.Append(']');
+					}
+
+					message.Append(' ');
+					format.Append(' ');
+
+					AppendMessage(logEvent, 0, message, format, properties);
+
+					logEvent.Message = message.ToString();
+					logEvent.Format = format.ToString();
+					logEvent.Properties = properties;
+				}
+				return logEvent;
+			}
+
+			static void AppendMessage(LogEvent logEvent, int id, StringBuilder message, StringBuilder format, Dictionary<string, object> properties)
+			{
+				message.Append(logEvent.Message);
+				if (logEvent.Format == null)
+				{
+					format.Append($"{{_scope{id}}}");
+					properties.Add($"_scope{id}", logEvent.Message);
+				}
+				else
+				{
+					format.Append(logEvent.Format);
+					if (logEvent.Properties != null)
+					{
+						foreach (KeyValuePair<string, object> item in logEvent.Properties)
+						{
+							properties[item.Key] = item.Value;
+						}
+					}
+				}
+			}
+
+			public void Log<TState>(LogLevel logLevel, EventId eventId, TState state, Exception? exception, Func<TState, Exception?, string> formatter)
+			{
+				DateTime time = DateTime.UtcNow;
+
+				LogEvent logEvent = CreateEvent(logLevel, eventId, state, exception, formatter);
+				string data = logEvent.ToJson();
+				AuditLogMessage message = new AuditLogMessage(Subject, time, logLevel, data);
+				Outer._messageChannel.Writer.TryWrite(message);
+
+#pragma warning disable CA2254 // Template should be a static expression
+				using (IDisposable _ = Outer._logger.BeginScope($"Subject: {{{Outer._subjectProperty}}}", Subject))
+				{
+					Outer._logger.Log(logLevel, eventId, state, exception, formatter);
+				}
+#pragma warning restore CA2254 // Template should be a static expression
+			}
+
+			public IAsyncEnumerable<IAuditLogMessage> FindAsync(DateTime? minTime, DateTime? maxTime, int? index, int? count) => Outer.FindAsync(Subject, minTime, maxTime, index, count);
+
+			public Task<long> DeleteAsync(DateTime? minTime, DateTime? maxTime) => Outer.DeleteAsync(Subject, minTime, maxTime);
+		}
+
+		readonly IMongoCollection<AuditLogMessage> _messages;
+		readonly Channel<AuditLogMessage> _messageChannel;
+		readonly string _subjectProperty;
+		readonly ILogger _logger;
+		readonly Task _backgroundTask;
+
+		public IAuditLogChannel<TSubject> this[TSubject subject] => new AuditLogChannel(this, subject);
+
+		public AuditLog(MongoService mongoService, string collectionName, string subjectProperty, ILogger logger)
+		{
+			List<MongoIndex<AuditLogMessage>> indexes = new List<MongoIndex<AuditLogMessage>>();
+			indexes.Add(builder => builder.Ascending(x => x.Subject).Descending(x => x.TimeUtc));
+
+			_messages = mongoService.GetCollection<AuditLogMessage>(collectionName, indexes);
+
+			_messageChannel = Channel.CreateUnbounded<AuditLogMessage>();
+			_subjectProperty = subjectProperty;
+			_logger = logger;
+
+			_backgroundTask = Task.Run(() => WriteMessagesAsync());
+		}
+
+		public async ValueTask DisposeAsync()
+		{
+			_messageChannel.Writer.TryComplete();
+			await _backgroundTask;
+		}
+
+		public void Dispose()
+		{
+			DisposeAsync().AsTask().Wait();
+		}
+
+		/// <summary>
+		/// Flush any pending messages to database
+		/// Exposed as internal for use in tests
+		/// </summary>
+		internal async Task<int> FlushMessagesInternalAsync()
+		{
+			List<AuditLogMessage> newMessages = new ();
+			while (_messageChannel.Reader.TryRead(out AuditLogMessage? newMessage))
+			{
+				newMessages.Add(newMessage);
+			}
+			if (newMessages.Count > 0)
+			{
+				await _messages.InsertManyAsync(newMessages);
+			}
+			
+			return newMessages.Count;
+		}
+		
+		async Task WriteMessagesAsync()
+		{
+			while (await _messageChannel.Reader.WaitToReadAsync())
+			{
+				await FlushMessagesInternalAsync();
+			}
+		}
+
+		async IAsyncEnumerable<IAuditLogMessage<TSubject>> FindAsync(TSubject subject, DateTime? minTime = null, DateTime? maxTime = null, int? index = null, int? count = null)
+		{
+			FilterDefinition<AuditLogMessage> filter = Builders<AuditLogMessage>.Filter.Eq(x => x.Subject, subject);
+			if (minTime != null)
+			{
+				filter &= Builders<AuditLogMessage>.Filter.Gte(x => x.TimeUtc, minTime.Value);
+			}
+			if (maxTime != null)
+			{
+				filter &= Builders<AuditLogMessage>.Filter.Lte(x => x.TimeUtc, maxTime.Value);
+			}
+
+			using (IAsyncCursor<AuditLogMessage> cursor = await _messages.Find(filter).SortByDescending(x => x.TimeUtc).Range(index, count).ToCursorAsync())
+			{
+				while (await cursor.MoveNextAsync())
+				{
+					foreach (AuditLogMessage message in cursor.Current)
+					{
+						yield return message;
+					}
+				}
+			}
+		}
+
+		async Task<long> DeleteAsync(TSubject subject, DateTime? minTime = null, DateTime? maxTime = null)
+		{
+			FilterDefinition<AuditLogMessage> filter = Builders<AuditLogMessage>.Filter.Eq(x => x.Subject, subject);
+			if (minTime != null)
+			{
+				filter &= Builders<AuditLogMessage>.Filter.Gte(x => x.TimeUtc, minTime.Value);
+			}
+			if (maxTime != null)
+			{
+				filter &= Builders<AuditLogMessage>.Filter.Lte(x => x.TimeUtc, maxTime.Value);
+			}
+
+			DeleteResult result = await _messages.DeleteManyAsync(filter);
+			return result.DeletedCount;
+		}
+	}
+
+	class AuditLogFactory<TSubject> : IAuditLogFactory<TSubject>
+	{
+		readonly MongoService _mongoService;
+		readonly ILogger<AuditLog<TSubject>> _logger;
+
+		public AuditLogFactory(MongoService mongoService, ILogger<AuditLog<TSubject>> logger)
+		{
+			_mongoService = mongoService;
+			_logger = logger;
+		}
+
+		public IAuditLog<TSubject> Create(string collectionName, string subjectProperty)
+		{
+			return new AuditLog<TSubject>(_mongoService, collectionName, subjectProperty, _logger);
+		}
+	}
+}
diff --git a/Engine/Source/Programs/Horde/Horde.Server/Authentication/AnonymousHandler.cs b/Engine/Source/Programs/Horde/Horde.Server/Authentication/AnonymousHandler.cs
new file mode 100644
--- /dev/null
+++ b/Engine/Source/Programs/Horde/Horde.Server/Authentication/AnonymousHandler.cs
@@ -0,0 +1,61 @@
+// Copyright Epic Games, Inc. All Rights Reserved.
+
+using System;
+using System.Collections.Generic;
+using System.Security.Claims;
+using System.Text.Encodings.Web;
+using System.Threading.Tasks;
+using Horde.Server.Users;
+using Horde.Server.Utilities;
+using Microsoft.AspNetCore.Authentication;
+using Microsoft.Extensions.Logging;
+using Microsoft.Extensions.Options;
+
+namespace Horde.Server.Authentication
+{
+	class AnonymousAuthenticationOptions : AuthenticationSchemeOptions
+	{
+		public string? AdminClaimType { get; set; }
+		public string? AdminClaimValue { get; set; }
+	}
+
+	class AnonymousAuthenticationHandler : AuthenticationHandler<AnonymousAuthenticationOptions>
+	{
+		public const string AuthenticationScheme = "Anonymous";
+
+		readonly IOptionsMonitor<ServerSettings> _settings;
+
+		public AnonymousAuthenticationHandler(IOptionsMonitor<AnonymousAuthenticationOptions> options, IOptionsMonitor<ServerSettings> settings, ILoggerFactory logger, UrlEncoder encoder, ISystemClock clock)
+			: base(options, logger, encoder, clock)
+		{
+			_settings = settings;
+		}
+
+		protected override Task<AuthenticateResult> HandleAuthenticateAsync()
+		{
+			List<Claim> claims = new List<Claim>();
+			claims.Add(new Claim(ClaimTypes.Name, AuthenticationScheme));
+			claims.Add(new Claim(HordeClaimTypes.UserId, UserId.Anonymous.ToString()));
+
+			if (Options.AdminClaimType != null && Options.AdminClaimValue != null)
+			{
+				claims.Add(new Claim(Options.AdminClaimType, Options.AdminClaimValue));
+			}
+
+			ClaimsIdentity identity = new ClaimsIdentity(claims, Scheme.Name);
+
+			ClaimsPrincipal principal = new ClaimsPrincipal(identity);
+			AuthenticationTicket ticket = new AuthenticationTicket(principal, Scheme.Name);
+
+			return Task.FromResult(AuthenticateResult.Success(ticket));
+		}
+	}
+
+	static class AnonymousExtensions
+	{
+		public static AuthenticationBuilder AddAnonymous(this AuthenticationBuilder builder, Action<AnonymousAuthenticationOptions> configure)
+		{
+			return builder.AddScheme<AnonymousAuthenticationOptions, AnonymousAuthenticationHandler>(AnonymousAuthenticationHandler.AuthenticationScheme, configure);
+		}
+	}
+}
diff --git a/Engine/Source/Programs/Horde/Horde.Server/Authentication/HordeOpenIdConnectHandler.cs b/Engine/Source/Programs/Horde/Horde.Server/Authentication/HordeOpenIdConnectHandler.cs
new file mode 100644
--- /dev/null
+++ b/Engine/Source/Programs/Horde/Horde.Server/Authentication/HordeOpenIdConnectHandler.cs
@@ -0,0 +1,133 @@
+// Copyright Epic Games, Inc. All Rights Reserved.
+
+using System;
+using System.Linq;
+using System.Security.Claims;
+using System.Text.Encodings.Web;
+using System.Text.Json;
+using System.Threading.Tasks;
+using Horde.Server.Users;
+using Horde.Server.Utilities;
+using Microsoft.AspNetCore.Authentication;
+using Microsoft.AspNetCore.Authentication.OAuth.Claims;
+using Microsoft.AspNetCore.Authentication.OpenIdConnect;
+using Microsoft.Extensions.DependencyInjection;
+using Microsoft.Extensions.DependencyInjection.Extensions;
+using Microsoft.Extensions.Logging;
+using Microsoft.Extensions.Options;
+using Microsoft.IdentityModel.Protocols.OpenIdConnect;
+
+namespace Horde.Server.Authentication
+{
+	class HordeOpenIdConnectHandler : OpenIdConnectHandler
+	{
+		readonly IUserCollection _userCollection;
+
+		public HordeOpenIdConnectHandler(IOptionsMonitor<OpenIdConnectOptions> options, ILoggerFactory logger, HtmlEncoder htmlEncoder, UrlEncoder encoder, ISystemClock clock, IUserCollection userCollection)
+			: base(options, logger, htmlEncoder, encoder, clock)
+		{
+			_userCollection = userCollection;
+		}
+
+		/// <summary>
+		/// Try to map a field found in user info as a claim
+		///
+		/// If the claim is already set, it is skipped.
+		/// </summary>
+		/// <param name="userInfo">User info from the OIDC /userinfo endpoint </param>
+		/// <param name="identity">Identity which claims to modify</param>
+		/// <param name="claimName">Name of claim</param>
+		/// <param name="userInfoFields">List of field names to try when mapping</param>
+		/// <exception cref="Exception">Thrown if no field name matches</exception>
+		private static void MapUserInfoFieldToClaim(JsonElement userInfo, ClaimsIdentity identity, string claimName, string[] userInfoFields)
+		{
+			if (identity.FindFirst(claimName) != null)
+			{
+				return;
+			}
+
+			foreach (string userDataField in userInfoFields)
+			{
+				if (userInfo.TryGetProperty(userDataField, out JsonElement fieldElement))
+				{
+					identity.AddClaim(new Claim(claimName, fieldElement.ToString()));
+					return;
+				}
+			}
+
+			string message = $"Unable to map a field from user info to claim '{claimName}' using list [{String.Join(", ", userInfoFields)}].";
+			message += " UserInfo: " + userInfo;
+			throw new Exception(message);
+		}
+
+		public static void AddUserInfoClaims(ServerSettings settings, JsonElement userInfo, ClaimsIdentity identity)
+		{
+			MapUserInfoFieldToClaim(userInfo, identity, ClaimTypes.Name, settings.OidcClaimNameMapping);
+			MapUserInfoFieldToClaim(userInfo, identity, ClaimTypes.Email, settings.OidcClaimEmailMapping);
+			MapUserInfoFieldToClaim(userInfo, identity, HordeClaimTypes.User, settings.OidcClaimHordeUserMapping);
+			MapUserInfoFieldToClaim(userInfo, identity, HordeClaimTypes.PerforceUser, settings.OidcClaimHordePerforceUserMapping);
+
+			if (userInfo.TryGetProperty("groups", out JsonElement groupsElement) && groupsElement.ValueKind == JsonValueKind.Array)
+			{
+				for (int idx = 0; idx < groupsElement.GetArrayLength(); idx++)
+				{
+					identity.AddClaim(new Claim(ClaimTypes.Role, groupsElement[idx].ToString()!));
+				}
+			}
+		}
+
+		protected override async Task<HandleRequestResult> HandleRemoteAuthenticateAsync()
+		{
+			// Authenticate with the OIDC provider
+			HandleRequestResult result = await base.HandleRemoteAuthenticateAsync();
+			if (!result.Succeeded)
+			{
+				return result;
+			}
+
+			ClaimsIdentity? identity = (ClaimsIdentity?)result.Principal?.Identity;
+			if (identity == null)
+			{
+				return HandleRequestResult.Fail("No identity specified");
+			}
+
+			string login = identity.FindFirst(ClaimTypes.Name)!.Value;
+			string? name = identity.FindFirst("name")?.Value;
+			string? email = identity.FindFirst(ClaimTypes.Email)?.Value;
+
+			IUser user = await _userCollection.FindOrAddUserByLoginAsync(login, name, email);
+			identity.AddClaim(new Claim(HordeClaimTypes.UserId, user.Id.ToString()));
+
+			await _userCollection.UpdateClaimsAsync(user.Id, identity.Claims.Select(x => new UserClaim(x.Type, x.Value)));
+
+			return result;
+		}
+	}
+
+	static class OpenIdConnectHandlerExtensions
+	{
+		class MapRolesClaimAction : ClaimAction
+		{
+			private readonly ServerSettings _settings;
+			public MapRolesClaimAction(ServerSettings settings) : base(ClaimTypes.Role, ClaimTypes.Role) { _settings = settings; }
+			public override void Run(JsonElement userData, ClaimsIdentity identity, string issuer) { HordeOpenIdConnectHandler.AddUserInfoClaims(_settings, userData, identity); }
+		}
+
+		static void ApplyDefaultOptions(ServerSettings settings, OpenIdConnectOptions options, Action<OpenIdConnectOptions> handler)
+		{
+			options.ResponseType = OpenIdConnectResponseType.Code;
+			options.GetClaimsFromUserInfoEndpoint = true;
+			options.SaveTokens = true;
+			options.TokenValidationParameters.NameClaimType = "name";
+			options.ClaimActions.Add(new MapRolesClaimAction(settings));
+
+			handler(options);
+		}
+
+		public static void AddHordeOpenId(this AuthenticationBuilder builder, ServerSettings settings, string authenticationScheme, string displayName, Action<OpenIdConnectOptions> handler)
+		{
+			builder.Services.TryAddEnumerable(ServiceDescriptor.Singleton<IPostConfigureOptions<OpenIdConnectOptions>, OpenIdConnectPostConfigureOptions>());
+			builder.AddRemoteScheme<OpenIdConnectOptions, HordeOpenIdConnectHandler>(authenticationScheme, displayName, options => ApplyDefaultOptions(settings, options, handler));
+		}
+	}
+}
diff --git a/Engine/Source/Programs/Horde/Horde.Server/Authentication/OktaHandler.cs b/Engine/Source/Programs/Horde/Horde.Server/Authentication/OktaHandler.cs
new file mode 100644
--- /dev/null
+++ b/Engine/Source/Programs/Horde/Horde.Server/Authentication/OktaHandler.cs
@@ -0,0 +1,135 @@
+// Copyright Epic Games, Inc. All Rights Reserved.
+
+using System;
+using System.Linq;
+using System.Security.Claims;
+using System.Text.Encodings.Web;
+using System.Text.Json;
+using System.Threading.Tasks;
+using Horde.Server.Users;
+using Horde.Server.Utilities;
+using Microsoft.AspNetCore.Authentication;
+using Microsoft.AspNetCore.Authentication.OAuth.Claims;
+using Microsoft.AspNetCore.Authentication.OpenIdConnect;
+using Microsoft.Extensions.DependencyInjection;
+using Microsoft.Extensions.DependencyInjection.Extensions;
+using Microsoft.Extensions.Logging;
+using Microsoft.Extensions.Options;
+using Microsoft.IdentityModel.Protocols.OpenIdConnect;
+
+namespace Horde.Server.Authentication
+{
+	class OktaDefaults
+	{
+		public const string AuthenticationScheme = "Okta" + OpenIdConnectDefaults.AuthenticationScheme;
+	}
+
+	class OktaHandler : OpenIdConnectHandler
+	{
+		readonly IUserCollection _userCollection;
+
+		public OktaHandler(IOptionsMonitor<OpenIdConnectOptions> options, ILoggerFactory logger, HtmlEncoder htmlEncoder, UrlEncoder encoder, ISystemClock clock, IUserCollection userCollection)
+			: base(options, logger, htmlEncoder, encoder, clock)
+		{
+			_userCollection = userCollection;
+		}
+
+		public static void AddUserInfoClaims(JsonElement userInfo, ClaimsIdentity identity)
+		{
+			JsonElement userElement;
+			if (userInfo.TryGetProperty("preferred_username", out userElement))
+			{
+				if (identity.FindFirst(ClaimTypes.Name) == null)
+				{
+					identity.AddClaim(new Claim(ClaimTypes.Name, userElement.ToString()!));
+				}
+				if (identity.FindFirst(HordeClaimTypes.User) == null)
+				{
+					identity.AddClaim(new Claim(HordeClaimTypes.User, userElement.ToString()!));
+				}
+				if (identity.FindFirst(HordeClaimTypes.PerforceUser) == null)
+				{
+					identity.AddClaim(new Claim(HordeClaimTypes.PerforceUser, userElement.ToString()!));
+				}
+			}
+
+			JsonElement groupsElement;
+			if (userInfo.TryGetProperty("groups", out groupsElement) && groupsElement.ValueKind == JsonValueKind.Array)
+			{
+				for (int idx = 0; idx < groupsElement.GetArrayLength(); idx++)
+				{
+					identity.AddClaim(new Claim(ClaimTypes.Role, groupsElement[idx].ToString()!));
+				}
+			}
+
+			JsonElement emailElement;
+			if (userInfo.TryGetProperty("email", out emailElement) && identity.FindFirst(ClaimTypes.Email) == null)
+			{
+				identity.AddClaim(new Claim(ClaimTypes.Email, emailElement.ToString()!));
+			}
+		}
+
+		protected override async Task<HandleRequestResult> HandleRemoteAuthenticateAsync()
+		{
+			// Authenticate with the OIDC provider
+			HandleRequestResult result = await base.HandleRemoteAuthenticateAsync();
+			if (!result.Succeeded)
+			{
+				return result;
+			}
+
+			ClaimsIdentity? identity = (ClaimsIdentity?)result.Principal?.Identity;
+			if (identity == null)
+			{
+				return HandleRequestResult.Fail("No identity specified");
+			}
+
+			string login = identity.FindFirst(HordeClaimTypes.User)!.Value;
+			string? name = identity.FindFirst(ClaimTypes.Name)?.Value;
+			string? email = identity.FindFirst(ClaimTypes.Email)?.Value;
+
+			IUser user = await _userCollection.FindOrAddUserByLoginAsync(login, name, email);
+			identity.AddClaim(new Claim(HordeClaimTypes.UserId, user.Id.ToString()));
+
+			await _userCollection.UpdateClaimsAsync(user.Id, identity.Claims.Select(x => new UserClaim(x.Type, x.Value)));
+
+			return result;
+		}
+	}
+
+	static class OktaExtensions
+	{
+		class MapRolesClaimAction : ClaimAction
+		{
+			public MapRolesClaimAction()
+				: base(ClaimTypes.Role, ClaimTypes.Role)
+			{
+			}
+
+			public override void Run(JsonElement userData, ClaimsIdentity identity, string issuer)
+			{
+				OktaHandler.AddUserInfoClaims(userData, identity);
+			}
+		}
+
+		static void ApplyDefaultOktaOptions(OpenIdConnectOptions options, Action<OpenIdConnectOptions> handler)
+		{
+			options.Scope.Add("profile");
+			options.Scope.Add("groups");
+			options.Scope.Add("email");
+			options.ResponseType = OpenIdConnectResponseType.Code;
+			options.GetClaimsFromUserInfoEndpoint = true;
+			options.SaveTokens = true;
+			options.TokenValidationParameters.NameClaimType = "name";
+			options.ClaimActions.Add(new MapRolesClaimAction());
+
+			handler(options);
+		}
+
+		public static void AddOkta(this AuthenticationBuilder builder, string authenticationScheme, string displayName, Action<OpenIdConnectOptions> handler)
+		{
+			builder.Services.TryAddEnumerable(ServiceDescriptor.Singleton<IPostConfigureOptions<OpenIdConnectOptions>, OpenIdConnectPostConfigureOptions>());
+			builder.AddRemoteScheme<OpenIdConnectOptions, OktaHandler>(authenticationScheme, displayName, options => ApplyDefaultOktaOptions(options, handler));
+		}
+	}
+}
diff --git a/Engine/Source/Programs/Horde/Horde.Server/Authentication/ServiceAccountAuthHandler.cs b/Engine/Source/Programs/Horde/Horde.Server/Authentication/ServiceAccountAuthHandler.cs
new file mode 100644
--- /dev/null
+++ b/Engine/Source/Programs/Horde/Horde.Server/Authentication/ServiceAccountAuthHandler.cs
@@ -0,0 +1,79 @@
+﻿// Copyright Epic Games, Inc. All Rights Reserved.
+
+using System;
+using System.Collections.Generic;
+using System.Linq;
+using System.Security.Claims;
+using System.Text.Encodings.Web;
+using System.Threading.Tasks;
+using Horde.Server.Users;
+using Microsoft.AspNetCore.Authentication;
+using Microsoft.Extensions.Logging;
+using Microsoft.Extensions.Options;
+using Microsoft.Net.Http.Headers;
+
+namespace Horde.Server.Authentication
+{
+	class ServiceAccountAuthOptions : AuthenticationSchemeOptions
+	{
+	}
+
+	class ServiceAccountAuthHandler : AuthenticationHandler<ServiceAccountAuthOptions>
+	{
+		public const string AuthenticationScheme = "ServiceAccount";
+		public const string Prefix = "ServiceAccount";
+
+		private readonly IServiceAccountCollection _serviceAccounts;
+
+		public ServiceAccountAuthHandler(IOptionsMonitor<ServiceAccountAuthOptions> options,
+			ILoggerFactory logger, UrlEncoder encoder, ISystemClock clock, IServiceAccountCollection serviceAccounts)
+			: base(options, logger, encoder, clock)
+		{
+			_serviceAccounts = serviceAccounts;
+		}
+		
+		protected override async Task<AuthenticateResult> HandleAuthenticateAsync()
+		{
+			if (!Context.Request.Headers.TryGetValue(HeaderNames.Authorization, out Microsoft.Extensions.Primitives.StringValues headerValue))
+			{
+				return AuthenticateResult.NoResult();
+			}
+
+			if (headerValue.Count < 1)
+			{
+				return AuthenticateResult.NoResult();
+			}
+
+			if (!headerValue[0].StartsWith(Prefix, StringComparison.Ordinal))
+			{
+				return AuthenticateResult.NoResult();
+			}
+			
+			string token = headerValue[0].Replace(Prefix, "", StringComparison.Ordinal).Trim();
+			IServiceAccount? serviceAccount = await _serviceAccounts.GetBySecretTokenAsync(token);
+
+			if (serviceAccount == null)
+			{
+				return AuthenticateResult.Fail($"Service account for token {token} not found");
+			}
+
+			List<Claim> claims = new List<Claim>(10);
+			claims.Add(new Claim(ClaimTypes.Name, AuthenticationScheme));
+			claims.AddRange(serviceAccount.GetClaims().Select(claimPair => new Claim(claimPair.Type, claimPair.Value)));
+
+			ClaimsIdentity identity = new ClaimsIdentity(claims, Scheme.Name);
+			ClaimsPrincipal principal = new ClaimsPrincipal(identity);
+			AuthenticationTicket ticket = new AuthenticationTicket(principal, Scheme.Name);
+
+			return AuthenticateResult.Success(ticket);
+		}
+	}
+
+	static class ServiceAccountExtensions
+	{
+		public static AuthenticationBuilder AddServiceAccount(this AuthenticationBuilder builder, Action<ServiceAccountAuthOptions> config)
+		{
+			return builder.AddScheme<ServiceAccountAuthOptions, ServiceAccountAuthHandler>(ServiceAccountAuthHandler.AuthenticationScheme, config);
+		}
+	}
+}
diff --git a/Engine/Source/Programs/Horde/Horde.Server/Commands/Bundles/PerforceCommand.cs b/Engine/Source/Programs/Horde/Horde.Server/Commands/Bundles/PerforceCommand.cs
new file mode 100644
--- /dev/null
+++ b/Engine/Source/Programs/Horde/Horde.Server/Commands/Bundles/PerforceCommand.cs
@@ -0,0 +1,58 @@
+// Copyright Epic Games, Inc. All Rights Reserved.
+
+using System;
+using System.Threading;
+using System.Threading.Tasks;
+using EpicGames.Core;
+using Horde.Server.Configuration;
+using Horde.Server.Perforce;
+using Horde.Server.Server;
+using Horde.Server.Streams;
+using Horde.Server.Utilities;
+using Microsoft.Extensions.Configuration;
+using Microsoft.Extensions.DependencyInjection;
+using Microsoft.Extensions.Logging;
+using Microsoft.Extensions.Options;
+
+namespace Horde.Server.Commands.Bundles
+{
+	[Command("bundle", "perforce", "Replicates commits for a particular change of changes from Perforce")]
+	class PerforceCommand : Command
+	{
+		[CommandLine("-Stream=", Required = true)]
+		public string StreamId { get; set; } = String.Empty;
+
+		[CommandLine(Required = true)]
+		public int Change { get; set; }
+
+		readonly IConfiguration _configuration;
+		readonly ILoggerProvider _loggerProvider;
+		readonly IOptionsSnapshot<GlobalConfig> _globalConfig;
+
+		public PerforceCommand(IConfiguration configuration, ILoggerProvider loggerProvider, IOptionsSnapshot<GlobalConfig> globalConfig)
+		{
+			_configuration = configuration;
+			_loggerProvider = loggerProvider;
+			_globalConfig = globalConfig;
+		}
+
+		public override async Task<int> ExecuteAsync(ILogger logger)
+		{
+			using ServiceProvider serviceProvider = Startup.CreateServiceProvider(_configuration, _loggerProvider);
+
+			PerforceReplicator replicator = serviceProvider.GetRequiredService<PerforceReplicator>();
+			IStreamCollection streamCollection = serviceProvider.GetRequiredService<IStreamCollection>();
+
+			StreamConfig? streamConfig;
+			if (!_globalConfig.Value.TryGetStream(new StreamId(StreamId), out streamConfig))
+			{
+				throw new FatalErrorException($"Stream '{StreamId}' not found");
+			}
+
+			PerforceReplicationOptions options = new PerforceReplicationOptions();
+			await replicator.WriteAsync(streamConfig, Change, options, default);
+
+			return 0;
+		}
+	}
+}
diff --git a/Engine/Source/Programs/Horde/Horde.Server/Commands/Config/SchemasCommand.cs b/Engine/Source/Programs/Horde/Horde.Server/Commands/Config/SchemasCommand.cs
new file mode 100644
--- /dev/null
+++ b/Engine/Source/Programs/Horde/Horde.Server/Commands/Config/SchemasCommand.cs
@@ -0,0 +1,32 @@
+// Copyright Epic Games, Inc. All Rights Reserved.
+
+using System;
+using System.Threading.Tasks;
+using EpicGames.Core;
+using Horde.Server.Server;
+using Horde.Server.Utilities;
+using Microsoft.Extensions.Logging;
+
+namespace Horde.Server.Commands.Config
+{
+	[Command("config", "schemas", "Writes JSON schemas for server settings")]
+	class SchemasCommand : Command
+	{
+		[CommandLine]
+		DirectoryReference? _outputDir = null!;
+
+		public override Task<int> ExecuteAsync(ILogger logger)
+		{
+			_outputDir ??= DirectoryReference.Combine(Program.AppDir, "Schemas");
+
+			DirectoryReference.CreateDirectory(_outputDir);
+			foreach (Type schemaType in SchemaController.ConfigSchemas)
+			{
+				FileReference outputFile = FileReference.Combine(_outputDir, $"{schemaType.Name}.json");
+				Schemas.CreateSchema(schemaType).Write(outputFile);
+			}
+
+			return Task.FromResult(0);
+		}
+	}
+}
diff --git a/Engine/Source/Programs/Horde/Horde.Server/Commands/Migrate/MigrateLogsCommand.cs b/Engine/Source/Programs/Horde/Horde.Server/Commands/Migrate/MigrateLogsCommand.cs
new file mode 100644
--- /dev/null
+++ b/Engine/Source/Programs/Horde/Horde.Server/Commands/Migrate/MigrateLogsCommand.cs
@@ -0,0 +1,112 @@
+// Copyright Epic Games, Inc. All Rights Reserved.
+
+using System;
+using System.Collections.Generic;
+using System.IO;
+using System.Threading;
+using System.Threading.Tasks;
+using EpicGames.Core;
+using EpicGames.Horde.Logs;
+using EpicGames.Horde.Storage;
+using EpicGames.Horde.Storage.Nodes;
+using Horde.Server.Logs;
+using Horde.Server.Perforce;
+using Horde.Server.Server;
+using Horde.Server.Storage;
+using Horde.Server.Utilities;
+using Microsoft.Extensions.Caching.Memory;
+using Microsoft.Extensions.Configuration;
+using Microsoft.Extensions.DependencyInjection;
+using Microsoft.Extensions.Logging;
+using MongoDB.Bson;
+using MongoDB.Bson.Serialization.Attributes;
+using MongoDB.Driver;
+
+namespace Horde.Server.Commands.Migrate
+{
+	[Command("migrate", "logs", "Extracts data from a bundle to the local hard drive")]
+	internal class MigrateLogsCommand : Command
+	{
+		class LogFile
+		{
+			public ObjectId Id { get; set; }
+			public RefName RefName { get; set; }
+		}
+
+		readonly IConfiguration _configuration;
+		readonly ILoggerProvider _loggerProvider;
+
+		public MigrateLogsCommand(IConfiguration configuration, ILoggerProvider loggerProvider)
+		{
+			_configuration = configuration;
+			_loggerProvider = loggerProvider;
+		}
+
+		public override async Task<int> ExecuteAsync(ILogger logger)
+		{
+			using ServiceProvider serviceProvider = Startup.CreateServiceProvider(_configuration, _loggerProvider);
+
+			MongoService mongoService = serviceProvider.GetRequiredService<MongoService>();
+			ILogFileCollection logCollection = serviceProvider.GetRequiredService<ILogFileCollection>();
+			ILogFileService logService = serviceProvider.GetRequiredService<ILogFileService>();
+
+			StorageService storageService = serviceProvider.GetRequiredService<StorageService>();
+			IStorageClient storageClient = await storageService.GetClientAsync(Namespace.Logs, CancellationToken.None);
+
+			byte[] buffer = new byte[1 * 1024];
+
+			ILogFile logFile = (await logCollection.GetLogFileAsync(LogId.Parse("6374195e5f2709066e7257a6"), CancellationToken.None))!;
+//			await foreach(ILogFile logFile in logCollection.GetLogFilesAsync())
+			{
+				if (logFile.RefName == RefName.Empty)
+				{
+//					TreeWriter writer = new TreeWriter(storageClient);
+					int bufferLength = 0;
+					int builderLength = 0;
+
+					LogBuilder builder = new LogBuilder(LogFormat.Json, logger);
+
+					using (Stream stream = await logService.OpenRawStreamAsync(logFile, CancellationToken.None))
+					{
+						for (; ; )
+						{
+							int readLength = await stream.ReadAsync(buffer.AsMemory(bufferLength));
+							if (readLength == 0)
+							{
+								break;
+							}
+
+							bufferLength += readLength;
+
+							int startIdx = 0;
+							for (int idx = 0; idx < bufferLength; idx++)
+							{
+								if (buffer[idx] == (byte)'\n')
+								{
+									ReadOnlyMemory<byte> line = buffer.AsMemory(startIdx, idx + 1 - startIdx);
+
+									builder.WriteData(line);
+									builderLength += line.Length;
+
+									if (builderLength > 1024 * 1024)
+									{
+//										await builder.FlushAsync(writer, false, CancellationToken.None);
+										builderLength = 0;
+									}
+
+									startIdx = idx + 1;
+								}
+							}
+
+							bufferLength -= startIdx;
+							buffer.AsSpan(startIdx, bufferLength).CopyTo(buffer.AsSpan());
+						}
+					}
+//					await builder.FlushAsync(writer, true, CancellationToken.None);
+				}
+			}
+
+			return 0;
+		}
+	}
+}
diff --git a/Engine/Source/Programs/Horde/Horde.Server/DashboardApp/index.html b/Engine/Source/Programs/Horde/Horde.Server/DashboardApp/index.html
new file mode 100644
--- /dev/null
+++ b/Engine/Source/Programs/Horde/Horde.Server/DashboardApp/index.html
@@ -0,0 +1 @@
+<!doctype html><html lang="en"><head><meta charset="utf-8"/><link rel="icon" href="/favicon.svg"/><meta name="viewport" content="width=device-width,initial-scale=1"/><meta name="theme-color" content="#000000"/><meta name="description" content="Web site created using create-react-app"/><link rel="apple-touch-icon" href="/logo192.png"/><link rel="manifest" href="/manifest.json"/><title>Horde</title><script defer="defer" src="/static/js/main.92c6e7ff.js"></script><link href="/static/css/main.7a06f76f.css" rel="stylesheet"></head><body><noscript>You need to enable JavaScript to run this app.</noscript><div id="root"></div></body></html>
\ No newline at end of file
