Horde: Rename the Horde.Build project to Horde.Server.

diff --git a/Engine/Source/Programs/Horde/Horde.Server/Devices/DeviceService.cs b/Engine/Source/Programs/Horde/Horde.Server/Devices/DeviceService.cs
new file mode 100644
--- /dev/null
+++ b/Engine/Source/Programs/Horde/Horde.Server/Devices/DeviceService.cs
@@ -0,0 +1,692 @@
+// Copyright Epic Games, Inc. All Rights Reserved.
+
+using System;
+using System.Collections.Generic;
+using System.Linq;
+using System.Security.Claims;
+using System.Threading;
+using System.Threading.Tasks;
+using Horde.Server.Acls;
+using Horde.Server.Jobs;
+using Horde.Server.Jobs.Graphs;
+using Horde.Server.Notifications;
+using Horde.Server.Projects;
+using Horde.Server.Server;
+using Horde.Server.Streams;
+using Horde.Server.Users;
+using Horde.Server.Utilities;
+using HordeCommon;
+using Microsoft.Extensions.Hosting;
+using Microsoft.Extensions.Logging;
+using Microsoft.Extensions.Options;
+using MongoDB.Bson;
+using OpenTracing;
+using OpenTracing.Util;
+
+namespace Horde.Server.Devices
+{
+	/// <summary>
+	///  Device Pool Authorization (convenience class for pool ACL's)
+	/// </summary>
+	public class DevicePoolAuthorization
+	{
+		/// <summary>
+		/// The device pool
+		/// </summary>
+		public IDevicePool Pool { get; private set; }
+
+		/// <summary>
+		///  Read access to pool
+		/// </summary>
+		public bool Read { get; private set; }
+
+		/// <summary>
+		///  Write access to pool
+		/// </summary>
+		public bool Write { get; private set; }
+
+		/// <summary>
+		/// Constructor
+		/// </summary>
+		/// <param name="pool"></param>
+		/// <param name="read"></param>
+		/// <param name="write"></param>
+		public DevicePoolAuthorization(IDevicePool pool, bool read, bool write)
+		{
+			Pool = pool;
+			Read = read;
+			Write = write;				
+		}
+	}
+
+	/// <summary>
+	/// Platform map required by V1 API
+	/// </summary>
+	[SingletonDocument("device-platform-map", "6165a2e26fd5f104e31e6862")]
+	public class DevicePlatformMapV1 : SingletonBase
+	{
+		/// <summary>
+		/// Platform V1 => Platform Id
+		/// </summary>
+		public Dictionary<string, DevicePlatformId> PlatformMap { get; set; } = new Dictionary<string, DevicePlatformId>();
+
+		/// <summary>
+		/// Perfspec V1 => Model
+		/// </summary>
+		public Dictionary<DevicePlatformId, string> PerfSpecHighMap { get; set; } = new Dictionary<DevicePlatformId, string>();
+		
+	}
+
+	/// <summary>
+	/// Device management service
+	/// </summary>
+	public sealed class DeviceService : IHostedService, IDisposable
+	{
+		readonly INotificationService _notificationService;
+		readonly JobService _jobService;
+		readonly IStreamCollection _streamCollection;
+		readonly IUserCollection _userCollection;
+		readonly ILogger<DeviceService> _logger;
+		readonly IDeviceCollection _devices;
+		readonly ITicker _ticker;
+		readonly ITicker _telemetryTicker;
+		readonly IOptionsMonitor<ServerSettings> _settings;
+		readonly IOptionsMonitor<GlobalConfig> _globalConfig;
+
+		/// <summary>
+		/// Platform map V1 singleton
+		/// </summary>
+		readonly ISingletonDocument<DevicePlatformMapV1> _platformMapSingleton;
+
+		bool runUpgrade = true;
+
+		/// <summary>
+		/// The number of days shared device sheckouts are held
+		/// </summary>
+		public int sharedDeviceCheckoutDays => _settings.CurrentValue.SharedDeviceCheckoutDays;
+
+		/// <summary>
+		/// Device service constructor
+		/// </summary>
+		public DeviceService(IDeviceCollection devices, ISingletonDocument<DevicePlatformMapV1> platformMapSingleton, IUserCollection userCollection, JobService jobService, IStreamCollection streamCollection, IOptionsMonitor<ServerSettings> settings, IOptionsMonitor<GlobalConfig> globalConfig, INotificationService notificationService, IClock clock, ILogger<DeviceService> logger)
+		{
+			_userCollection = userCollection;
+			_devices = devices;
+			_jobService = jobService;
+			_streamCollection = streamCollection;
+			_notificationService = notificationService;
+			_ticker = clock.AddSharedTicker<DeviceService>(TimeSpan.FromMinutes(1.0), TickAsync, logger);
+			_telemetryTicker = clock.AddSharedTicker("DeviceService.Telemetry", TimeSpan.FromMinutes(10.0), TickTelemetryAsync, logger);
+			_logger = logger;
+			_settings = settings;
+			_globalConfig = globalConfig;
+			_platformMapSingleton = platformMapSingleton;
+		}
+
+		/// <inheritdoc/>
+		public async Task StartAsync(CancellationToken cancellationToken)
+		{
+			await _ticker.StartAsync();
+			await _telemetryTicker.StartAsync();
+		}
+
+		/// <inheritdoc/>
+		public async Task StopAsync(CancellationToken cancellationToken)
+		{
+			await _ticker.StopAsync();
+			await _telemetryTicker.StopAsync();
+		}
+
+		/// <inheritdoc/>
+		public void Dispose() 
+		{ 
+			_ticker.Dispose(); 
+			_telemetryTicker.Dispose(); 
+		}
+
+		/// <summary>
+		/// Ticks service
+		/// </summary>
+		async ValueTask TickTelemetryAsync(CancellationToken stoppingToken)
+		{
+			try
+			{
+				GlobalConfig globalConfig = _globalConfig.CurrentValue;
+				if (globalConfig.Devices != null)
+				{
+					await _platformMapSingleton.UpdateAsync(platformMap => {
+
+						platformMap.PlatformMap.Clear();
+						platformMap.PerfSpecHighMap.Clear();
+
+						foreach (DevicePlatformConfig platform in globalConfig.Devices.Platforms)
+						{
+							DevicePlatformId id = new DevicePlatformId(platform.Id);
+							foreach (string name in platform.Names)
+							{
+								platformMap.PlatformMap[name] = id;
+							}
+
+							if (platform.LegacyPerfSpecHighModel != null)
+							{
+								platformMap.PerfSpecHighMap[id] = platform.LegacyPerfSpecHighModel;
+							}
+						}
+					} );
+				}				
+			}
+			catch (Exception ex)
+			{
+				_logger.LogError(ex, "Exception while updating platform map: {Message}", ex.Message);
+			}
+
+			if (!stoppingToken.IsCancellationRequested)
+			{
+				using IScope scope = GlobalTracer.Instance.BuildSpan("DeviceService.TickTelemetryAsync").StartActive();
+
+				_logger.LogInformation("Updating pool telemetry");
+				await _devices.CreatePoolTelemetrySnapshot();
+			}
+		}
+
+		/// <summary>
+		/// Ticks service
+		/// </summary>
+		async ValueTask TickAsync(CancellationToken stoppingToken)
+		{
+			if (!stoppingToken.IsCancellationRequested)
+			{
+				using IScope scope = GlobalTracer.Instance.BuildSpan("DeviceService.TickAsync").StartActive();
+
+				GlobalConfig globalConfig = _globalConfig.CurrentValue;
+
+				if (runUpgrade)
+				{
+					try
+					{
+						await _devices.UpgradeAsync();
+						runUpgrade = false;
+					}
+					catch (Exception ex)
+					{
+						_logger.LogError(ex, "Exception while upgrading device collection: {Message}", ex.Message);
+					}
+
+				}
+
+				try
+				{
+					_logger.LogDebug("Expiring reservations");
+					await _devices.ExpireReservationsAsync();
+				}
+				catch (Exception ex)
+				{
+					_logger.LogError(ex, "Exception while expiring reservations: {Message}", ex.Message);
+				}
+
+				// expire shared devices and notifications
+				try
+				{
+
+					_logger.LogDebug("Sending shared device notifications");
+					List<(UserId, IDevice)>? expireNotifications = await _devices.ExpireNotificatonsAsync(_settings.CurrentValue.SharedDeviceCheckoutDays);
+					if (expireNotifications != null && expireNotifications.Count > 0)
+					{
+						foreach ((UserId, IDevice) expiredDevice in expireNotifications)
+						{
+							await NotifyDeviceServiceAsync(globalConfig, $"Device {expiredDevice.Item2.PlatformId.ToString().ToUpperInvariant()} / {expiredDevice.Item2.Name} checkout will expire in 24 hours.  Please visit https://horde.devtools.epicgames.com/devices to renew the checkout if needed.", null, null, null, expiredDevice.Item1);
+						}
+					}
+
+					_logger.LogDebug("Expiring shared device checkouts");
+					List<(UserId, IDevice)>? expireCheckouts = await _devices.ExpireCheckedOutAsync(_settings.CurrentValue.SharedDeviceCheckoutDays);
+					if (expireCheckouts != null && expireCheckouts.Count > 0)
+					{
+						foreach ((UserId, IDevice) expiredDevice in expireCheckouts)
+						{
+							await NotifyDeviceServiceAsync(globalConfig, $"Device {expiredDevice.Item2.PlatformId.ToString().ToUpperInvariant()} / {expiredDevice.Item2.Name} checkout has expired.  The device has been returned to the shared pool and should no longer be accessed.  Please visit https://horde.devtools.epicgames.com/devices to checkout devices as needed.", null, null, null, expiredDevice.Item1);
+						}
+					}
+				}
+				catch (Exception ex)
+				{
+					_logger.LogError(ex, "Exception while expiring device checkouts: {Message}", ex.Message);
+				}
+			}
+		}
+
+		internal async Task TickForTestingAsync()
+		{
+			await TickAsync(CancellationToken.None);
+			await TickTelemetryAsync(CancellationToken.None);
+		}
+
+		/// <summary>
+		/// Create a new device platform
+		/// </summary>
+		public Task<IDevicePlatform?> TryCreatePlatformAsync(DevicePlatformId id, string name)
+		{
+			return _devices.TryAddPlatformAsync(id, name);
+		}
+
+		/// <summary>
+		/// Get a list of existing device platforms
+		/// </summary>
+		public Task<List<IDevicePlatform>> GetPlatformsAsync()
+		{
+			return _devices.FindAllPlatformsAsync();
+		}
+
+		/// <summary>
+		/// Update an existing platform
+		/// </summary>
+		public Task<bool> UpdatePlatformAsync(DevicePlatformId platformId, string[]? modelIds)
+		{
+			return _devices.UpdatePlatformAsync(platformId, modelIds);
+		}
+
+		/// <summary>
+		/// Get a specific device platform
+		/// </summary>
+		public Task<IDevicePlatform?> GetPlatformAsync(DevicePlatformId id)
+		{
+			return _devices.GetPlatformAsync(id);
+		}
+
+		/// <summary>
+		/// Get a device pool by id
+		/// </summary>
+		public Task<IDevicePool?> GetPoolAsync(DevicePoolId id)
+		{
+			return _devices.GetPoolAsync(id);
+		}
+
+		/// <summary>
+		/// Create a new device pool
+		/// </summary>
+		public Task<IDevicePool?> TryCreatePoolAsync(DevicePoolId id, string name, DevicePoolType poolType, List<ProjectId>? projectIds)
+		{
+			return _devices.TryAddPoolAsync(id, name, poolType, projectIds);
+		}
+
+		/// <summary>
+		/// Update a device pool
+		/// </summary>
+		public Task UpdatePoolAsync(DevicePoolId id, List<ProjectId>? projectIds)
+		{
+			return _devices.UpdatePoolAsync(id, projectIds);
+		}
+
+		/// <summary>
+		/// Get a list of existing device pools
+		/// </summary>
+		public Task<List<IDevicePool>> GetPoolsAsync()
+		{
+			return _devices.FindAllPoolsAsync();
+		}
+
+		/// <summary>
+		/// Get a list of devices, optionally filtered to provided ids
+		/// </summary>
+		public Task<List<IDevice>> GetDevicesAsync(List<DeviceId>? deviceIds = null, DevicePoolId? poolId = null, DevicePlatformId? platformId = null)
+		{
+			return _devices.FindAllDevicesAsync(deviceIds, poolId, platformId);
+		}
+
+		/// <summary>
+		/// Get device telemetry
+		/// </summary>
+		public Task<List<IDeviceTelemetry>> GetDeviceTelemetryAsync(DeviceId[]? deviceIds = null, DateTimeOffset? minCreateTime=null, DateTimeOffset? maxCreateTime=null, int? index = null, int? count = null)
+		{
+			return _devices.FindDeviceTelemetryAsync(deviceIds, minCreateTime, maxCreateTime, index, count);
+		}
+
+		/// <summary>
+		/// Get device pool telemetry
+		/// </summary>
+		public Task<List<IDevicePoolTelemetry>> GetDevicePoolTelemetryAsync(DateTimeOffset? minCreateTime = null, DateTimeOffset? maxCreateTime = null, int? index = null, int? count = null)
+		{
+			return _devices.FindPoolTelemetryAsync(minCreateTime, maxCreateTime, index, count);
+		}
+
+		/// <summary>
+		/// Get a specific device
+		/// </summary>
+		public Task<IDevice?> GetDeviceAsync(DeviceId id)
+		{
+			return _devices.GetDeviceAsync(id);
+		}
+
+		/// <summary>
+		/// Get a device by name
+		/// </summary>
+		public Task<IDevice?> GetDeviceByNameAsync(string deviceName)
+		{
+			return _devices.GetDeviceByNameAsync(deviceName);
+		}
+
+		/// <summary>
+		/// Delete a device
+		/// </summary>
+		public async Task<bool> DeleteDeviceAsync(DeviceId id)
+		{
+			return await _devices.DeleteDeviceAsync(id);
+		}
+
+		/// <summary>
+		/// Create a new device
+		/// </summary>
+		/// <param name="id">Unique id of the device</param>
+		/// <param name="name">Friendly name of the device</param>
+		/// <param name="platformId">The device platform</param>
+		/// <param name="poolId">Which pool to add the device</param>
+		/// <param name="enabled">Whether the device is enabled</param>
+		/// <param name="address">Address or hostname of device</param>
+		/// <param name="modelId">Vendor model id</param>
+        /// <param name="userId">User adding the device</param>
+		/// <returns></returns>
+		public Task<IDevice?> TryCreateDeviceAsync(DeviceId id, string name, DevicePlatformId platformId, DevicePoolId poolId, bool? enabled, string? address, string? modelId, UserId? userId = null)
+		{
+			return _devices.TryAddDeviceAsync(id, name, platformId, poolId, enabled, address, modelId, userId);
+		}
+
+		/// <summary>
+		/// Update a device
+		/// </summary>
+		public Task UpdateDeviceAsync(DeviceId deviceId, DevicePoolId? newPoolId = null, string? newName = null, string? newAddress = null, string? newModelId = null, string? newNotes = null, bool? newEnabled = null, bool? newProblem = null, bool? newMaintenance = null, UserId? modifiedByUserId = null)
+		{
+			return _devices.UpdateDeviceAsync(deviceId, newPoolId, newName, newAddress, newModelId, newNotes, newEnabled, newProblem, newMaintenance, modifiedByUserId);
+		}
+
+		/// <summary>
+		/// Checkout a device
+		/// </summary>
+		public Task CheckoutDeviceAsync(DeviceId deviceId, UserId? userId)
+		{
+            return _devices.CheckoutDeviceAsync(deviceId, userId);
+        }
+
+		/// <summary>
+		/// Try to create a reservation satisfying the specified device platforms and models
+		/// </summary>
+		public async Task<IDeviceReservation?> TryCreateReservationAsync(DevicePoolId pool, List<DeviceRequestData> request, string? hostname = null, string? reservationDetails = null, string? jobId = null, string? stepId = null)
+		{
+			string? streamId = null;
+			string? jobName = null;
+			string? stepName = null;
+
+			if (jobId != null)
+			{
+				IJob? job = await _jobService.GetJobAsync(JobId.Parse(jobId));
+				if (job != null)
+				{					
+					streamId = job.StreamId.ToString();
+					jobName = job.Name;
+
+					if (stepId != null)
+					{
+						SubResourceId id = SubResourceId.Parse(stepId);
+						IGraph graph = await _jobService.GetGraphAsync(job);
+						foreach (IJobStepBatch batch in job.Batches)
+						{
+							IJobStep? step;							
+							if (batch.TryGetStep(id, out step))
+							{
+								stepName = graph.Groups[batch.GroupIdx].Nodes[step.NodeIdx].Name;
+								break;
+							}
+						}
+					}
+				}
+			}
+
+			return await _devices.TryAddReservationAsync(pool, request, hostname, reservationDetails, streamId, jobId, stepId, jobName, stepName);
+		}
+
+		/// <summary>
+		/// Update/renew an existing reservation
+		/// </summary>
+		public Task<bool> TryUpdateReservationAsync(ObjectId id)
+		{
+			return _devices.TryUpdateReservationAsync(id);
+		}
+
+		/// <summary>
+		///  Delete an existing reservation
+		/// </summary>
+		public Task<bool> DeleteReservationAsync(ObjectId id)
+		{
+			return _devices.DeleteReservationAsync(id);
+		}
+
+		/// <summary>
+		/// Get a reservation from a legacy guid
+		/// </summary>
+		public Task<IDeviceReservation?> TryGetReservationFromLegacyGuidAsync(string legacyGuid)
+		{
+			return _devices.TryGetReservationFromLegacyGuidAsync(legacyGuid);
+		}
+
+		/// <summary>
+		/// Get a reservation from a device id
+		/// </summary>
+		public Task<IDeviceReservation?> TryGetDeviceReservation(DeviceId deviceId)
+		{
+			return _devices.TryGetDeviceReservationAsync(deviceId);
+		}
+
+		/// <summary>
+		/// Get a list of existing device reservations
+		/// </summary>
+		public Task<List<IDeviceReservation>> GetReservationsAsync()
+		{
+			return _devices.FindAllReservationsAsync();
+		}
+
+		/// <summary>
+		/// 
+		/// </summary>
+		public async Task NotifyDeviceServiceAsync(GlobalConfig globalConfig, string message, DeviceId? deviceId = null, string? jobId = null, string? stepId = null, UserId? userId = null)
+		{
+			try 
+			{
+				IDevice? device = null;
+				IDevicePool? pool = null;
+				IJob? job = null;
+				IJobStep? step = null;
+				INode? node = null;
+				StreamConfig? streamConfig = null;
+				IUser? user = null;
+
+				if (userId.HasValue)
+				{
+					user = await _userCollection.GetUserAsync(userId.Value);
+					if (user == null)
+					{
+						_logger.LogError("Unable to send device notification, can't find User {UserId}", userId.Value);
+						return;
+					}
+				}
+
+				if (deviceId.HasValue)
+				{
+					device = await GetDeviceAsync(deviceId.Value);
+					pool = await GetPoolAsync(device!.PoolId);
+				}
+
+				if (jobId != null)
+				{
+					job = await _jobService.GetJobAsync(JobId.Parse(jobId));
+
+					if (job != null)
+					{
+						globalConfig.TryGetStream(job.StreamId, out streamConfig);
+
+						if (stepId != null)
+						{
+							IGraph graph = await _jobService.GetGraphAsync(job)!;
+
+							SubResourceId stepIdValue = SubResourceId.Parse(stepId);
+							IJobStepBatch? batch = job.Batches.FirstOrDefault(b => b.Steps.FirstOrDefault(s => s.Id == stepIdValue) != null);
+							if (batch != null)
+							{
+								step = batch.Steps.FirstOrDefault(s => s.Id == stepIdValue)!;
+								INodeGroup group = graph.Groups[batch.GroupIdx];
+								node = group.Nodes[step.NodeIdx];
+							}
+						}
+					}
+				}
+
+				_notificationService.NotifyDeviceService(message, device, pool, streamConfig, job, step, node, user);
+
+			}
+			catch (Exception ex)
+			{
+                _logger.LogError(ex, "Error on device notification {Message}", ex.Message);
+            }
+        }
+
+		/// <summary>
+		/// Authorize device action
+		/// </summary>
+		/// <param name="action"></param>
+		/// <param name="user"></param>
+		/// <param name="globalConfig"></param>
+		/// <returns></returns>
+		public static bool Authorize(AclAction action, ClaimsPrincipal user, GlobalConfig globalConfig)
+		{
+			// This is deprecated and device auth should be going through GetUserPoolAuthorizationsAsync
+
+			if (action == AclAction.DeviceRead)
+			{
+				return true;
+			}
+
+			if (user.IsInRole("Internal-Employees"))
+			{
+				return true;
+			}
+
+			if (globalConfig.Authorize(AclAction.AdminWrite, user))
+			{
+				return true;
+			}
+
+			return false;
+		}
+
+		/// <summary>
+		/// Get list of pool authorizations for a user
+		/// </summary>
+		/// <param name="user"></param>
+		/// <param name="globalConfig"></param>
+		/// <returns></returns>
+		public async Task<List<DevicePoolAuthorization>> GetUserPoolAuthorizationsAsync(ClaimsPrincipal user, GlobalConfig globalConfig)
+		{
+			List<DevicePoolAuthorization> authPools = new List<DevicePoolAuthorization>();
+
+			List<IDevicePool> allPools = await GetPoolsAsync();
+			IReadOnlyList<ProjectConfig> projects = globalConfig.Projects;
+
+			// Set of projects associated with device pools
+			HashSet<ProjectId> projectIds = new HashSet<ProjectId>(allPools.Where(x => x.ProjectIds != null).SelectMany(x => x.ProjectIds!));
+
+			Dictionary<ProjectId, bool> deviceRead = new Dictionary<ProjectId, bool>();
+			Dictionary<ProjectId, bool> deviceWrite = new Dictionary<ProjectId, bool>();
+
+			foreach (ProjectId projectId in projectIds)
+			{
+				if (projects.Where(x => x.Id == projectId).FirstOrDefault() == null)
+				{
+					_logger.LogWarning("Device pool authorization references missing project id {ProjectId}", projectId);
+					continue;
+				}
+
+				if (globalConfig.TryGetProject(projectId, out ProjectConfig? projectConfig))
+				{
+					deviceRead.Add(projectId, projectConfig.Authorize(AclAction.DeviceRead, user));
+					deviceWrite.Add(projectId, projectConfig.Authorize(AclAction.DeviceWrite, user));
+				}
+			}
+
+			// for global pools which aren't associated with a project
+			bool globalPoolAccess = user.IsInRole("Internal-Employees");
+
+			if (!globalPoolAccess)
+			{
+				if (globalConfig.Authorize(AclAction.AdminWrite, user))
+				{
+					globalPoolAccess = true;
+				}
+			}
+
+			foreach (IDevicePool pool in allPools)
+			{				
+				if (pool.ProjectIds == null || pool.ProjectIds.Count == 0)
+				{
+					if (pool.PoolType == DevicePoolType.Shared && !globalPoolAccess)
+					{
+						authPools.Add(new DevicePoolAuthorization(pool, false, false));
+						continue;
+					}
+
+					authPools.Add(new DevicePoolAuthorization(pool, true, true));
+					continue;
+				}
+
+				bool read = false;
+				bool write = false;
+
+				foreach (ProjectId projectId in pool.ProjectIds)
+				{
+					bool value;
+
+					if (!read && deviceRead.TryGetValue(projectId, out value))
+					{
+						read = value;
+					}
+
+					if (!write && deviceWrite.TryGetValue(projectId, out value))
+					{
+						write = value;
+					}
+				}
+
+				authPools.Add(new DevicePoolAuthorization(pool, read, write));
+
+			}
+
+			return authPools;
+		}
+
+		/// <summary>
+		/// Get list of pool authorizations for a user
+		/// </summary>
+		/// <param name="id"></param>
+		/// <param name="user"></param>
+		/// <param name="globalConfig"></param>
+		/// <returns></returns>
+		public async Task<DevicePoolAuthorization?> GetUserPoolAuthorizationAsync(DevicePoolId id, ClaimsPrincipal user, GlobalConfig globalConfig)
+		{
+			List<DevicePoolAuthorization> auth = await GetUserPoolAuthorizationsAsync(user, globalConfig);
+			return auth.Where(x => x.Pool.Id == id).FirstOrDefault();
+		}
+
+		/// <summary>
+		/// Update a device
+		/// </summary>
+		public async Task UpdateDevicePoolAsync(DevicePoolId poolId, List<ProjectId>? projectIds)
+		{
+			await _devices.UpdatePoolAsync(poolId, projectIds);
+		}
+
+		/// <summary>
+		/// Get Platform mappings for V1 API
+		/// </summary>		
+		public async Task<DevicePlatformMapV1> GetPlatformMapV1()
+		{
+			return await _platformMapSingleton.GetAsync();
+		}
+	}
+}
diff --git a/Engine/Source/Programs/Horde/Horde.Server/Devices/DevicesController.cs b/Engine/Source/Programs/Horde/Horde.Server/Devices/DevicesController.cs
new file mode 100644
--- /dev/null
+++ b/Engine/Source/Programs/Horde/Horde.Server/Devices/DevicesController.cs
@@ -0,0 +1,1125 @@
+// Copyright Epic Games, Inc. All Rights Reserved.
+
+using System;
+using System.Collections.Generic;
+using System.Linq;
+using System.Text.Json;
+using System.Threading.Tasks;
+using Horde.Server.Acls;
+using Horde.Server.Projects;
+using Horde.Server.Server;
+using Horde.Server.Users;
+using Horde.Server.Utilities;
+using Microsoft.AspNetCore.Authorization;
+using Microsoft.AspNetCore.Mvc;
+using Microsoft.Extensions.Logging;
+using Microsoft.Extensions.Options;
+using MongoDB.Bson;
+using MongoDB.Driver;
+
+namespace Horde.Server.Devices
+{
+	/// <summary>
+	/// Controller for device service
+	/// </summary>
+	public class DevicesController : ControllerBase
+	{
+		readonly IUserCollection _userCollection;
+		readonly DeviceService _deviceService;
+		readonly IOptionsSnapshot<GlobalConfig> _globalConfig;
+
+		/// <summary>
+		///  Logger for controller
+		/// </summary>
+		private readonly ILogger<DevicesController> _logger;
+
+		/// <summary>
+		/// Constructor
+		/// </summary>
+		public DevicesController(IUserCollection userCollection, DeviceService deviceService, IOptionsSnapshot<GlobalConfig> globalConfig, ILogger<DevicesController> logger)
+		{
+			_userCollection = userCollection;
+			_deviceService = deviceService;
+			_globalConfig = globalConfig;
+			_logger = logger;
+		}
+
+		// DEVICES
+
+		/// <summary>
+		/// Create a new device
+		/// </summary>
+		[HttpPost]
+		[Authorize]
+		[Route("/api/v2/devices")]
+		public async Task<ActionResult<CreateDeviceResponse>> CreateDeviceAsync([FromBody] CreateDeviceRequest deviceRequest)
+		{
+			DevicePoolAuthorization? poolAuth = await _deviceService.GetUserPoolAuthorizationAsync(new DevicePoolId(deviceRequest.PoolId!), User, _globalConfig.Value);
+
+			if (poolAuth == null || !poolAuth.Write)
+			{
+				return Forbid();
+			}
+
+			IUser? internalUser = await _userCollection.GetUserAsync(User);
+			if (internalUser == null)
+			{
+				return NotFound();
+			}
+
+			IDevicePlatform? platform = await _deviceService.GetPlatformAsync(new DevicePlatformId(deviceRequest.PlatformId!));
+
+			if (platform == null)
+			{
+				return BadRequest($"Bad platform id {deviceRequest.PlatformId} on request");
+			}
+
+			IDevicePool? pool = await _deviceService.GetPoolAsync(new DevicePoolId(deviceRequest.PoolId!));
+
+			if (pool == null)
+			{
+				return BadRequest($"Bad pool id {deviceRequest.PoolId} on request");
+			}
+
+			string? modelId = null;
+
+			if (!String.IsNullOrEmpty(deviceRequest.ModelId))
+			{
+				modelId = new string(deviceRequest.ModelId);
+
+				if (platform.Models?.FirstOrDefault(x => x == modelId) == null)
+				{
+					return BadRequest($"Bad model id {modelId} for platform {platform.Id} on request");
+				}
+			}
+
+			string name = deviceRequest.Name.Trim();
+			string? address = deviceRequest.Address?.Trim();
+
+			IDevice? device = await _deviceService.TryCreateDeviceAsync(DeviceId.Sanitize(name), name, platform.Id, pool.Id, deviceRequest.Enabled, address, modelId, internalUser.Id);
+
+			if (device == null)
+			{
+				return BadRequest($"Unable to create device");
+			}
+
+			return new CreateDeviceResponse(device.Id.ToString());
+		}
+
+		/// <summary>
+		/// Get list of devices
+		/// </summary>
+		[HttpGet]
+		[Authorize]
+		[Route("/api/v2/devices")]
+		[ProducesResponseType(typeof(List<GetDeviceResponse>), 200)]
+		public async Task<ActionResult<List<object>>> GetDevicesAsync()
+		{
+			List<DevicePoolAuthorization> poolAuth = await _deviceService.GetUserPoolAuthorizationsAsync(User, _globalConfig.Value);
+
+			List<IDevice> devices = await _deviceService.GetDevicesAsync();
+
+			List<object> responses = new List<object>();
+
+			foreach (IDevice device in devices)
+			{
+				DevicePoolAuthorization? auth = poolAuth.Where(x => x.Pool.Id == device.PoolId).FirstOrDefault();
+
+				if (auth == null || !auth.Read)
+				{
+					continue;
+				}
+
+				DateTime? checkoutExpiration = null;
+				if (device.CheckOutTime != null)
+				{
+					checkoutExpiration = device.CheckOutTime.Value.AddDays(_deviceService.sharedDeviceCheckoutDays);
+				}
+
+				responses.Add(new GetDeviceResponse(device.Id.ToString(), device.PlatformId.ToString(), device.PoolId.ToString(), device.Name, device.Enabled, device.Address, device.ModelId?.ToString(), device.ModifiedByUser, device.Notes, device.ProblemTimeUtc, device.MaintenanceTimeUtc, device.Utilization, device.CheckedOutByUser, device.CheckOutTime, checkoutExpiration));
+			}
+
+			return responses;
+		}
+
+		/// <summary>
+		/// Get a specific device
+		/// </summary>
+		[HttpGet]
+		[Authorize]
+		[Route("/api/v2/devices/{deviceId}")]
+		[ProducesResponseType(typeof(GetDeviceResponse), 200)]
+		public async Task<ActionResult<object>> GetDeviceAsync(string deviceId)
+		{
+
+			DeviceId deviceIdValue = new DeviceId(deviceId);
+
+			IDevice? device = await _deviceService.GetDeviceAsync(deviceIdValue);
+
+			if (device == null)
+			{
+				return BadRequest($"Unable to find device with id {deviceId}");
+			}
+
+			DevicePoolAuthorization? poolAuth = await _deviceService.GetUserPoolAuthorizationAsync(device.PoolId, User, _globalConfig.Value);
+
+			if (poolAuth == null || !poolAuth.Write)
+			{
+				return Forbid();
+			}
+
+			DateTime? checkoutExpiration = null;
+			if (device.CheckOutTime != null)
+			{
+				checkoutExpiration = device.CheckOutTime.Value.AddDays(_deviceService.sharedDeviceCheckoutDays);
+			}
+
+			return new GetDeviceResponse(device.Id.ToString(), device.PlatformId.ToString(), device.PoolId.ToString(), device.Name, device.Enabled, device.Address, device.ModelId?.ToString(), device.ModifiedByUser?.ToString(), device.Notes, device.ProblemTimeUtc, device.MaintenanceTimeUtc, device.Utilization, device.CheckedOutByUser, device.CheckOutTime, checkoutExpiration);
+		}
+
+		/// <summary>
+		/// Update a specific device
+		/// </summary>
+		[HttpPut]
+		[Authorize]
+		[Route("/api/v2/devices/{deviceId}")]
+		[ProducesResponseType(typeof(List<GetDeviceResponse>), 200)]
+		public async Task<ActionResult> UpdateDeviceAsync(string deviceId, [FromBody] UpdateDeviceRequest update)
+		{
+			IUser? internalUser = await _userCollection.GetUserAsync(User);
+			if (internalUser == null)
+			{
+				return NotFound();
+			}
+
+			DeviceId deviceIdValue = new DeviceId(deviceId);
+
+			IDevice? device = await _deviceService.GetDeviceAsync(deviceIdValue);
+
+			if (device == null)
+			{
+				return BadRequest($"Device with id ${deviceId} does not exist");
+			}
+
+			DevicePoolAuthorization? poolAuth = await _deviceService.GetUserPoolAuthorizationAsync(device.PoolId, User, _globalConfig.Value);
+
+			if (poolAuth == null || !poolAuth.Write)
+			{
+				return Forbid();
+			}
+
+			IDevicePlatform? platform = await _deviceService.GetPlatformAsync(device.PlatformId);
+
+			if (platform == null)
+			{
+				return BadRequest($"Platform id ${device.PlatformId} does not exist");
+			}
+
+			DevicePoolId? poolIdValue = null;
+
+			if (update.PoolId != null)
+			{
+				poolIdValue = new DevicePoolId(update.PoolId);
+			}
+
+			string? modelIdValue = null;
+
+			if (update.ModelId != null)
+			{
+				modelIdValue = new string(update.ModelId);
+
+				if (platform.Models?.FirstOrDefault(x => x == modelIdValue) == null)
+				{
+					return BadRequest($"Bad model id {update.ModelId} for platform {platform.Id} on request");
+				}
+			}
+
+			string? name = update.Name?.Trim();
+			string? address = update.Address?.Trim();
+
+			await _deviceService.UpdateDeviceAsync(deviceIdValue, poolIdValue, name, address, modelIdValue, update.Notes, update.Enabled, update.Problem, update.Maintenance, internalUser.Id);
+
+			return Ok();
+		}
+
+		/// <summary>
+		/// Checkout a device
+		/// </summary>
+		[HttpPut]
+		[Authorize]
+		[Route("/api/v2/devices/{deviceId}/checkout")]
+		[ProducesResponseType(typeof(List<GetDeviceResponse>), 200)]
+		public async Task<ActionResult> CheckoutDeviceAsync(string deviceId, [FromBody] CheckoutDeviceRequest request)
+		{
+
+			IUser? internalUser = await _userCollection.GetUserAsync(User);
+			if (internalUser == null)
+			{
+				return NotFound();
+			}
+
+			DeviceId deviceIdValue = new DeviceId(deviceId);
+
+			IDevice? device = await _deviceService.GetDeviceAsync(deviceIdValue);
+
+			if (device == null)
+			{
+				return BadRequest($"Device with id ${deviceId} does not exist");
+			}
+
+			DevicePoolAuthorization? poolAuth = await _deviceService.GetUserPoolAuthorizationAsync(device.PoolId, User, _globalConfig.Value);
+
+			if (poolAuth == null || !poolAuth.Write)
+			{
+				return Forbid();
+			}
+
+			if (request.Checkout)
+			{
+				if (!String.IsNullOrEmpty(device.CheckedOutByUser))
+				{
+					return BadRequest($"Already checked out by user {device.CheckedOutByUser}");
+				}
+
+				await _deviceService.CheckoutDeviceAsync(deviceIdValue, internalUser.Id);
+
+			}
+			else
+			{
+				await _deviceService.CheckoutDeviceAsync(deviceIdValue, null);
+			}
+
+			return Ok();
+		}
+
+		/// <summary>
+		/// Delete a specific device
+		/// </summary>
+		[HttpDelete]
+		[Authorize]
+		[Route("/api/v2/devices/{deviceId}")]
+		public async Task<ActionResult> DeleteDeviceAsync(string deviceId)
+		{
+			if (!DeviceService.Authorize(AclAction.DeviceWrite, User, _globalConfig.Value))
+			{
+				return Forbid();
+			}
+
+			DeviceId deviceIdValue = new DeviceId(deviceId);
+
+			IDevice? device = await _deviceService.GetDeviceAsync(deviceIdValue);
+			if (device == null)
+			{
+				return NotFound();
+			}
+
+			DevicePoolAuthorization? poolAuth = await _deviceService.GetUserPoolAuthorizationAsync(device.PoolId, User, _globalConfig.Value);
+
+			if (poolAuth == null || !poolAuth.Write)
+			{
+				return Forbid();
+			}
+
+			await _deviceService.DeleteDeviceAsync(deviceIdValue);
+			return Ok();
+		}
+
+		// PLATFORMS
+
+		/// <summary>
+		/// Create a new device platform
+		/// </summary>
+		[HttpPost]
+		[Authorize]
+		[Route("/api/v2/devices/platforms")]
+		public async Task<ActionResult<CreateDevicePlatformResponse>> CreatePlatformAsync([FromBody] CreateDevicePlatformRequest request)
+		{
+			if (!DeviceService.Authorize(AclAction.DeviceWrite, User, _globalConfig.Value))
+			{
+				return Forbid();
+			}
+
+			string name = request.Name.Trim();
+
+			IDevicePlatform? platform = await _deviceService.TryCreatePlatformAsync(DevicePlatformId.Sanitize(name), name);
+
+			if (platform == null)
+			{
+				return BadRequest($"Unable to create platform for {name}");
+			}
+
+			return new CreateDevicePlatformResponse(platform.Id.ToString());
+		}
+
+		/// <summary>
+		/// Update a device platform
+		/// </summary>
+		[HttpPut]
+		[Authorize]
+		[Route("/api/v2/devices/platforms/{platformId}")]
+		public async Task<ActionResult<CreateDevicePlatformResponse>> UpdatePlatformAsync(string platformId, [FromBody] UpdateDevicePlatformRequest request)
+		{
+			if (!DeviceService.Authorize(AclAction.DeviceWrite, User, _globalConfig.Value))
+			{
+				return Forbid();
+			}
+
+			await _deviceService.UpdatePlatformAsync(new DevicePlatformId(platformId), request.ModelIds);
+
+			return Ok();
+		}
+
+		/// <summary>
+		/// Get a list of supported device platforms
+		/// </summary>
+		[HttpGet]
+		[Authorize]
+		[Route("/api/v2/devices/platforms")]
+		[ProducesResponseType(typeof(List<GetDevicePlatformResponse>), 200)]
+		public async Task<ActionResult<List<object>>> GetDevicePlatformsAsync()
+		{
+			if (!DeviceService.Authorize(AclAction.DeviceRead, User, _globalConfig.Value))
+			{
+				return Forbid();
+			}
+
+			List<IDevicePlatform> platforms = await _deviceService.GetPlatformsAsync();
+
+			List<object> responses = new List<object>();
+
+			foreach (IDevicePlatform platform in platforms)
+			{
+				// @todo: ACL per platform
+				responses.Add(new GetDevicePlatformResponse(platform.Id.ToString(), platform.Name, platform.Models?.ToArray() ?? Array.Empty<string>()));
+			}
+
+			return responses;
+		}
+
+		// POOLS
+
+		/// <summary>
+		/// Create a new device pool
+		/// </summary>
+		[HttpPost]
+		[Authorize]
+		[Route("/api/v2/devices/pools")]
+		public async Task<ActionResult<CreateDevicePoolResponse>> CreatePoolAsync([FromBody] CreateDevicePoolRequest request)
+		{
+			if (!DeviceService.Authorize(AclAction.DeviceWrite, User, _globalConfig.Value))
+			{
+				return Forbid();
+			}
+
+			string name = request.Name.Trim();
+
+			IDevicePool? pool = await _deviceService.TryCreatePoolAsync(DevicePoolId.Sanitize(name), name, request.PoolType, request.ProjectIds?.Select(x => new ProjectId(x)).ToList());
+
+			if (pool == null)
+			{
+				return BadRequest($"Unable to create pool for {request.Name}");
+			}
+
+			return new CreateDevicePoolResponse(pool.Id.ToString());
+
+		}
+
+		/// <summary>
+		/// Update a device pool
+		/// </summary>
+		[HttpPut]
+		[Authorize]
+		[Route("/api/v2/devices/pools")]
+		public async Task<ActionResult> UpdatePoolAsync([FromBody] UpdateDevicePoolRequest request)
+		{
+			DevicePoolAuthorization? poolAuth = await _deviceService.GetUserPoolAuthorizationAsync(new DevicePoolId(request.Id), User, _globalConfig.Value);
+
+			if (poolAuth == null || !poolAuth.Write)
+			{
+				return Forbid();
+			}
+
+			List<ProjectId>? projectIds = null;
+
+			if (request.ProjectIds != null)
+			{
+				projectIds = request.ProjectIds.Select(x => new ProjectId(x)).ToList();
+			}
+
+			await _deviceService.UpdatePoolAsync(new DevicePoolId(request.Id), projectIds);
+
+			return Ok();
+		}
+
+		/// <summary>
+		/// Get a list of existing device pools
+		/// </summary>
+		[HttpGet]
+		[Authorize]
+		[Route("/api/v2/devices/pools")]
+		[ProducesResponseType(typeof(List<GetDevicePoolResponse>), 200)]
+		public async Task<ActionResult<List<object>>> GetDevicePoolsAsync()
+		{
+			List<DevicePoolAuthorization> poolAuth = await _deviceService.GetUserPoolAuthorizationsAsync(User, _globalConfig.Value);
+
+			List<IDevicePool> pools = await _deviceService.GetPoolsAsync();
+
+			List<object> responses = new List<object>();
+
+			foreach (IDevicePool pool in pools)
+			{
+				DevicePoolAuthorization? auth = poolAuth.Where(x => x.Pool.Id == pool.Id).FirstOrDefault();
+
+				if (auth == null || !auth.Read)
+				{
+					continue;
+				}
+
+				responses.Add(new GetDevicePoolResponse(pool.Id.ToString(), pool.Name, pool.PoolType, auth.Write));
+			}
+
+			return responses;
+		}
+
+		/// <summary>
+		/// Get device telemetry
+		/// </summary>		
+		[HttpGet]
+		[Authorize]
+		[Route("/api/v2/devices/pools/telemetry")]
+		[ProducesResponseType(typeof(List<GetDevicePoolTelemetryResponse>), 200)]
+		public async Task<ActionResult<List<GetDevicePoolTelemetryResponse>>> GetDevicePoolTelemetry(
+			[FromQuery] DateTimeOffset? minCreateTime = null,
+			[FromQuery] DateTimeOffset? maxCreateTime = null,
+			[FromQuery] int index = 0,
+			[FromQuery] int count = 1024)
+		{
+			if (!DeviceService.Authorize(AclAction.DeviceRead, User, _globalConfig.Value))
+			{
+				return Forbid();
+			}
+
+			if (minCreateTime == null)
+			{
+				minCreateTime = new DateTimeOffset(DateTime.UtcNow - TimeSpan.FromHours(24));
+			}
+
+			List<IDevicePoolTelemetry> telemetryData = await _deviceService.GetDevicePoolTelemetryAsync(minCreateTime, maxCreateTime, index, count);
+
+			List<GetDevicePoolTelemetryResponse> response = new List<GetDevicePoolTelemetryResponse>();
+						
+			foreach (IDevicePoolTelemetry t in telemetryData)
+			{
+				Dictionary<string, List<GetDevicePlatformTelemetryResponse>> poolData = new Dictionary<string, List<GetDevicePlatformTelemetryResponse>>();
+
+				foreach ( KeyValuePair<DevicePoolId, IReadOnlyList<IDevicePlatformTelemetry>> pool in t.Pools)
+				{
+					List<GetDevicePlatformTelemetryResponse> platformTelemetry = new List<GetDevicePlatformTelemetryResponse>();
+
+					foreach (IDevicePlatformTelemetry telemetry in pool.Value)
+					{
+						IReadOnlyDictionary<string, IReadOnlyList<IDevicePoolReservationTelemetry>>? reserved = telemetry.Reserved?.ToDictionary(kvp => kvp.Key.ToString(), kvp => kvp.Value);
+						platformTelemetry.Add(new GetDevicePlatformTelemetryResponse(telemetry.PlatformId.ToString(), telemetry.Available?.Select(d => d.ToString()).ToList(), telemetry.Maintenance?.Select(d => d.ToString()).ToList(), telemetry.Problem?.Select(d => d.ToString()).ToList(), telemetry.Disabled?.Select(d => d.ToString()).ToList(), reserved));
+					}
+
+					poolData[pool.Key.ToString()] = platformTelemetry;
+				}
+
+				response.Add(new GetDevicePoolTelemetryResponse(t.CreateTimeUtc, poolData));
+			}		
+
+			return response;
+		}
+
+		#region Reservations
+
+		/// <summary>
+		/// Create a new device reservation
+		/// </summary>
+		[HttpPost]
+		[Authorize]
+		[Route("/api/v2/devices/reservations")]
+		[ProducesResponseType(typeof(CreateDeviceReservationResponse), 200)]
+		public async Task<ActionResult<CreateDeviceReservationResponse>> CreateDeviceReservation([FromBody] CreateDeviceReservationRequest request)
+		{
+			if (!DeviceService.Authorize(AclAction.DeviceWrite, User, _globalConfig.Value))
+			{
+				return Forbid();
+			}
+
+			List<IDevicePool> pools = await _deviceService.GetPoolsAsync();
+			List<IDevicePlatform> platforms = await _deviceService.GetPlatformsAsync();
+
+			DevicePoolId poolIdValue = new DevicePoolId(request.PoolId);
+			IDevicePool? pool = pools.FirstOrDefault(x => x.Id == poolIdValue);
+			if (pool == null)
+			{
+				return BadRequest($"Unknown pool {request.PoolId} on device reservation request");
+			}
+
+			List<DeviceRequestData> requestedDevices = new List<DeviceRequestData>();
+
+			foreach (DeviceReservationRequest deviceRequest in request.Devices)
+			{
+				DevicePlatformId platformIdValue = new DevicePlatformId(deviceRequest.PlatformId);
+				IDevicePlatform? platform = platforms.FirstOrDefault(x => x.Id == platformIdValue);
+
+				if (platform == null)
+				{
+					return BadRequest($"Unknown platform {deviceRequest.PlatformId} on device reservation request");
+				}
+
+				if (deviceRequest.IncludeModels != null)
+				{
+					foreach (string model in deviceRequest.IncludeModels)
+					{
+						if (platform.Models?.FirstOrDefault(x => x == model) == null)
+						{
+							return BadRequest($"Unknown model {model} for platform {deviceRequest.PlatformId} on device reservation request");
+						}
+					}
+				}
+
+				if (deviceRequest.ExcludeModels != null)
+				{
+					foreach (string model in deviceRequest.ExcludeModels)
+					{
+						if (platform.Models?.FirstOrDefault(x => x == model) == null)
+						{
+							return BadRequest($"Unknown model {model} for platform {deviceRequest.PlatformId} on device reservation request");
+						}
+					}
+				}
+
+				requestedDevices.Add(new DeviceRequestData(platformIdValue, platformIdValue.ToString(), deviceRequest.IncludeModels, deviceRequest.ExcludeModels));
+			}
+
+			IDeviceReservation? reservation = await _deviceService.TryCreateReservationAsync(poolIdValue, requestedDevices);
+
+			if (reservation == null)
+			{
+				return Conflict("Unable to allocated devices for reservation");
+			}
+
+			List<IDevice> devices = await _deviceService.GetDevicesAsync(reservation.Devices);
+
+			CreateDeviceReservationResponse response = new CreateDeviceReservationResponse();
+
+			response.Id = reservation.Id.ToString();
+
+			foreach (IDevice device in devices)
+			{
+				response.Devices.Add(new GetDeviceResponse(device.Id.ToString(), device.PlatformId.ToString(), device.PoolId.ToString(), device.Name, device.Enabled, device.Address, device.ModelId?.ToString(), device.ModifiedByUser, device.Notes, device.ProblemTimeUtc, device.MaintenanceTimeUtc, device.Utilization));
+			}
+
+			return response;
+		}
+
+		/// <summary>
+		/// Get active device reservations
+		/// </summary>		
+		[HttpGet]
+		[Authorize]
+		[Route("/api/v2/devices/reservations")]
+		[ProducesResponseType(typeof(List<GetDeviceReservationResponse>), 200)]
+		public async Task<ActionResult<List<GetDeviceReservationResponse>>> GetDeviceReservations()
+		{
+			if (!DeviceService.Authorize(AclAction.DeviceRead, User, _globalConfig.Value))
+			{
+				return Forbid();
+			}
+
+			List<IDeviceReservation> reservations = await _deviceService.GetReservationsAsync();
+
+			List<GetDeviceReservationResponse> response = new List<GetDeviceReservationResponse>();
+
+			foreach (IDeviceReservation reservation in reservations)
+			{
+				GetDeviceReservationResponse reservationResponse = new GetDeviceReservationResponse()
+				{
+					Id = reservation.Id.ToString(),
+					PoolId = reservation.PoolId.ToString(),
+					Devices = reservation.Devices.Select(x => x.ToString()).ToList(),
+					JobId = reservation.JobId?.ToString(),
+					StepId = reservation.StepId?.ToString(),
+					JobName = reservation.JobName?.ToString(),
+					StepName = reservation.StepName?.ToString(),
+					UserId = reservation.UserId?.ToString(),
+					Hostname = reservation.Hostname,
+					ReservationDetails = reservation.ReservationDetails,
+					CreateTimeUtc = reservation.CreateTimeUtc,
+					LegacyGuid = reservation.LegacyGuid
+				};
+
+				response.Add(reservationResponse);
+
+			}
+
+			return response;
+		}
+
+		/// <summary>
+		/// Renew an existing reservation
+		/// </summary>
+		[HttpPut]
+		[Authorize]
+		[Route("/api/v2/devices/reservations/{reservationId}")]
+		public async Task<ActionResult> UpdateReservationAsync(string reservationId)
+		{
+			if (!DeviceService.Authorize(AclAction.DeviceWrite, User, _globalConfig.Value))
+			{
+				return Forbid();
+			}
+
+			ObjectId reservationIdValue = new ObjectId(reservationId);
+
+			bool updated = await _deviceService.TryUpdateReservationAsync(reservationIdValue);
+
+			if (!updated)
+			{
+				return BadRequest("Failed to update reservation");
+			}
+
+			return Ok();
+		}
+
+		/// <summary>
+		/// Delete a reservation
+		/// </summary>
+		[HttpDelete]
+		[Authorize]
+		[Route("/api/v2/devices/reservations/{reservationId}")]
+		public async Task<ActionResult> DeleteReservationAsync(string reservationId)
+		{
+			if (!DeviceService.Authorize(AclAction.DeviceWrite, User, _globalConfig.Value))
+			{
+				return Forbid();
+			}
+
+			ObjectId reservationIdValue = new ObjectId(reservationId);
+
+			bool deleted = await _deviceService.DeleteReservationAsync(reservationIdValue);
+
+			if (!deleted)
+			{
+				return BadRequest("Failed to delete reservation");
+			}
+
+			return Ok();
+		}
+
+		#endregion
+
+		/// <summary>
+		/// Get device telemetry
+		/// </summary>		
+		[HttpGet]
+		[Authorize]
+		[Route("/api/v2/devices/telemetry")]
+		[ProducesResponseType(typeof(List<GetDeviceTelemetryResponse>), 200)]
+		public async Task<ActionResult<List<GetDeviceTelemetryResponse>>> GetDeviceTelemetry(
+			[FromQuery(Name = "Id")] string[]? deviceIds = null,
+			[FromQuery] string? poolId = null,
+			[FromQuery] string? platformId = null,
+			[FromQuery] DateTimeOffset? minCreateTime = null,
+			[FromQuery] DateTimeOffset? maxCreateTime = null,
+			[FromQuery] int index = 0,
+			[FromQuery] int count = 1024)
+		{
+			if (!DeviceService.Authorize(AclAction.DeviceRead, User, _globalConfig.Value))
+			{
+				return Forbid();
+			}
+
+			if (minCreateTime == null)
+			{
+				minCreateTime = new DateTimeOffset(DateTime.UtcNow - TimeSpan.FromHours(24));
+			}
+
+			HashSet<DeviceId> deviceIdValues = new HashSet<DeviceId>();
+
+			if (deviceIds != null && deviceIds.Length > 0)
+			{
+				foreach(string deviceId in deviceIds)
+				{
+					deviceIdValues.Add(new DeviceId(deviceId));
+				}
+			}
+
+			DevicePoolId? poolIdValue = poolId == null ? null : new DevicePoolId(poolId);
+			DevicePlatformId? platformIdValue = platformId == null ? null : new DevicePlatformId(platformId);
+
+			List<IDevice> devices = await _deviceService.GetDevicesAsync(deviceIdValues.Count > 0 ? deviceIdValues.ToList() : null, poolIdValue, platformIdValue);
+
+			List<GetDeviceTelemetryResponse> response = new List<GetDeviceTelemetryResponse>();
+
+			if (devices.Count == 0)
+			{
+				return response;
+			}
+			
+			List<IDeviceTelemetry> telemetryData = await _deviceService.GetDeviceTelemetryAsync(devices.Select(x => x.Id).ToArray(), minCreateTime, maxCreateTime, index, count);
+
+			Dictionary<string, List<GetTelemetryInfoResponse>> results = new Dictionary<string, List<GetTelemetryInfoResponse>>();
+
+			foreach (IDeviceTelemetry t in telemetryData)
+			{
+				List<GetTelemetryInfoResponse>? info = null;
+				string deviceId = t.DeviceId.ToString();
+				if (!results.TryGetValue(deviceId, out info))
+				{
+					info = new List<GetTelemetryInfoResponse>();
+					results[deviceId] = info;
+				}
+
+				info.Add(new GetTelemetryInfoResponse(t));
+			}
+
+			foreach(KeyValuePair<string, List<GetTelemetryInfoResponse >> result in results)
+			{
+				GetDeviceTelemetryResponse deviceTelemetry = new GetDeviceTelemetryResponse(result.Key, result.Value);
+				response.Add(deviceTelemetry);
+			}
+
+			return response;
+		}
+
+		// LEGACY V1 API
+
+		enum LegacyPerfSpec
+		{
+			Unspecified,
+			Minimum,
+			Recommended,
+			High
+		};
+
+		/// <summary>
+		/// Create a device reservation
+		/// </summary>
+		[HttpPost]
+		[Route("/api/v1/reservations")]
+		[ProducesResponseType(typeof(GetLegacyReservationResponse), 200)]
+		public async Task<ActionResult<GetLegacyReservationResponse>> CreateDeviceReservationV1Async([FromBody] LegacyCreateReservationRequest request)
+		{
+
+			List<IDevicePool> pools = await _deviceService.GetPoolsAsync();
+			List<IDevicePlatform> platforms = await _deviceService.GetPlatformsAsync();
+
+			string? poolId = request.PoolId;
+
+			// @todo: Remove this once all streams are updated to provide jobid
+			string details = "";
+			if ((String.IsNullOrEmpty(request.JobId) || String.IsNullOrEmpty(request.StepId)))
+			{
+				if (!String.IsNullOrEmpty(request.ReservationDetails))
+				{
+					details = $" - {request.ReservationDetails}";
+				}
+				else
+				{
+					details = $" - Host {request.Hostname}";
+				}
+			}
+
+			if (String.IsNullOrEmpty(poolId))
+			{
+				// @todo: We default to ue4, though this should be an error, or a configuration setting
+				poolId = "ue4";
+				//return BadRequest(Message);
+			}
+
+			DevicePoolId poolIdValue = DevicePoolId.Sanitize(poolId);
+			IDevicePool? pool = pools.FirstOrDefault(x => x.Id == poolIdValue);
+			if (pool == null)
+			{
+				_logger.LogError("Unknown pool {PoolId} {Details}", poolId, details);
+				string message = $"Unknown pool {poolId} " + details;
+				return BadRequest(message);
+			}
+
+			List<DeviceRequestData> requestedDevices = new List<DeviceRequestData>();
+
+			HashSet<string> legacyPerfSpecs = new HashSet<string> { "Unspecified", "Minimum", "Recommended", "High" };
+
+			List<string> perfSpecs = new List<string>();
+
+			foreach (string deviceType in request.DeviceTypes)
+			{
+
+				string platformName = deviceType;
+				string? constraint = null;
+
+				if (deviceType.Contains(':', StringComparison.Ordinal))
+				{
+					string[] tokens = deviceType.Split(":");
+					platformName = tokens[0];
+					constraint = tokens[1];
+				}
+
+				DevicePlatformMapV1 mapV1 = await _deviceService.GetPlatformMapV1();
+
+				DevicePlatformId platformId = DevicePlatformId.Sanitize(platformName);
+
+				IDevicePlatform? platform = platforms.FirstOrDefault(x => x.Id == platformId);
+
+				if (platform == null)
+				{
+					if (mapV1.PlatformMap.TryGetValue(platformName, out platformId))
+					{
+						platform = platforms.FirstOrDefault(x => x.Id == platformId);
+					}
+				}
+
+				if (platform == null)
+				{
+					return BadRequest($"Unknown platform {platformId}" + details);
+				}
+
+				List<string> includeModels = new List<string>();
+				List<string> excludeModels = new List<string>();
+
+				if (constraint == null)
+				{
+					perfSpecs.Add("Unspecified");
+				}
+				else if (legacyPerfSpecs.Contains(constraint))
+				{
+					perfSpecs.Add(constraint);
+
+					if (constraint == "High")
+					{
+						string? model = null;
+						if (mapV1.PerfSpecHighMap.TryGetValue(platformId, out model))
+						{
+							includeModels.Add(model);
+						}
+					}
+					else if (constraint == "Minimum" || constraint == "Recommended")
+					{
+						string? model = null;
+						if (mapV1.PerfSpecHighMap.TryGetValue(platformId, out model))
+						{
+							excludeModels.Add(model);
+						}
+					}
+				}
+				else
+				{
+					string[] requestedModels = constraint.Split(';');
+
+					if (requestedModels.Length > 0)
+					{
+						List<string>? models = platform.Models?.Where(x => requestedModels.Contains(x, StringComparer.OrdinalIgnoreCase)).ToList();
+
+						if (models == null || models.Count == 0)
+						{
+							return NotFound($"Invalid model constraint for platform {platform.Id}: {constraint}");
+						}
+
+						includeModels = models;
+					}
+				}
+
+				requestedDevices.Add(new DeviceRequestData(platformId, platformName, includeModels, excludeModels));
+			}
+
+			IDeviceReservation? reservation = await _deviceService.TryCreateReservationAsync(poolIdValue, requestedDevices, request.Hostname, request.ReservationDetails, request.JobId, request.StepId);
+
+			if (reservation == null)
+			{
+				return Conflict("Unable to allocated devices for reservation");
+			}
+
+			List<IDevice> devices = await _deviceService.GetDevicesAsync(reservation.Devices);
+
+			GetLegacyReservationResponse response = new GetLegacyReservationResponse();
+
+			response.Guid = reservation.LegacyGuid;
+			response.DeviceNames = devices.Select(x => x.Name).ToArray();
+			response.DevicePerfSpecs = perfSpecs.ToArray();
+			response.DeviceModels = devices.Select(x => x.ModelId ?? "Base").ToArray();
+			response.HostName = request.Hostname;
+			response.StartDateTime = reservation.CreateTimeUtc.ToString("O", System.Globalization.CultureInfo.InvariantCulture);
+			response.Duration = $"{request.Duration}";
+			response.JobId = reservation.JobId;
+			response.StepId = reservation.StepId;
+			response.JobName = reservation.JobName;
+			response.StepName = reservation.StepName;
+
+			return new JsonResult(response, new JsonSerializerOptions() { PropertyNamingPolicy = null });
+
+		}
+
+		/// <summary>
+		/// Renew a reservation
+		/// </summary>
+		[HttpPut]
+		[Route("/api/v1/reservations/{reservationGuid}")]
+		[ProducesResponseType(typeof(GetLegacyReservationResponse), 200)]
+		public async Task<ActionResult<GetLegacyReservationResponse>> UpdateReservationV1Async(string reservationGuid /* [FromBody] string Duration */)
+		{
+			IDeviceReservation? reservation = await _deviceService.TryGetReservationFromLegacyGuidAsync(reservationGuid);
+
+			if (reservation == null)
+			{
+				_logger.LogError("Unable to find reservation for legacy guid {ReservationGuid}", reservationGuid);
+				return BadRequest();
+			}
+
+			bool updated = await _deviceService.TryUpdateReservationAsync(reservation.Id);
+
+			if (!updated)
+			{
+				_logger.LogError("Unable to find reservation for reservation {ReservationId}", reservation.Id);
+				return BadRequest();
+			}
+
+			List<IDevice> devices = await _deviceService.GetDevicesAsync(reservation.Devices);
+
+			GetLegacyReservationResponse response = new GetLegacyReservationResponse();
+
+			response.Guid = reservation.LegacyGuid;
+			response.DeviceNames = reservation.Devices.Select(deviceId => devices.First(d => d.Id == deviceId).Name).ToArray();
+			response.HostName = reservation.Hostname ?? "";
+			response.StartDateTime = reservation.CreateTimeUtc.ToString("O", System.Globalization.CultureInfo.InvariantCulture);
+			response.Duration = "00:10:00"; // matches gauntlet duration
+
+			return new JsonResult(response, new JsonSerializerOptions() { PropertyNamingPolicy = null });
+		}
+
+		/// <summary>
+		/// Delete a reservation
+		/// </summary>
+		[HttpDelete]
+		[Route("/api/v1/reservations/{reservationGuid}")]
+		public async Task<ActionResult> DeleteReservationV1Async(string reservationGuid)
+		{
+			IDeviceReservation? reservation = await _deviceService.TryGetReservationFromLegacyGuidAsync(reservationGuid);
+
+			if (reservation == null)
+			{
+				return BadRequest($"Unable to find reservation for guid {reservationGuid}");
+			}
+
+			bool deleted = await _deviceService.DeleteReservationAsync(reservation.Id);
+
+			if (!deleted)
+			{
+				return BadRequest("Failed to delete reservation");
+			}
+
+			return Ok();
+		}
+
+		/// <summary>
+		/// Get device info for a reserved device
+		/// </summary>
+		[HttpGet]
+		[Route("/api/v1/devices/{deviceName}")]
+		[ProducesResponseType(typeof(GetLegacyDeviceResponse), 200)]
+		public async Task<ActionResult<GetLegacyDeviceResponse>> GetDeviceV1Async(string deviceName)
+		{
+			IDevice? device = await _deviceService.GetDeviceByNameAsync(deviceName);
+
+			if (device == null)
+			{
+				return BadRequest($"Unknown device {deviceName}");
+			}
+
+			IDeviceReservation? reservation = await _deviceService.TryGetDeviceReservation(device.Id);
+
+			string? platformName = null;
+
+			if (reservation != null && reservation.Devices.Count == reservation.RequestedDevicePlatforms.Count)
+			{
+				for (int i = 0; i < reservation.Devices.Count; i++)
+				{
+					if (reservation.Devices[i] == device.Id)
+					{
+						platformName = reservation.RequestedDevicePlatforms[i];
+					}
+				}
+			}
+
+			DevicePlatformMapV1 mapV1 = await _deviceService.GetPlatformMapV1();
+
+			if (String.IsNullOrEmpty(platformName))
+			{				
+				_logger.LogError("Unable to map platform for {DeviceName} : {PlatformId} from reservation", deviceName, device.PlatformId);			
+				return BadRequest($"Unable to get platform for {deviceName} from reservation : {device.PlatformId}");
+			}
+
+			GetLegacyDeviceResponse response = new GetLegacyDeviceResponse();
+
+			response.Id = device.Id.ToString();
+			response.Name = device.Name;
+			response.Type = platformName;
+			response.IPOrHostName = device.Address ?? "";
+			response.AvailableStartTime = "00:00:00";
+			response.AvailableEndTime = "00:00:00";
+			response.Enabled = true;
+			response.Model = device.ModelId ?? "Base";
+			response.DeviceData = "";
+			response.PerfSpec = "Minimum";
+
+			if (device.ModelId != null)
+			{
+				string? model;
+				if (mapV1.PerfSpecHighMap.TryGetValue(device.PlatformId, out model))
+				{
+					if (model == device.ModelId)
+					{
+						response.PerfSpec = "High";
+					}
+				}
+			}
+
+			return new JsonResult(response, new JsonSerializerOptions() { PropertyNamingPolicy = null });
+		}
+
+		/// <summary>
+		/// Mark a problem device
+		/// </summary>
+		[HttpPut]
+		[Route("/api/v1/deviceerror/{deviceName}")]
+		public async Task<ActionResult> PutDeviceErrorAsync(string deviceName)
+		{
+			IDevice? device = await _deviceService.GetDeviceByNameAsync(deviceName);
+
+			if (device == null)
+			{
+				return BadRequest($"Unknown device {deviceName}");
+			}
+
+			await _deviceService.UpdateDeviceAsync(device.Id, newProblem: true);
+
+			string message = $"Device problem, {device.Name} : {device.PoolId.ToString().ToUpperInvariant()}";
+
+			IDeviceReservation? reservation = await _deviceService.TryGetDeviceReservation(device.Id);
+
+			string? jobId = null;
+			string? stepId = null;
+			if (reservation != null)
+			{
+				jobId = !String.IsNullOrEmpty(reservation.JobId) ? reservation.JobId : null;
+				stepId = !String.IsNullOrEmpty(reservation.StepId) ? reservation.StepId : null;
+
+				if ((jobId == null || stepId == null))
+				{
+					if (!String.IsNullOrEmpty(reservation.ReservationDetails))
+					{
+						message += $" - {reservation.ReservationDetails}";
+					}
+					else
+					{
+						message += $" - Host {reservation.Hostname}";
+					}
+				}
+			}			
+
+			return Ok();
+
+		}
+	}
+}
diff --git a/Engine/Source/Programs/Horde/Horde.Server/Devices/IDeviceCollection.cs b/Engine/Source/Programs/Horde/Horde.Server/Devices/IDeviceCollection.cs
new file mode 100644
--- /dev/null
+++ b/Engine/Source/Programs/Horde/Horde.Server/Devices/IDeviceCollection.cs
@@ -0,0 +1,266 @@
+// Copyright Epic Games, Inc. All Rights Reserved.
+
+using System;
+using System.Collections.Generic;
+using System.Threading.Tasks;
+using Horde.Server.Projects;
+using Horde.Server.Users;
+using Horde.Server.Utilities;
+using MongoDB.Bson;
+
+namespace Horde.Server.Devices
+{
+	/// <summary>
+	/// Device reservation request data
+	/// </summary>
+	public class DeviceRequestData
+	{
+		/// <summary>
+		/// The platform of the device to reserve
+		/// </summary>
+		public DevicePlatformId PlatformId { get; set; }
+
+		/// <summary>
+		/// The platform client requested, which can differ from the devices platformid due to mapping
+		/// </summary>
+		public string RequestedPlatform { get; set; }
+
+		/// <summary>
+		/// Models to include for this request
+		/// </summary>
+		public List<string> IncludeModels { get; set; }
+
+		/// <summary>
+		/// Models to exclude for this request
+		/// </summary>
+		public List<string> ExcludeModels { get; set; }
+
+		/// <summary>
+		/// Constructor
+		/// </summary>
+		public DeviceRequestData(DevicePlatformId platformId, string requestedPlatform, List<string>? includeModels = null, List<string>? excludeModels = null)
+		{
+			PlatformId = platformId;
+			IncludeModels = includeModels ?? new List<string>();
+			ExcludeModels = excludeModels ?? new List<string>();
+			RequestedPlatform = requestedPlatform;
+		}
+	}
+
+	/// <summary>
+	/// A collection for device management
+	/// </summary>
+	public interface IDeviceCollection
+	{
+		// PLATFORMS
+
+		/// <summary>
+		/// Add a platform 
+		/// </summary>
+		/// <param name="id">The id of the platform</param>
+		/// <param name="name">The friendly name of the platform</param>
+		Task<IDevicePlatform?> TryAddPlatformAsync(DevicePlatformId id, string name);
+
+		/// <summary>
+		/// Get a list of all available device platforms
+		/// </summary>		
+		Task<List<IDevicePlatform>> FindAllPlatformsAsync();
+
+		/// <summary>
+		/// Get a specific platform by id
+		/// </summary>
+		Task<IDevicePlatform?> GetPlatformAsync(DevicePlatformId platformId);
+
+		/// <summary>
+		/// Update a device platform
+		/// </summary>
+		/// <param name="platformId">The id of the device</param>
+		/// <param name="modelIds">The available model ids for the platform</param>
+		Task<bool> UpdatePlatformAsync(DevicePlatformId platformId, string[]? modelIds);
+
+		// POOLS
+
+		/// <summary>
+		/// Add a new device pool to the collection
+		/// </summary>
+		/// <param name="id">The id of the new pool</param>
+		/// <param name="name">The friendly name of the new pool</param>
+		/// <param name="poolType">The pool type</param>
+		/// <param name="projectIds">Projects associated with this pool</param>
+		Task<IDevicePool?> TryAddPoolAsync(DevicePoolId id, string name, DevicePoolType poolType, List<ProjectId>? projectIds );
+
+		/// <summary>
+		/// Update a device pool
+		/// </summary>
+		/// <param name="id">The id of the device pool to update</param>
+		/// <param name="projectIds">Associated project ids</param>
+		Task UpdatePoolAsync(DevicePoolId id, List<ProjectId>? projectIds);
+
+		/// <summary>
+		/// Get a pool by id
+		/// </summary>
+		Task<IDevicePool?> GetPoolAsync(DevicePoolId poolId);
+
+		/// <summary>
+		/// Gets a list of existing device pools
+		/// </summary>		
+		Task<List<IDevicePool>> FindAllPoolsAsync();
+
+		// DEVICES
+
+		/// <summary>
+		/// Get a device by id
+		/// </summary>
+		Task<IDevice?> GetDeviceAsync(DeviceId deviceId);
+
+		/// <summary>
+		/// Get a device by name
+		/// </summary>
+		Task<IDevice?> GetDeviceByNameAsync(string deviceName);
+
+		/// <summary>
+		/// Get a list of all devices
+		/// </summary>
+		/// <param name="deviceIds">Optional list of device ids to get</param>
+		/// <param name="poolId">Optional id of pool for devices</param>
+		/// <param name="platformId">Optional id of platform for devices</param>
+		Task<List<IDevice>> FindAllDevicesAsync(List<DeviceId>? deviceIds, DevicePoolId? poolId = null, DevicePlatformId? platformId = null);
+
+		/// <summary>
+		/// Adds a new device
+		/// </summary>
+		/// <param name="id">The device id</param>
+		/// <param name="name">The name of the device (unique)</param>
+		/// <param name="platformId">The platform of the device</param>
+		/// <param name="poolId">Which pool to add it to</param>
+		/// <param name="enabled">Whather the device is enabled by default</param>
+		/// <param name="address">The network address or hostname the device can be reached at</param>
+		/// <param name="modelId">The vendor model id of the device</param>
+        /// <param name="userId">The user adding the device</param>
+		Task<IDevice?> TryAddDeviceAsync(DeviceId id, string name, DevicePlatformId platformId, DevicePoolId poolId, bool? enabled, string? address, string? modelId, UserId? userId);
+
+		/// <summary>
+		/// Update a device
+		/// </summary>
+		/// <param name="deviceId">The id of the device to update</param>
+		/// <param name="newPoolId">The new pool to assign</param>
+		/// <param name="newName">The new name to assign</param>		
+		/// <param name="newAddress">The new device address or hostname</param>
+		/// <param name="newModelId">The new model id</param>
+		/// <param name="newNotes">The devices markdown notes</param>
+		/// <param name="newEnabled">Whether the device is enabled or not</param>
+		/// <param name="newProblem">Whether to set or clear problem state</param>
+		/// <param name="newMaintenance">Whether to set or clear maintenance state</param>
+        /// <param name="modifiedByUserId">The user who is updating the device</param>
+		Task UpdateDeviceAsync(DeviceId deviceId, DevicePoolId? newPoolId, string? newName, string? newAddress, string? newModelId, string? newNotes, bool? newEnabled, bool? newProblem, bool? newMaintenance, UserId? modifiedByUserId = null);
+
+		/// <summary>
+		/// Delete a device from the collection
+		/// </summary>
+		/// <param name="deviceId">The id of the device to delete</param>
+		Task<bool> DeleteDeviceAsync(DeviceId deviceId);
+
+		/// <summary>
+        /// Checkout or checkin the specified device
+        /// </summary>
+        /// <param name="deviceId"></param>
+        /// <param name="checkedOutByUserId"></param>
+        /// <returns></returns>
+        Task CheckoutDeviceAsync(DeviceId deviceId, UserId? checkedOutByUserId);
+
+		// RESERVATIONS
+
+		/// <summary>
+		/// Create a new reseveration in the pool with the specified devices
+		/// </summary>
+		/// <param name="poolId">The pool of devices to use for the new reservation</param>
+		/// <param name="request">The requested devices for the reservation</param>
+		/// <param name="hostname">The hostname of the machine making the reservation</param>
+		/// <param name="reservationDetails">The details of the reservation</param>
+		/// <param name="streamId">The Stream Id associated with the job</param>
+		/// <param name="jobId">The Job Id associated with the job</param>
+		/// <param name="stepId">The Step Id associated with the job</param>
+		/// <param name="jobName">The Job name associated with the job</param>
+		/// <param name="stepName">The Step name associated with the job</param>		
+		Task<IDeviceReservation?> TryAddReservationAsync(DevicePoolId poolId, List<DeviceRequestData> request, string? hostname, string? reservationDetails, string? streamId, string? jobId, string? stepId, string? jobName, string? stepName);
+
+		/// <summary>
+		/// Gets a reservation by guid for legacy clients
+		/// </summary>
+		/// <param name="legacyGuid">YThe legacy guid of the reservation</param>
+		Task<IDeviceReservation?> TryGetReservationFromLegacyGuidAsync(string legacyGuid);
+
+        /// <summary>
+        /// Gets a reservation by device id
+        /// </summary>
+        /// <param name="id">A device contained in reservation</param>
+        Task<IDeviceReservation?> TryGetDeviceReservationAsync(DeviceId id);
+
+        /// <summary>
+        /// Get a list of all reservations
+        /// </summary>
+        Task<List<IDeviceReservation>> FindAllReservationsAsync();
+
+		/// <summary>
+		/// Updates a reservation to the current time, for expiration
+		/// </summary>
+		/// <param name="id">The id of the reservation to update</param>
+		public Task<bool> TryUpdateReservationAsync(ObjectId id);
+
+		/// <summary>
+		/// Deletes a reservation and releases reserved devices
+		/// </summary>
+		/// <param name="id">The id of the reservation to delete</param>
+		public Task<bool> DeleteReservationAsync(ObjectId id);
+
+		/// <summary>
+		/// Deletes expired reservations
+		/// </summary>
+		public Task<bool> ExpireReservationsAsync();
+
+		/// <summary>
+		/// Deletes user device checkouts
+		/// </summary>
+		/// <param name="checkoutDays"></param>
+		public Task<List<(UserId, IDevice)>?> ExpireCheckedOutAsync(int checkoutDays);
+
+		/// <summary>
+		/// Gets a list of users to notify whose device checkout is about to expire
+		/// </summary>
+		/// <param name="checkoutDays"></param>
+		public Task<List<(UserId, IDevice)>?> ExpireNotificatonsAsync(int checkoutDays);
+
+		/// <summary>
+		/// Find device telemetry
+		/// </summary>
+		/// <param name="deviceIds"></param>
+		/// <param name="minCreateTime"></param>
+		/// <param name="maxCreateTime"></param>
+		/// <param name="index"></param>
+		/// <param name="count"></param>
+		/// <returns></returns>
+		public Task<List<IDeviceTelemetry>> FindDeviceTelemetryAsync(DeviceId[]? deviceIds = null, DateTimeOffset? minCreateTime = null, DateTimeOffset? maxCreateTime = null, int? index = null, int? count = null);
+
+		/// <summary>
+		/// Creates a device pool telemetry snapshot
+		/// </summary>
+		/// <returns></returns>
+		public Task CreatePoolTelemetrySnapshot();
+
+		/// <summary>
+		/// Gets pool telemetry for an optional date range
+		/// </summary>
+		/// <param name="minCreateTime"></param>
+		/// <param name="maxCreateTime"></param>
+		/// <param name="index"></param>
+		/// <param name="count"></param>
+		/// <returns></returns>
+		public Task<List<IDevicePoolTelemetry>> FindPoolTelemetryAsync(DateTimeOffset? minCreateTime = null, DateTimeOffset? maxCreateTime = null, int? index = null, int? count = null);
+
+		/// <summary>
+		/// Upgrades the collection
+		/// </summary>
+		public Task UpgradeAsync();
+
+	}
+}
\ No newline at end of file
diff --git a/Engine/Source/Programs/Horde/Horde.Server/Dockerfile b/Engine/Source/Programs/Horde/Horde.Server/Dockerfile
new file mode 100644
--- /dev/null
+++ b/Engine/Source/Programs/Horde/Horde.Server/Dockerfile
@@ -0,0 +1,64 @@
+#
+# NOTE: This Dockerfile must be built with the Engine directory as CWD.
+#       In the Horde build script we stage the required files before running:
+#
+#         docker build -f Source/Programs/Horde/Horde.Server/Dockerfile . 
+#
+#       Running directly from the source tree will cause a lot of unnecessary
+#       files to be uploaded to Docker before building.
+
+FROM redis:6.2.7 AS redis
+
+# Build image
+FROM mcr.microsoft.com/dotnet/sdk:6.0 AS build-env
+WORKDIR /app
+COPY Source/Programs/Horde/nuget.config ./Source/Programs/Horde/nuget.config
+RUN DEBIAN_FRONTEND=noninteractive apt-get -qq update && apt-get -qq install curl
+
+COPY --from=redis /usr/local/bin/redis-server /usr/local/bin/redis-server
+
+# Since the .deb does not install in this image, just download it and extract the static binary
+RUN wget https://repo.mongodb.org/apt/debian/dists/buster/mongodb-org/4.4/main/binary-amd64/mongodb-org-server_4.4.0_amd64.deb && \
+	dpkg -x mongodb-org-server_4.4.0_amd64.deb /tmp/mongodb && \
+	cp /tmp/mongodb/usr/bin/mongod /usr/local/bin/mongod
+
+# Restore all the C# projects
+COPY Source/Programs/Shared/EpicGames.Core/*.csproj ./Source/Programs/Shared/EpicGames.Core/
+COPY Source/Programs/Shared/EpicGames.Horde/*.csproj ./Source/Programs/Shared/EpicGames.Horde/
+COPY Source/Programs/Shared/EpicGames.IoHash/*.csproj ./Source/Programs/Shared/EpicGames.IoHash/
+COPY Source/Programs/Shared/EpicGames.Perforce/*.csproj ./Source/Programs/Shared/EpicGames.Perforce/
+COPY Source/Programs/Shared/EpicGames.Perforce.Managed/*.csproj ./Source/Programs/Shared/EpicGames.Perforce.Managed/
+COPY Source/Programs/Shared/EpicGames.Redis/*.csproj ./Source/Programs/Shared/EpicGames.Redis/
+COPY Source/Programs/Shared/EpicGames.Serialization/*.csproj ./Source/Programs/Shared/EpicGames.Serialization/
+COPY Source/Programs/Shared/EpicGames.Serialization.Tests/*.csproj ./Source/Programs/Shared/EpicGames.Serialization.Tests/
+COPY Source/Programs/Horde/Horde.Agent/*.csproj ./Source/Programs/Horde/Horde.Agent/
+COPY Source/Programs/Horde/Horde.Server/*.csproj ./Source/Programs/Horde/Horde.Server/
+COPY Source/Programs/Horde/Horde.Server.Tests/*.csproj ./Source/Programs/Horde/Horde.Server.Tests/
+RUN dotnet restore Source/Programs/Horde/Horde.Server/Horde.Server.csproj -p:WithHordeStorage=false
+
+# Build the server
+COPY Binaries/DotNET/EpicGames.Perforce.Native/ Binaries/DotNET/EpicGames.Perforce.Native/
+COPY Source/Programs/Shared/ Source/Programs/Shared/
+COPY Source/Programs/Horde/Horde.Server/ Source/Programs/Horde/Horde.Server/
+COPY Source/Programs/AutomationTool/AutomationUtils/Matchers/ Source/Programs/AutomationTool/AutomationUtils/Matchers/
+COPY Source/Programs/UnrealBuildTool/Matchers/ Source/Programs/UnrealBuildTool/Matchers/
+ARG msbuild_args
+RUN dotnet publish Source/Programs/Horde/Horde.Server/Horde.Server.csproj -c Release -o out -p:WithHordeStorage=false $msbuild_args -p:RunAnalyzersDuringBuild=True
+
+# Run tests
+COPY Source/Programs/Horde/Horde.Agent/ Source/Programs/Horde/Horde.Agent/
+COPY Source/Programs/Horde/Horde.Server.Tests/ Source/Programs/Horde/Horde.Server.Tests/
+COPY Source/Programs/Shared/EpicGames.Serialization.Tests/ Source/Programs/Shared/EpicGames.Serialization.Tests/
+RUN redis-server --save "" --appendonly no --daemonize yes && \
+	mongod --noauth --quiet --fork --dbpath /tmp/mongodb --logpath /tmp/mongod.log && \
+	dotnet test Source/Programs/Shared/EpicGames.BuildGraph.Tests/EpicGames.BuildGraph.Tests.csproj --blame-hang-timeout 2m --logger console;verbosity=normal && \
+	dotnet test Source/Programs/Shared/EpicGames.Core.Tests/EpicGames.Core.Tests.csproj --blame-hang-timeout 2m --logger console;verbosity=normal && \
+	dotnet test Source/Programs/Shared/EpicGames.Horde.Tests/EpicGames.Horde.Tests.csproj --blame-hang-timeout 2m --logger console;verbosity=normal && \
+	dotnet test Source/Programs/Shared/EpicGames.Serialization.Tests/EpicGames.Serialization.Tests.csproj --blame-hang-timeout 2m --logger console;verbosity=normal && \
+	dotnet test Source/Programs/Horde/Horde.Server.Tests/Horde.Server.Tests.csproj --blame-hang-timeout 2m --logger console;verbosity=normal
+
+# Create the runtime image
+FROM mcr.microsoft.com/dotnet/aspnet:6.0 AS runtime
+WORKDIR /app
+COPY --from=build-env /app/out .
+ENTRYPOINT [ "dotnet", "Horde.Server.dll" ]
diff --git a/Engine/Source/Programs/Horde/Horde.Server/Horde.Server.csproj b/Engine/Source/Programs/Horde/Horde.Server/Horde.Server.csproj
new file mode 100644
--- /dev/null
+++ b/Engine/Source/Programs/Horde/Horde.Server/Horde.Server.csproj
@@ -0,0 +1,138 @@
+<Project Sdk="Microsoft.NET.Sdk.Web">
+
+  <PropertyGroup>
+    <TargetFramework>net6.0</TargetFramework>
+    <Nullable>enable</Nullable>
+    <Deterministic>true</Deterministic>
+    <AssemblyName>Horde.Server</AssemblyName>
+    <DocumentationFile>Horde.Server.xml</DocumentationFile>
+    <ApplicationIcon>HordeServer.ico</ApplicationIcon>
+    <RunAnalyzersDuringBuild>False</RunAnalyzersDuringBuild>
+    <Configurations>Debug;Release;Analyze</Configurations>
+  </PropertyGroup>
+
+  <PropertyGroup Condition=" '$(Configuration)|$(Platform)' == 'Analyze|AnyCPU' ">
+    <RunAnalyzersDuringBuild>True</RunAnalyzersDuringBuild>
+    <TreatWarningsAsErrors>false</TreatWarningsAsErrors>
+    <GenerateDocumentationFile>True</GenerateDocumentationFile>
+    <DocumentationFile></DocumentationFile>
+  </PropertyGroup>
+
+  <ItemGroup>
+    <PackageReference Include="AWSSDK.AutoScaling" Version="3.7.100" />
+    <PackageReference Include="AWSSDK.CloudWatch" Version="3.7.100" />
+    <PackageReference Include="AWSSDK.EC2" Version="3.7.100" />
+    <PackageReference Include="AWSSDK.Extensions.NETCore.Setup" Version="3.7.2" />
+    <PackageReference Include="AWSSDK.S3" Version="3.7.100" />
+    <PackageReference Include="AWSSDK.SecurityToken" Version="3.7.100" />
+    <PackageReference Include="Crc32.NET" Version="1.2.0" />
+    <PackageReference Include="Datadog.Trace.OpenTracing" Version="2.14.0" />
+    <PackageReference Include="DogStatsD-CSharp-Client" Version="7.0.1" />
+    <PackageReference Include="Grpc.AspNetCore" Version="2.49.0" />
+    <PackageReference Include="Grpc.AspNetCore.Server.Reflection" Version="2.49.0" />
+    <PackageReference Include="Grpc.Net.Client" Version="2.49.0" />
+    <PackageReference Include="Grpc.Tools" Version="2.50.0">
+      <PrivateAssets>all</PrivateAssets>
+      <IncludeAssets>runtime; build; native; contentfiles; analyzers; buildtransitive</IncludeAssets>
+    </PackageReference>
+    <PackageReference Include="JetBrains.Profiler.SelfApi" Version="2.3.0" />
+    <PackageReference Include="Microsoft.AspNetCore.Authentication.Cookies" Version="2.2.0" />
+    <PackageReference Include="Microsoft.AspNetCore.Authentication.JwtBearer" Version="6.0.10" />
+    <PackageReference Include="Microsoft.AspNetCore.Authentication.OpenIdConnect" Version="6.0.10" />
+    <PackageReference Include="Microsoft.AspNetCore.DataProtection.StackExchangeRedis" Version="6.0.10" />
+    <PackageReference Include="Microsoft.AspNetCore.SpaServices.Extensions" Version="6.0.10" />
+    <PackageReference Include="Microsoft.Extensions.Caching.Memory" Version="6.0.1" />
+    <PackageReference Include="Microsoft.Extensions.Http" Version="6.0.0" />
+    <PackageReference Include="MongoDB.Driver" Version="2.18.0" />
+    <PackageReference Include="OpenTracing" Version="0.12.1" />
+    <PackageReference Include="OpenTracing.Contrib.Grpc" Version="0.4.0" />
+    <PackageReference Include="Polly" Version="7.2.3" />
+    <PackageReference Include="Polly.Extensions.Http" Version="3.0.0" />
+    <PackageReference Include="Serilog.AspNetCore" Version="6.0.1" />
+    <PackageReference Include="Serilog.Enrichers.OpenTracing" Version="0.0.2" />
+    <PackageReference Include="Serilog.Exceptions.Grpc" Version="8.4.0" />
+    <PackageReference Include="Serilog.Extensions.Hosting" Version="5.0.1" />
+    <PackageReference Include="Serilog.Sinks.Console" Version="4.1.0" />
+    <PackageReference Include="Serilog.Sinks.File" Version="5.0.0" />
+    <PackageReference Include="SerilogAnalyzer" Version="0.15.0" />
+    <PackageReference Include="SharpZipLib" Version="1.4.0" />
+    <PackageReference Include="StackExchange.Redis" Version="2.6.70" />
+    <PackageReference Include="Swashbuckle.AspNetCore" Version="6.4.0" />
+    <PackageReference Include="System.Linq.Async" Version="6.0.1" />
+    <PackageReference Include="TimeZoneConverter" Version="6.0.1" />
+  </ItemGroup>
+
+  <ItemGroup>
+    <Content Include="..\..\Shared\EpicGames.Horde.Common\Protos\**\*.*">
+      <Link>Protos\%(RecursiveDir)%(FileName)%(Extension)</Link>
+    </Content>
+  </ItemGroup>
+
+  <ItemGroup>
+    <None Remove="HordeServer.xml" />
+  </ItemGroup>
+
+  <ItemGroup>
+    <Protobuf Include="Protos\horde\job_rpc.proto" ProtoRoot="Protos" GrpcServices="Server" />
+    <Protobuf Include="Protos\horde\job_rpc_messages.proto" ProtoRoot="Protos" />
+    <Protobuf Include="Protos\horde\horde_rpc.proto" ProtoRoot="Protos" GrpcServices="Server" />
+    <Protobuf Include="Protos\horde\horde_rpc_messages.proto" ProtoRoot="Protos" />
+    <Protobuf Include="Protos\horde\log_rpc.proto" ProtoRoot="Protos" GrpcServices="Server" />
+    <Protobuf Include="Protos\horde\log_rpc_messages.proto" ProtoRoot="Protos" />
+    <Protobuf Include="Protos\horde\common\agent_status.proto" ProtoRoot="Protos" />
+    <Protobuf Include="Protos\horde\common\event_severity.proto" ProtoRoot="Protos" />
+    <Protobuf Include="Protos\horde\common\io_hash.proto" ProtoRoot="Protos" />
+    <Protobuf Include="Protos\horde\common\jobstep_outcome.proto" ProtoRoot="Protos" />
+    <Protobuf Include="Protos\horde\common\jobstep_state.proto" ProtoRoot="Protos" />
+    <Protobuf Include="Protos\horde\common\label_change.proto" ProtoRoot="Protos" />
+    <Protobuf Include="Protos\horde\common\lease_outcome.proto" ProtoRoot="Protos" />
+    <Protobuf Include="Protos\horde\common\lease_state.proto" ProtoRoot="Protos" />
+    <Protobuf Include="Protos\horde\common\namespace_id.proto" ProtoRoot="Protos" />
+    <Protobuf Include="Protos\horde\common\priority.proto" ProtoRoot="Protos" />
+    <Protobuf Include="Protos\horde\messages\agent_capabilities.proto" ProtoRoot="Protos" />
+    <Protobuf Include="Protos\horde\messages\agent_requirements.proto" ProtoRoot="Protos" />
+    <Protobuf Include="Protos\horde\messages\agent_workspace.proto" ProtoRoot="Protos" />
+    <Protobuf Include="Protos\horde\messages\lease.proto" ProtoRoot="Protos" />
+    <Protobuf Include="Protos\horde\tasks\compute_task.proto" ProtoRoot="Protos" />
+    <Protobuf Include="Protos\horde\tasks\conform_task.proto" ProtoRoot="Protos" />
+    <Protobuf Include="Protos\horde\tasks\job_task.proto" ProtoRoot="Protos" />
+    <Protobuf Include="Protos\horde\tasks\restart_task.proto" ProtoRoot="Protos" />
+    <Protobuf Include="Protos\horde\tasks\shutdown_task.proto" ProtoRoot="Protos" />
+    <Protobuf Include="Protos\horde\tasks\upgrade_task.proto" ProtoRoot="Protos" />
+    <Protobuf Include="Protos\grpc\health\v1\health.proto" ProtoRoot="Protos" GrpcServices="Server" />
+  </ItemGroup>
+
+  <ItemGroup>
+    <ProjectReference Include="..\..\Shared\EpicGames.AspNet\EpicGames.AspNet.csproj" />
+    <ProjectReference Include="..\..\Shared\EpicGames.BuildGraph\EpicGames.BuildGraph.csproj" />
+    <ProjectReference Include="..\..\Shared\EpicGames.Core\EpicGames.Core.csproj" />
+    <ProjectReference Include="..\..\Shared\EpicGames.Horde\EpicGames.Horde.csproj" />
+    <ProjectReference Include="..\..\Shared\EpicGames.IoHash\EpicGames.IoHash.csproj" />
+    <ProjectReference Include="..\..\Shared\EpicGames.Oodle\EpicGames.Oodle.csproj" />
+    <ProjectReference Include="..\..\Shared\EpicGames.Perforce.Managed\EpicGames.Perforce.Managed.csproj" />
+    <ProjectReference Include="..\..\Shared\EpicGames.Perforce\EpicGames.Perforce.csproj" />
+    <ProjectReference Include="..\..\Shared\EpicGames.Redis\EpicGames.Redis.csproj" />
+    <ProjectReference Include="..\..\Shared\EpicGames.Serialization\EpicGames.Serialization.csproj" />
+    <ProjectReference Include="..\..\Shared\EpicGames.Slack\EpicGames.Slack.csproj" />
+  </ItemGroup>
+
+  <ItemGroup>
+    <Content Include="..\ThirdParty\**" Condition="$([MSBuild]::IsOSPlatform('Windows'))">
+      <CopyToOutputDirectory>Always</CopyToOutputDirectory>
+      <Link>ThirdParty/%(RecursiveDir)/%(Filename)%(Extension)</Link>
+    </Content>
+  </ItemGroup>
+
+  <ItemGroup>
+    <None Include="DashboardApp\**">
+      <CopyToOutputDirectory>Always</CopyToOutputDirectory>
+    </None>
+  </ItemGroup>
+
+  <ItemGroup>
+    <Content Include="Defaults\ue5.project.png">
+      <CopyToOutputDirectory>Always</CopyToOutputDirectory>
+    </Content>
+  </ItemGroup>
+
+</Project>
diff --git a/Engine/Source/Programs/Horde/Horde.Server/Issues/Handlers/ContentIssueHandler.cs b/Engine/Source/Programs/Horde/Horde.Server/Issues/Handlers/ContentIssueHandler.cs
new file mode 100644
--- /dev/null
+++ b/Engine/Source/Programs/Horde/Horde.Server/Issues/Handlers/ContentIssueHandler.cs
@@ -0,0 +1,117 @@
+// Copyright Epic Games, Inc. All Rights Reserved.
+
+using System;
+using System.Collections.Generic;
+using System.Linq;
+using EpicGames.Core;
+using Horde.Server.Jobs;
+using Horde.Server.Jobs.Graphs;
+using Horde.Server.Logs;
+using Horde.Server.Utilities;
+using Microsoft.Extensions.Logging;
+using MongoDB.Driver;
+
+namespace Horde.Server.Issues.Handlers
+{
+	/// <summary>
+	/// Instance of a particular compile error
+	/// </summary>
+	class ContentIssueHandler : IssueHandler
+	{
+		/// <inheritdoc/>
+		public override string Type => "Content";
+
+		/// <inheritdoc/>
+		public override int Priority => 10;
+
+		/// <summary>
+		/// Determines if the given event id matches
+		/// </summary>
+		/// <param name="eventId">The event id to compare</param>
+		/// <returns>True if the given event id matches</returns>
+		public static bool IsMatchingEventId(EventId eventId)
+		{
+			return eventId == KnownLogEvents.Engine_AssetLog;
+		}
+
+		/// <summary>
+		/// Determines if an event should be masked by this 
+		/// </summary>
+		/// <param name="eventId"></param>
+		/// <returns></returns>
+		static bool IsMaskedEventId(EventId eventId)
+		{
+			return eventId == KnownLogEvents.ExitCode;
+		}
+
+		/// <summary>
+		/// Adds all the assets from the given log event
+		/// </summary>
+		/// <param name="eventData">The log event to parse</param>
+		/// <param name="assetNames">Receives the referenced asset names</param>
+		public static void GetAssetNames(ILogEventData eventData, HashSet<string> assetNames)
+		{
+			foreach (ILogEventLine line in eventData.Lines)
+			{
+				string? relativePath;
+				if (line.Data.TryGetNestedProperty("properties.asset.relativePath", out relativePath) || line.Data.TryGetNestedProperty("properties.asset.$text", out relativePath))
+				{
+					int endIdx = relativePath.LastIndexOfAny(new char[] { '/', '\\' }) + 1;
+					string fileName = relativePath.Substring(endIdx);
+					assetNames.Add(fileName);
+				}
+			}
+		}
+
+		/// <inheritdoc/>
+		public override void TagEvents(IJob job, INode node, IReadOnlyNodeAnnotations annotations, IReadOnlyList<IssueEvent> stepEvents)
+		{
+			bool hasMatches = false;
+			foreach (IssueEvent stepEvent in stepEvents)
+			{
+				if (stepEvent.EventId != null)
+				{
+					if (IsMatchingEventId(stepEvent.EventId.Value))
+					{
+						HashSet<string> newAssetNames = new HashSet<string>();
+						GetAssetNames(stepEvent.EventData, newAssetNames);
+
+						stepEvent.Fingerprint = new NewIssueFingerprint(Type, newAssetNames, null, null);
+						hasMatches = true;
+					}
+					else if (hasMatches && IsMaskedEventId(stepEvent.EventId.Value))
+					{
+						stepEvent.Ignored = true;
+					}
+				}
+			}
+		}
+
+		/// <inheritdoc/>
+		public override string GetSummary(IIssueFingerprint fingerprint, IssueSeverity severity)
+		{
+			string type = (severity == IssueSeverity.Warning) ? "Warnings" : "Errors";
+			string list = StringUtils.FormatList(fingerprint.Keys.ToArray(), 2);
+			return $"{type} in {list}";
+		}
+
+		/// <inheritdoc/>
+		public override void RankSuspects(IIssueFingerprint fingerprint, List<SuspectChange> suspects)
+		{
+			foreach (SuspectChange suspect in suspects)
+			{
+				if (suspect.ContainsContent)
+				{
+					if (suspect.Files.Any(x => fingerprint.Keys.Any(y => x.Contains(y, StringComparison.OrdinalIgnoreCase))))
+					{
+						suspect.Rank += 20;
+					}
+					else
+					{
+						suspect.Rank += 10;
+					}
+				}
+			}
+		}
+	}
+}
diff --git a/Engine/Source/Programs/Horde/Horde.Server/Issues/Handlers/CopyrightIssueHandler.cs b/Engine/Source/Programs/Horde/Horde.Server/Issues/Handlers/CopyrightIssueHandler.cs
new file mode 100644
--- /dev/null
+++ b/Engine/Source/Programs/Horde/Horde.Server/Issues/Handlers/CopyrightIssueHandler.cs
@@ -0,0 +1,53 @@
+// Copyright Epic Games, Inc. All Rights Reserved.
+
+using System.Collections.Generic;
+using System.Linq;
+using EpicGames.Core;
+using Horde.Server.Jobs;
+using Horde.Server.Jobs.Graphs;
+using Microsoft.Extensions.Logging;
+
+namespace Horde.Server.Issues.Handlers
+{
+	/// <summary>
+	/// Instance of a particular compile error
+	/// </summary>
+	class CopyrightIssueHandler : SourceFileIssueHandler
+	{
+		/// <inheritdoc/>
+		public override string Type => "Copyright";
+
+		/// <inheritdoc/>
+		public override int Priority => 10;
+
+		/// <summary>
+		/// Determines if the given event id matches
+		/// </summary>
+		/// <param name="eventId">The event id to compare</param>
+		/// <returns>True if the given event id matches</returns>
+		public static bool IsMatchingEventId(EventId eventId)
+		{
+			return eventId == KnownLogEvents.AutomationTool_MissingCopyright;
+		}
+
+		/// <inheritdoc/>
+		public override void TagEvents(IJob job, INode node, IReadOnlyNodeAnnotations annotations, IReadOnlyList<IssueEvent> stepEvents)
+		{
+			foreach (IssueEvent stepEvent in stepEvents)
+			{
+				if (stepEvent.EventId != null && IsMatchingEventId(stepEvent.EventId.Value))
+				{
+					List<string> newFileNames = new List<string>();
+					GetSourceFiles(stepEvent.EventData, newFileNames);
+					stepEvent.Fingerprint = new NewIssueFingerprint(Type, newFileNames, null, null);
+				}
+			}
+		}
+
+		/// <inheritdoc/>
+		public override string GetSummary(IIssueFingerprint fingerprint, IssueSeverity severity)
+		{
+			return $"Missing copyright notice in {StringUtils.FormatList(fingerprint.Keys.ToArray(), 2)}";
+		}
+	}
+}
diff --git a/Engine/Source/Programs/Horde/Horde.Server/Issues/Handlers/DefaultIssueHandler.cs b/Engine/Source/Programs/Horde/Horde.Server/Issues/Handlers/DefaultIssueHandler.cs
new file mode 100644
--- /dev/null
+++ b/Engine/Source/Programs/Horde/Horde.Server/Issues/Handlers/DefaultIssueHandler.cs
@@ -0,0 +1,55 @@
+// Copyright Epic Games, Inc. All Rights Reserved.
+
+using System.Collections.Generic;
+using System.Linq;
+using Horde.Server.Jobs;
+using Horde.Server.Jobs.Graphs;
+
+namespace Horde.Server.Issues.Handlers
+{
+	/// <summary>
+	/// Instance of a particular compile error
+	/// </summary>
+	class DefaultIssueHandler : IssueHandler
+	{
+		/// <summary>
+		/// Name of the handler
+		/// </summary>
+		public const string TypeConst = "Default";
+
+		/// <inheritdoc/>
+		public override string Type => TypeConst;
+
+		/// <inheritdoc/>
+		public override int Priority => 0;
+
+		/// <inheritdoc/>
+		public override void TagEvents(IJob job, INode node, IReadOnlyNodeAnnotations annotations, IReadOnlyList<IssueEvent> stepEvents)
+		{
+			NewIssueFingerprint fingerprint = new NewIssueFingerprint(TypeConst, new[] { node.Name }, null, null);
+			foreach (IssueEvent stepEvent in stepEvents)
+			{
+				stepEvent.Fingerprint = fingerprint;
+			}
+		}
+
+		/// <inheritdoc/>
+		public override void RankSuspects(IIssueFingerprint fingerprint, List<SuspectChange> suspects)
+		{
+		}
+
+		/// <inheritdoc/>
+		public override string GetSummary(IIssueFingerprint fingerprint, IssueSeverity severity)
+		{
+			string nodeName = fingerprint.Keys.FirstOrDefault() ?? "(unknown)";
+			if(severity == IssueSeverity.Warning)
+			{
+				return $"Warnings in {nodeName}";
+			}
+			else
+			{
+				return $"Errors in {nodeName}";
+			}
+		}
+	}
+}
diff --git a/Engine/Source/Programs/Horde/Horde.Server/Issues/Handlers/SymbolIssueHandler.cs b/Engine/Source/Programs/Horde/Horde.Server/Issues/Handlers/SymbolIssueHandler.cs
new file mode 100644
--- /dev/null
+++ b/Engine/Source/Programs/Horde/Horde.Server/Issues/Handlers/SymbolIssueHandler.cs
@@ -0,0 +1,154 @@
+// Copyright Epic Games, Inc. All Rights Reserved.
+
+using System;
+using System.Collections.Generic;
+using System.Linq;
+using System.Text;
+using EpicGames.Core;
+using Horde.Server.Jobs;
+using Horde.Server.Jobs.Graphs;
+using Horde.Server.Logs;
+using Horde.Server.Utilities;
+using Microsoft.Extensions.Logging;
+
+namespace Horde.Server.Issues.Handlers
+{
+	/// <summary>
+	/// Instance of a particular compile error
+	/// </summary>
+	class SymbolIssueHandler : IssueHandler
+	{
+		const string NodeName = "Node";
+		const string EventIdName = "EventId";
+
+		static readonly string s_duplicateEventIdMetadata = $"{EventIdName}={KnownLogEvents.Linker_DuplicateSymbol.Id}";
+
+		/// <inheritdoc/>
+		public override string Type => "Symbol";
+
+		/// <inheritdoc/>
+		public override int Priority => 10;
+
+		/// <summary>
+		/// Determines if the given event id matches
+		/// </summary>
+		/// <param name="eventId">The event id to compare</param>
+		/// <returns>True if the given event id matches</returns>
+		public static bool IsMatchingEventId(EventId eventId)
+		{
+			return eventId == KnownLogEvents.Linker_UndefinedSymbol || eventId == KnownLogEvents.Linker_DuplicateSymbol || eventId == KnownLogEvents.Linker;
+		}
+
+		/// <summary>
+		/// Determines if an event should be masked by this 
+		/// </summary>
+		/// <param name="eventId"></param>
+		/// <returns></returns>
+		static bool IsMaskedEventId(EventId eventId)
+		{
+			return eventId == KnownLogEvents.ExitCode || eventId == KnownLogEvents.Systemic_Xge_BuildFailed;
+		}
+
+		/// <summary>
+		/// Parses symbol names from a log event
+		/// </summary>
+		/// <param name="eventData">The log event data</param>
+		/// <param name="symbolNames">Receives the list of symbol names</param>
+		public static void GetSymbolNames(ILogEventData eventData, SortedSet<string> symbolNames)
+		{
+			foreach (ILogEventLine line in eventData.Lines)
+			{
+				string? identifier;
+				if (line.Data.TryGetNestedProperty("properties.symbol.identifier", out identifier))
+				{
+					symbolNames.Add(identifier);
+				}
+			}
+		}
+
+		/// <inheritdoc/>
+		public override void RankSuspects(IIssueFingerprint fingerprint, List<SuspectChange> changes)
+		{
+			HashSet<string> names = new HashSet<string>();
+			foreach (string name in fingerprint.Keys)
+			{
+				names.UnionWith(name.Split("::", StringSplitOptions.RemoveEmptyEntries));
+			}
+
+			foreach (SuspectChange change in changes)
+			{
+				if (change.ContainsCode)
+				{
+					int matches = names.Count(x => change.Files.Any(y => y.Contains(x, StringComparison.OrdinalIgnoreCase)));
+					change.Rank += 10 + (10 * matches);
+				}
+			}
+		}
+
+		/// <inheritdoc/>
+		public override string GetSummary(IIssueFingerprint fingerprint, IssueSeverity severity)
+		{
+			HashSet<string> symbols = fingerprint.Keys;
+			if (symbols.Count == 0)
+			{
+				string[] nodes = fingerprint.GetMetadataValues(NodeName).ToArray();
+
+				StringBuilder summary = new StringBuilder("Linker ");
+				summary.Append((severity == IssueSeverity.Warning) ? "warnings" : "errors");
+				if (nodes.Length > 0)
+				{
+					summary.Append($" in {StringUtils.FormatList(nodes, 2)}");
+				}
+
+				return summary.ToString();
+			}
+			else
+			{
+				string problemType = fingerprint.HasMetadataValue(s_duplicateEventIdMetadata) ? "Duplicate" : "Undefined";
+				if (symbols.Count == 1)
+				{
+					return $"{problemType} symbol '{symbols.First()}'";
+				}
+				else
+				{
+					return $"{problemType} symbols: {StringUtils.FormatList(symbols.ToArray(), 3)}";
+				}
+			}
+		}
+
+		public override void TagEvents(IJob job, INode node, IReadOnlyNodeAnnotations annotations, IReadOnlyList<IssueEvent> stepEvents)
+		{
+			bool hasMatches = false;
+			foreach (IssueEvent stepEvent in stepEvents)
+			{
+				if (stepEvent.EventId != null)
+				{
+					EventId eventId = stepEvent.EventId.Value;
+					if (IsMatchingEventId(eventId))
+					{
+						SortedSet<string> symbolNames = new SortedSet<string>();
+						GetSymbolNames(stepEvent.EventData, symbolNames);
+
+						if (symbolNames.Count > 0)
+						{
+							List<string> metadata = new List<string>();
+							metadata.Add($"{NodeName}={node.Name}");
+							metadata.Add($"{EventIdName}={eventId.Id}");
+
+							stepEvent.Fingerprint = new NewIssueFingerprint(Type, symbolNames, null, metadata);
+							hasMatches = true;
+						}
+						else if (hasMatches)
+						{
+							stepEvent.Ignored = true;
+						}
+					}
+					else if (hasMatches && IsMaskedEventId(eventId))
+					{
+						stepEvent.Ignored = true;
+					}
+				}
+			}
+		}
+	}
+}
