Horde: Rename the Horde.Build project to Horde.Server.

diff --git a/Engine/Source/Programs/Horde/Horde.Server/Issues/IIssueCollection.cs b/Engine/Source/Programs/Horde/Horde.Server/Issues/IIssueCollection.cs
new file mode 100644
--- /dev/null
+++ b/Engine/Source/Programs/Horde/Horde.Server/Issues/IIssueCollection.cs
@@ -0,0 +1,652 @@
+// Copyright Epic Games, Inc. All Rights Reserved.
+
+using System;
+using System.Collections.Generic;
+using System.Linq;
+using System.Threading.Tasks;
+using Horde.Server.Auditing;
+using Horde.Server.Jobs;
+using Horde.Server.Jobs.Graphs;
+using Horde.Server.Logs;
+using Horde.Server.Streams;
+using Horde.Server.Users;
+using Horde.Server.Utilities;
+using MongoDB.Bson;
+using MongoDB.Driver;
+
+namespace Horde.Server.Issues
+{
+	/// <summary>
+	/// Fingerprint for an issue
+	/// </summary>
+	public class NewIssueFingerprint : IIssueFingerprint, IEquatable<IIssueFingerprint>
+	{
+		/// <inheritdoc/>
+		public string Type { get; }
+
+		/// <inheritdoc/>
+		public CaseInsensitiveStringSet Keys { get; set; }
+
+		/// <inheritdoc/>
+		public CaseInsensitiveStringSet? RejectKeys { get; set; }
+
+		/// <inheritdoc/>
+		public CaseInsensitiveStringSet? Metadata { get; set; }
+
+		/// <summary>
+		/// Constructor
+		/// </summary>
+		/// <param name="type">The type of issue</param>
+		/// <param name="keys">Keys which uniquely identify this issue</param>
+		/// <param name="rejectKeys">Keys which should not match with this issue</param>
+		/// <param name="metadata">Additional metadata added by the issue handler</param>
+		public NewIssueFingerprint(string type, IEnumerable<string> keys, IEnumerable<string>? rejectKeys, IEnumerable<string>? metadata)
+		{
+			Type = type;
+			Keys = new CaseInsensitiveStringSet(keys);
+
+			if (rejectKeys != null && rejectKeys.Any())
+			{
+				RejectKeys = new CaseInsensitiveStringSet(rejectKeys);
+			}
+			if (metadata != null && metadata.Any())
+			{
+				Metadata = new CaseInsensitiveStringSet(metadata);
+			}
+		}
+
+		/// <summary>
+		/// Copy constructor
+		/// </summary>
+		/// <param name="other">The fingerprint to copy from</param>
+		public NewIssueFingerprint(IIssueFingerprint other)
+			: this(other.Type, other.Keys, other.RejectKeys, other.Metadata)
+		{
+		}
+
+		/// <summary>
+		/// Merges another fingerprint into this one
+		/// </summary>
+		/// <param name="other">The other fingerprint to merge with</param>
+		public void MergeWith(IIssueFingerprint other)
+		{
+			Keys.UnionWith(other.Keys);
+			if (other.RejectKeys != null)
+			{
+				RejectKeys ??= new CaseInsensitiveStringSet();
+				RejectKeys.UnionWith(other.RejectKeys);
+			}
+			if (other.Metadata != null)
+			{
+				Metadata ??= new CaseInsensitiveStringSet();
+				Metadata.UnionWith(other.Metadata);
+			}
+		}
+
+		/// <inheritdoc/>
+		public override bool Equals(object? other) => Equals(other as IIssueFingerprint);
+
+		/// <inheritdoc/>
+		public bool Equals(IIssueFingerprint? otherFingerprint)
+		{
+			if (!ReferenceEquals(this, otherFingerprint))
+			{
+				if (otherFingerprint == null || !Type.Equals(otherFingerprint.Type, StringComparison.Ordinal) || !Keys.SetEquals(otherFingerprint.Keys))
+				{
+					return false;
+				}
+				if (!ContentsEqual(RejectKeys, otherFingerprint.RejectKeys))
+				{
+					return false;
+				}
+				if (!ContentsEqual(Metadata, otherFingerprint.Metadata))
+				{
+					return false;
+				}
+			}
+			return true;
+		}
+
+		/// <summary>
+		/// Gets a hashcode for this fingerprint
+		/// </summary>
+		/// <returns>The hashcode value</returns>
+		public override int GetHashCode()
+		{
+			int result = StringComparer.Ordinal.GetHashCode(Type);
+			return HashCode.Combine(result, GetContentsHash(Keys), GetContentsHash(RejectKeys), GetContentsHash(Metadata));
+		}
+
+		/// <summary>
+		/// Checks if the contents of two sets are equal
+		/// </summary>
+		/// <param name="setA"></param>
+		/// <param name="setB"></param>
+		/// <returns></returns>
+		static bool ContentsEqual(CaseInsensitiveStringSet? setA, CaseInsensitiveStringSet? setB)
+		{
+			if (setA == null || setA.Count == 0)
+			{
+				return setB == null || setB.Count == 0;
+			}
+			else
+			{
+				return setB != null && setA.SetEquals(setB);
+			}
+		}
+
+		/// <summary>
+		/// Gets the hash of the contents of a case insensitive set
+		/// </summary>
+		/// <param name="set"></param>
+		/// <returns></returns>
+		static int GetContentsHash(CaseInsensitiveStringSet? set)
+		{
+			int value = 0;
+			if (set != null)
+			{
+				foreach (string element in set)
+				{
+					value = HashCode.Combine(value, StringComparer.OrdinalIgnoreCase.GetHashCode(element));
+				}
+			}
+			return value;
+		}
+
+		/// <summary>
+		/// Merges two fingerprints togetherthis issue fingerprint with another
+		/// </summary>
+		/// <param name="a">The first fingerprint</param>
+		/// <param name="b">The second fingerprint to merge with</param>
+		/// <returns>Merged fingerprint</returns>
+		public static NewIssueFingerprint Merge(IIssueFingerprint a, IIssueFingerprint b)
+		{
+			NewIssueFingerprint newFingerprint = new NewIssueFingerprint(a);
+			newFingerprint.MergeWith(b);
+			return newFingerprint;
+		}
+
+		/// <inheritdoc/>
+		public override string ToString()
+		{
+			return $"{Type}: \"{String.Join("\", \"", Keys)}\"";
+		}
+	}
+
+	/// <summary>
+	/// Information about a new span
+	/// </summary>
+	public class NewIssueSpanData
+	{
+		/// <inheritdoc cref="IIssueSpan.StreamId"/>
+		public StreamId StreamId { get; set; }
+
+		/// <inheritdoc cref="IIssueSpan.StreamName"/>
+		public string StreamName { get; set; }
+
+		/// <inheritdoc cref="IIssueSpan.TemplateRefId"/>
+		public TemplateId TemplateRefId { get; set; }
+
+		/// <inheritdoc cref="IIssueSpan.NodeName"/>
+		public string NodeName { get; set; }
+
+		/// <inheritdoc cref="IIssueSpan.Fingerprint"/>
+		public NewIssueFingerprint Fingerprint { get; set; }
+
+		/// <inheritdoc cref="IIssueSpan.LastSuccess"/>
+		public NewIssueStepData? LastSuccess { get; set; }
+
+		/// <inheritdoc cref="IIssueSpan.FirstFailure"/>
+		public NewIssueStepData FirstFailure { get; set; }
+
+		/// <inheritdoc cref="IIssueSpan.NextSuccess"/>
+		public NewIssueStepData? NextSuccess { get; set; }
+
+		/// <inheritdoc cref="IIssueSpan.Suspects"/>
+		public List<NewIssueSpanSuspectData> Suspects { get; set; } = new List<NewIssueSpanSuspectData>();
+
+		/// <summary>
+		/// Constructor
+		/// </summary>
+		public NewIssueSpanData(StreamId streamId, string streamName, TemplateId templateRefId, string nodeName, NewIssueFingerprint fingerprint, NewIssueStepData firstFailure)
+		{
+			StreamId = streamId;
+			StreamName = streamName;
+			TemplateRefId = templateRefId;
+			NodeName = nodeName;
+			Fingerprint = fingerprint;
+			FirstFailure = firstFailure;
+		}
+	}
+
+	/// <summary>
+	/// Identifies a particular changelist and job
+	/// </summary>
+	public class NewIssueStepData
+	{
+		/// <inheritdoc cref="IIssueStep.Change"/>
+		public int Change { get; set; }
+
+		/// <inheritdoc cref="IIssueStep.Severity"/>
+		public IssueSeverity Severity { get; set; }
+
+		/// <inheritdoc cref="IIssueStep.JobName"/>
+		public string JobName { get; set; }
+
+		/// <inheritdoc cref="IIssueStep.JobId"/>
+		public JobId JobId { get; set; }
+
+		/// <inheritdoc cref="IIssueStep.BatchId"/>
+		public SubResourceId BatchId { get; set; }
+
+		/// <inheritdoc cref="IIssueStep.StepId"/>
+		public SubResourceId StepId { get; set; }
+
+		/// <inheritdoc cref="IIssueStep.StepTime"/>
+		public DateTime StepTime { get; set; }
+
+		/// <inheritdoc cref="IIssueStep.LogId"/>
+		public LogId? LogId { get; set; }
+
+		/// <inheritdoc cref="IIssueStep.Annotations"/>
+		public NodeAnnotations? Annotations { get; set; }
+
+		/// <inheritdoc cref="IIssueStep.PromoteByDefault"/>
+		public bool PromoteByDefault { get; set; }
+
+		/// <summary>
+		/// Constructor
+		/// </summary>
+		/// <param name="change">The changelist number for this job</param>
+		/// <param name="severity">Severity of the issue in this step</param>
+		/// <param name="jobName">The job name</param>
+		/// <param name="jobId">The unique job id</param>
+		/// <param name="batchId">The batch id</param>
+		/// <param name="stepId">The step id</param>
+		/// <param name="stepTime">Time that the step started</param>
+		/// <param name="logId">Unique id of the log file for this step</param>
+		/// <param name="annotations">Annotations for this step</param>
+		/// <param name="promoted">Whether this step is promoted</param>
+		public NewIssueStepData(int change, IssueSeverity severity, string jobName, JobId jobId, SubResourceId batchId, SubResourceId stepId, DateTime stepTime, LogId? logId, IReadOnlyNodeAnnotations? annotations, bool promoted)
+		{
+			Change = change;
+			Severity = severity;
+			JobName = jobName;
+			JobId = jobId;
+			BatchId = batchId;
+			StepId = stepId;
+			StepTime = stepTime;
+			LogId = logId;
+			if (annotations != null)
+			{
+				Annotations = new NodeAnnotations(annotations);
+			}
+			PromoteByDefault = promoted;
+		}
+
+		/// <summary>
+		/// Constructor
+		/// </summary>
+		/// <param name="job">The job being built</param>
+		/// <param name="batch">Batch of the job for the step</param>
+		/// <param name="step">The step being built</param>
+		/// <param name="severity">Severity of the issue in this step</param>
+		/// <param name="annotations">Annotations for this step</param>
+		/// <param name="promoted">Whether this step is promoted</param>
+		public NewIssueStepData(IJob job, IJobStepBatch batch, IJobStep step, IssueSeverity severity, IReadOnlyNodeAnnotations? annotations, bool promoted)
+			: this(job.Change, severity, job.Name, job.Id, batch.Id, step.Id, step.StartTimeUtc ?? default, step.LogId, annotations, promoted)
+		{
+		}
+
+		/// <summary>
+		/// Constructor for sentinel steps
+		/// </summary>
+		/// <param name="jobStepRef">The jobstep reference</param>
+		public NewIssueStepData(IJobStepRef jobStepRef)
+			: this(jobStepRef.Change, IssueSeverity.Unspecified, jobStepRef.JobName, jobStepRef.Id.JobId, jobStepRef.Id.BatchId, jobStepRef.Id.StepId, jobStepRef.StartTimeUtc, jobStepRef.LogId, null, false)
+		{
+		}
+	}
+
+	/// <summary>
+	/// Information about a suspect changelist that may have caused an issue
+	/// </summary>
+	public class NewIssueSuspectData
+	{
+		/// <inheritdoc cref="IIssueSuspect.AuthorId"/>
+		public UserId AuthorId { get; set; }
+
+		/// <inheritdoc cref="IIssueSuspect.Change"/>
+		public int Change { get; set; }
+
+		/// <summary>
+		/// Constructor
+		/// </summary>
+		/// <param name="authorId">Author of the change</param>
+		/// <param name="change">The changelist number</param>
+		public NewIssueSuspectData(UserId authorId, int change)
+		{
+			AuthorId = authorId;
+			Change = change;
+		}
+	}
+
+	/// <summary>
+	/// Information about a suspect changelist that may have caused an issue
+	/// </summary>
+	public class NewIssueSpanSuspectData
+	{
+		/// <inheritdoc cref="IIssueSpanSuspect.Change"/>
+		public int Change { get; set; }
+
+		/// <inheritdoc cref="IIssueSpanSuspect.AuthorId"/>
+		public UserId AuthorId { get; set; }
+
+		/// <inheritdoc cref="IIssueSpanSuspect.OriginatingChange"/>
+		public int? OriginatingChange { get; set; }
+
+		/// <summary>
+		/// Constructor
+		/// </summary>
+		/// <param name="change">The changelist number</param>
+		/// <param name="authorId">Author of the change</param>
+		public NewIssueSpanSuspectData(int change, UserId authorId)
+		{
+			Change = change;
+			AuthorId = authorId;
+		}
+	}
+
+	/// <summary>
+	/// Information about a stream containing an issue
+	/// </summary>
+	public class NewIssueStream : IIssueStream
+	{
+		/// <inheritdoc cref="IIssueStream.StreamId"/>
+		public StreamId StreamId { get; set; }
+
+		/// <inheritdoc cref="IIssueStream.MergeOrigin"/>
+		public bool? MergeOrigin { get; set; }
+
+		/// <inheritdoc cref="IIssueStream.ContainsFix"/>
+		public bool? ContainsFix { get; set; }
+
+		/// <inheritdoc cref="IIssueStream.FixFailed"/>
+		public bool? FixFailed { get; set; }
+
+		/// <summary>
+		/// Constructor
+		/// </summary>
+		/// <param name="streamId"></param>
+		public NewIssueStream(StreamId streamId)
+		{
+			StreamId = streamId;
+		}
+	}
+
+	/// <summary>
+	/// Interface for an issue collection
+	/// </summary>
+	public interface IIssueCollection
+	{
+		/// <summary>
+		/// Attempts to enters a critical section over the issue collection. Used for performing operations that require coordination of several documents (eg. attaching spans to issues).
+		/// </summary>
+		/// <returns>The disposable critical section</returns>
+		Task<IAsyncDisposable> EnterCriticalSectionAsync();
+
+		#region Issues
+
+		/// <summary>
+		/// Creates a new issue
+		/// </summary>
+		/// <param name="summary">Summary text for the issue</param>
+		/// <returns>The new issue instance</returns>
+		Task<IIssue> AddIssueAsync(string summary);
+
+		/// <summary>
+		/// Retrieves and issue by id
+		/// </summary>
+		/// <param name="issueId">Unique id of the issue</param>
+		/// <returns>The issue matching the given id, or null</returns>
+		Task<IIssue?> GetIssueAsync(int issueId);
+
+		/// <summary>
+		/// Finds the suspects for an issue
+		/// </summary>
+		/// <param name="issueId">The issue to retrieve suspects for</param>
+		/// <returns>List of suspects</returns>
+		Task<List<IIssueSuspect>> FindSuspectsAsync(int issueId);
+
+		/// <summary>
+		/// Searches for open issues
+		/// </summary>
+		/// <param name="ids">Set of issue ids to find</param>
+		/// <param name="ownerId">The user to find issues for</param>
+		/// <param name="streamId">The stream affected by the issue</param>
+		/// <param name="minChange">Minimum changelist affected by the issue</param>
+		/// <param name="maxChange">Maximum changelist affected by the issue</param>
+		/// <param name="resolved">Include issues that are now resolved</param>
+		/// <param name="promoted">Include only promoted issues</param>
+		/// <param name="index">Index within the results to return</param>
+		/// <param name="count">Number of results</param>
+		/// <returns>List of streams open in the given stream at the given changelist</returns>
+		Task<List<IIssue>> FindIssuesAsync(IEnumerable<int>? ids = null, UserId? ownerId = null, StreamId? streamId = null, int? minChange = null, int? maxChange = null, bool? resolved = null, bool? promoted = null, int? index = null, int? count = null);
+
+		/// <summary>
+		/// Searches for open issues
+		/// </summary>
+		/// <param name="changes">List of suspect changes</param>
+		/// <returns>List of issues that are affected by the given changes</returns>
+		Task<List<IIssue>> FindIssuesForChangesAsync(List<int> changes);
+
+		/// <summary>
+		/// Try to update the state of an issue
+		/// </summary>
+		/// <param name="issue">The issue to update</param>
+		/// <param name="initiatedByUserId">User that is making the change</param>
+		/// <param name="newSeverity">New severity for the issue</param>
+		/// <param name="newSummary">New summary for the issue</param>
+		/// <param name="newUserSummary">New user summary for the issue</param>
+		/// <param name="newDescription">New description for the issue</param>
+		/// <param name="newPromoted">New promoted state of the issue</param>
+		/// <param name="newOwnerId">New owner of the issue</param>
+		/// <param name="newNominatedById">Person that nominated the new owner</param>
+		/// <param name="newAcknowledged">Whether the issue has been acknowledged</param>
+		/// <param name="newDeclinedById">Name of a user that has declined the issue</param>
+		/// <param name="newFixChange">Fix changelist for the issue. Pass 0 to clear the fix changelist, -1 for systemic issue.</param>
+		/// <param name="newResolvedById">User that resolved the issue (may be ObjectId.Empty to clear)</param>
+		/// <param name="newExcludeSpanIds">List of span ids to exclude from this issue</param>
+		/// <param name="newLastSeenAt"></param>
+		/// <param name="newExternalIssueKey">Issue key for external issue tracking</param>
+		/// <param name="newQuarantinedById">The user that quarantined the issue</param>
+		/// <param name="newForceClosedById">The user that force closed the issue</param>
+		/// <param name="newWorkflowThreadUrl">The workflow thread url associated with the issue</param>
+		/// <returns>True if the issue was updated</returns>
+		Task<IIssue?> TryUpdateIssueAsync(IIssue issue, UserId? initiatedByUserId, IssueSeverity? newSeverity = null, string? newSummary = null, string? newUserSummary = null, string? newDescription = null, bool? newPromoted = null, UserId? newOwnerId = null, UserId? newNominatedById = null, bool? newAcknowledged = null, UserId? newDeclinedById = null, int? newFixChange = null, UserId? newResolvedById = null, List<ObjectId>? newExcludeSpanIds = null, DateTime? newLastSeenAt = null, string? newExternalIssueKey = null, UserId? newQuarantinedById = null, UserId? newForceClosedById = null, Uri? newWorkflowThreadUrl = null);
+
+		/// <summary>
+		/// Updates derived data for an issue (ie. data computed from the spans attached to it). Also clears the issue's 'modified' state.
+		/// </summary>
+		/// <param name="issue">Issue to update</param>
+		/// <param name="newSummary">New summary for the issue</param>
+		/// <param name="newSeverity">New severity for the issue</param>
+		/// <param name="newFingerprints">New fingerprints for the issue</param>
+		/// <param name="newStreams">New streams for the issue</param>
+		/// <param name="newSuspects">New suspects for the issue</param>
+		/// <param name="newResolvedAt">Time for the last resolved change</param>
+		/// <param name="newVerifiedAt">Time that the issue was resolved</param>
+		/// <param name="newLastSeenAt">Last time the issue was seen</param>
+		/// <returns>Updated issue, or null if the issue is modified in the interim</returns>
+		Task<IIssue?> TryUpdateIssueDerivedDataAsync(IIssue issue, string newSummary, IssueSeverity newSeverity, List<NewIssueFingerprint> newFingerprints, List<NewIssueStream> newStreams, List<NewIssueSuspectData> newSuspects, DateTime? newResolvedAt, DateTime? newVerifiedAt, DateTime newLastSeenAt);
+
+		#endregion
+
+		#region Spans
+
+		/// <summary>
+		/// Creates a new span from the given failure
+		/// </summary>
+		/// <param name="issueId">The issue that the span belongs to</param>
+		/// <param name="newSpan">Information about the new span</param>
+		/// <returns>New span, or null if the sequence token is not valid</returns>
+		Task<IIssueSpan> AddSpanAsync(int issueId, NewIssueSpanData newSpan);
+
+		/// <summary>
+		/// Gets a particular span
+		/// </summary>
+		/// <param name="spanId">Unique id of the span</param>
+		/// <returns>New span, or null if the sequence token is not valid</returns>
+		Task<IIssueSpan?> GetSpanAsync(ObjectId spanId);
+
+		/// <summary>
+		/// Updates the given span. Note that data in the span's issue may be derived from this, and the issue should be updated afterwards.
+		/// </summary>
+		/// <param name="span">Span to update</param>
+		/// <param name="newLastSuccess">New last successful step</param>
+		/// <param name="newFailure">New failed step</param>
+		/// <param name="newNextSuccess">New next successful step</param>
+		/// <param name="newSuspects">New suspects for the span</param>
+		/// <param name="newIssueId">The new issue id for this span</param>
+		/// <returns>The updated span, or null on failure</returns>
+		Task<IIssueSpan?> TryUpdateSpanAsync(IIssueSpan span, NewIssueStepData? newLastSuccess = null, NewIssueStepData? newFailure = null, NewIssueStepData? newNextSuccess = null, List<NewIssueSpanSuspectData>? newSuspects = null, int? newIssueId = null);
+
+		/// <summary>
+		/// Gets all the spans for a particular issue
+		/// </summary>
+		/// <param name="issueId">Issue id</param>
+		/// <returns>List of spans</returns>
+		Task<List<IIssueSpan>> FindSpansAsync(int issueId);
+
+		/// <summary>
+		/// Retrieves multiple spans
+		/// </summary>
+		/// <param name="spanIds">The span ids</param>
+		/// <returns>List of spans</returns>
+		Task<List<IIssueSpan>> FindSpansAsync(IEnumerable<ObjectId> spanIds);
+
+		/// <summary>
+		/// Finds the open issues for a given stream
+		/// </summary>
+		/// <param name="streamId">The stream id</param>
+		/// <param name="templateId">The template id</param>
+		/// <param name="name">Name of the node</param>
+		/// <param name="change">Changelist number to query</param>
+		/// <returns>List of open issues</returns>
+		Task<List<IIssueSpan>> FindOpenSpansAsync(StreamId streamId, TemplateId templateId, string name, int change);
+
+		/// <summary>
+		/// Searches for open issues
+		/// </summary>
+		/// <param name="ids">Set of issue ids to find</param>
+		/// <param name="issueIds">The issue ids to retrieve spans for</param>
+		/// <param name="streamId">The stream affected by the issue</param>
+		/// <param name="minChange">Minimum changelist affected by the issue</param>
+		/// <param name="maxChange">Maximum changelist affected by the issue</param>
+		/// <param name="resolved">Include issues that are now resolved</param>
+		/// <param name="index">Index within the results to return</param>
+		/// <param name="count">Number of results</param>
+		/// <returns>List of streams open in the given stream at the given changelist</returns>
+		Task<List<IIssueSpan>> FindSpansAsync(IEnumerable<ObjectId>? ids = null, IEnumerable<int>? issueIds = null, StreamId? streamId = null, int? minChange = null, int? maxChange = null, bool? resolved = null, int? index = null, int? count = null);
+
+		#endregion
+
+		#region Steps
+
+		/// <summary>
+		/// Adds a new step
+		/// </summary>
+		/// <param name="spanId">Initial span for the step</param>
+		/// <param name="newStep">Information about the new step</param>
+		/// <returns>New step object</returns>
+		Task<IIssueStep> AddStepAsync(ObjectId spanId, NewIssueStepData newStep);
+
+		/// <summary>
+		/// Find steps for the given spans
+		/// </summary>
+		/// <param name="spanIds">Span ids</param>
+		/// <returns>List of steps</returns>
+		Task<List<IIssueStep>> FindStepsAsync(IEnumerable<ObjectId> spanIds);
+
+		/// <summary>
+		/// Find steps for the given spans
+		/// </summary>
+		/// <param name="jobId">The job id</param>
+		/// <param name="batchId">The batch id</param>
+		/// <param name="stepId">The step id</param>
+		/// <returns>List of steps</returns>
+		Task<List<IIssueStep>> FindStepsAsync(JobId jobId, SubResourceId? batchId, SubResourceId? stepId);
+
+		#endregion
+
+		/// <summary>
+		/// Gets the logger for a particular issue
+		/// </summary>
+		/// <param name="issueId">The issue id</param>
+		/// <returns>Logger for this issue</returns>
+		IAuditLogChannel<int> GetLogger(int issueId);
+	}
+
+	/// <summary>
+	/// Extension methods for issues
+	/// </summary>
+	public static class IssueCollectionExtensions
+	{
+		/// <summary>
+		/// Find suspects for the given issue
+		/// </summary>
+		/// <param name="issueCollection"></param>
+		/// <param name="issue"></param>
+		/// <returns></returns>
+		public static Task<List<IIssueSuspect>> FindSuspectsAsync(this IIssueCollection issueCollection, IIssue issue)
+		{
+			return issueCollection.FindSuspectsAsync(issue.Id);
+		}
+
+		/// <summary>
+		/// Find steps for the given spans
+		/// </summary>
+		/// <param name="issueCollection"></param>
+		/// <param name="spanId">Span ids</param>
+		/// <returns>List of steps</returns>
+		public static Task<List<IIssueStep>> FindStepsAsync(this IIssueCollection issueCollection, ObjectId spanId)
+		{
+			return issueCollection.FindStepsAsync(new[] { spanId });
+		}
+
+		/// <summary>
+		/// Searches for open issues affecting a job
+		/// </summary>
+		/// <param name="issueCollection"></param>
+		/// <param name="job"></param>
+		/// <param name="graph"></param>
+		/// <param name="stepId"></param>
+		/// <param name="batchId"></param>
+		/// <param name="labelIdx"></param>
+		/// <param name="ownerId"></param>
+		/// <param name="resolved">Whether to include results that are resolved</param>
+		/// <param name="promoted">Whether to filter by promoted issues</param>
+		/// <param name="index">Index within the results to return</param>
+		/// <param name="count">Number of results</param>
+		/// <returns></returns>
+		public static async Task<List<IIssue>> FindIssuesForJobAsync(this IIssueCollection issueCollection, IJob job, IGraph graph, SubResourceId? stepId = null, SubResourceId? batchId = null, int? labelIdx = null, UserId? ownerId = null, bool? resolved = null, bool? promoted = null, int? index = null, int? count = null)
+		{
+			List<IIssueStep> steps = await issueCollection.FindStepsAsync(job.Id, batchId, stepId);
+			List<IIssueSpan> spans = await issueCollection.FindSpansAsync(steps.Select(x => x.SpanId));
+
+			if (labelIdx != null)
+			{
+				HashSet<string> nodeNames = new HashSet<string>(job.GetNodesForLabel(graph, labelIdx.Value).Select(x => graph.GetNode(x).Name));
+				spans.RemoveAll(x => !nodeNames.Contains(x.NodeName));
+			}
+
+			List<int> issueIds = new List<int>(spans.Select(x => x.IssueId));
+			if (issueIds.Count == 0)
+			{
+				return new List<IIssue>();
+			}
+
+			return await issueCollection.FindIssuesAsync(ids: issueIds, ownerId: ownerId, resolved: resolved, promoted: promoted, index: index, count: count);
+		}
+	}
+}
diff --git a/Engine/Source/Programs/Horde/Horde.Server/Issues/IssueCollection.cs b/Engine/Source/Programs/Horde/Horde.Server/Issues/IssueCollection.cs
new file mode 100644
--- /dev/null
+++ b/Engine/Source/Programs/Horde/Horde.Server/Issues/IssueCollection.cs
@@ -0,0 +1,1491 @@
+// Copyright Epic Games, Inc. All Rights Reserved.
+
+using System;
+using System.Collections.Generic;
+using System.Diagnostics;
+using System.Linq;
+using System.Threading.Tasks;
+using EpicGames.Redis.Utility;
+using Horde.Server.Auditing;
+using Horde.Server.Jobs;
+using Horde.Server.Jobs.Graphs;
+using Horde.Server.Logs;
+using Horde.Server.Server;
+using Horde.Server.Streams;
+using Horde.Server.Telemetry;
+using Horde.Server.Users;
+using Horde.Server.Utilities;
+using Microsoft.Extensions.Caching.Memory;
+using Microsoft.Extensions.Logging;
+using MongoDB.Bson;
+using MongoDB.Bson.Serialization.Attributes;
+using MongoDB.Driver;
+
+namespace Horde.Server.Issues
+{
+	class IssueCollection : IIssueCollection
+	{
+		[SingletonDocument("issue-ledger", "5e4c226440ce25fa3207a9af")]
+		class IssueLedger : SingletonBase
+		{
+			public int NextId { get; set; }
+		}
+
+		class Issue : IIssue
+		{
+			[BsonId]
+			public int Id { get; set; }
+
+			public string Summary { get; set; }
+
+			[BsonIgnoreIfNull]
+			public string? UserSummary { get; set; }
+
+			[BsonIgnoreIfNull]
+			public string? Description { get; set; }
+
+			[BsonIgnoreIfNull]
+			public IssueFingerprint? Fingerprint { get; set; }
+			public List<IssueFingerprint>? Fingerprints { get; set; }
+
+			public IssueSeverity Severity { get; set; }
+
+			[BsonElement("PromotedV2")]
+			public bool Promoted { get; set; }
+
+			[BsonIgnoreIfNull]
+			public bool? ManuallyPromoted { get; set; }
+
+			[BsonIgnoreIfNull, BsonElement("Promoted")]
+			public bool? ManuallyPromotedDeprecated { get; set; }
+
+			[BsonIgnoreIfNull]
+			public UserId? OwnerId { get; set; }
+
+			[BsonIgnoreIfNull]
+			public UserId? DefaultOwnerId { get; set; }
+
+			[BsonIgnoreIfNull]
+			public UserId? NominatedById { get; set; }
+
+			public DateTime CreatedAt { get; set; }
+
+			[BsonIgnoreIfNull]
+			public DateTime? NominatedAt { get; set; }
+
+			[BsonIgnoreIfNull]
+			public DateTime? AcknowledgedAt { get; set; }
+
+			[BsonIgnoreIfNull]
+			public DateTime? ResolvedAt { get; set; }
+
+			[BsonIgnoreIfNull]
+			public UserId? ResolvedById { get; set; }
+
+			[BsonIgnoreIfNull]
+			public DateTime? VerifiedAt { get; set; }
+
+			public DateTime LastSeenAt { get; set; }
+
+			[BsonIgnoreIfNull]
+			public int? FixChange { get; set; }
+
+			public List<IssueStream> Streams { get; set; } = new List<IssueStream>();
+
+			public int MinSuspectChange { get; set; }
+			public int MaxSuspectChange { get; set; }
+
+			[BsonElement("Suspects"), BsonIgnoreIfNull]
+			public List<IssueSuspect>? SuspectsDeprecated { get; set; }
+
+			[BsonIgnoreIfNull]
+			public List<ObjectId>? ExcludeSpans { get; set; }
+
+			public int UpdateIndex { get; set; }
+
+			IReadOnlyList<IIssueFingerprint> IIssue.Fingerprints => Fingerprints ?? ((Fingerprint == null)? new List<IssueFingerprint>() : new List<IssueFingerprint> { Fingerprint });
+			UserId? IIssue.OwnerId => OwnerId ?? DefaultOwnerId ?? GetDefaultOwnerId();
+			IReadOnlyList<IIssueStream> IIssue.Streams => Streams;
+			DateTime IIssue.LastSeenAt => (LastSeenAt == default) ? DateTime.UtcNow : LastSeenAt;
+
+			[BsonIgnoreIfNull]
+			public string? ExternalIssueKey { get; set; }
+
+			[BsonIgnoreIfNull]
+			public UserId? QuarantinedByUserId { get; set; }
+
+			[BsonIgnoreIfNull]
+			public DateTime? QuarantineTimeUtc { get; set; }
+
+			[BsonIgnoreIfNull]
+			public UserId? ForceClosedByUserId { get; set; }
+
+			[BsonIgnoreIfNull]
+			public Uri? WorkflowThreadUrl { get; set; }
+
+			[BsonConstructor]
+			private Issue()
+			{
+				Summary = String.Empty;
+				Fingerprint = null!;
+			}
+
+			public Issue(int id, string summary)
+			{
+				Id = id;
+				Summary = summary;
+				CreatedAt = DateTime.UtcNow;
+				LastSeenAt = DateTime.UtcNow;
+			}
+
+			UserId? GetDefaultOwnerId()
+			{
+				if (SuspectsDeprecated != null && SuspectsDeprecated.Count > 0)
+				{
+					UserId possibleOwner = SuspectsDeprecated[0].AuthorId;
+					if (SuspectsDeprecated.All(x => x.AuthorId == possibleOwner) && SuspectsDeprecated.Any(x => x.DeclinedAt == null))
+					{
+						return SuspectsDeprecated[0].AuthorId;
+					}
+				}
+				return null;
+			}
+
+			public string FingerprintsDesc
+			{
+				get
+				{
+					if (Fingerprints == null || Fingerprints.Count == 0)
+					{
+						return String.Empty;
+					}
+
+					return String.Join(", ", Fingerprints.Select(x => {
+					   return $"(Type: {x.Type} / Keys: {String.Join(", ", x.Keys)} / RejectKeys: {String.Join(", ", x.RejectKeys ?? new CaseInsensitiveStringSet(new string[] {"No Reject Keys"}))})";
+				   }));
+				}
+			}
+		}
+
+		class IssueStream : IIssueStream
+		{
+			public StreamId StreamId { get; set; }
+			public bool? MergeOrigin { get; set; }
+			public bool? ContainsFix { get; set; }
+			public bool? FixFailed { get; set; }
+
+			public IssueStream()
+			{
+			}
+
+			public IssueStream(IIssueStream other)
+			{
+				StreamId = other.StreamId;
+				MergeOrigin = other.MergeOrigin;
+				ContainsFix = other.ContainsFix;
+				FixFailed = other.FixFailed;
+			}
+		}
+
+		class IssueSuspect : IIssueSuspect
+		{
+			public ObjectId Id { get; set; }
+			public int IssueId { get; set; }
+			public UserId AuthorId { get; set; }
+			public int Change { get; set; }
+			public DateTime? DeclinedAt { get; set; }
+			public DateTime? ResolvedAt { get; set; } // Degenerate
+
+			private IssueSuspect()
+			{
+			}
+
+			public IssueSuspect(int issueId, NewIssueSuspectData newSuspect, DateTime? resolvedAt)
+			{
+				Id = ObjectId.GenerateNewId();
+				IssueId = issueId;
+				AuthorId = newSuspect.AuthorId;
+				Change = newSuspect.Change;
+				ResolvedAt = resolvedAt;
+			}
+
+			public IssueSuspect(int issueId, IIssueSpanSuspect suspect)
+				: this(issueId, suspect.AuthorId, suspect.OriginatingChange ?? suspect.Change, null, null)
+			{
+			}
+
+			public IssueSuspect(int issueId, UserId authorId, int change, DateTime? declinedAt, DateTime? resolvedAt)
+			{
+				Id = ObjectId.GenerateNewId();
+				IssueId = issueId;
+				AuthorId = authorId;
+				Change = change;
+				DeclinedAt = declinedAt;
+				ResolvedAt = resolvedAt;
+			}
+		}
+
+		class IssueFingerprint : IIssueFingerprint
+		{
+			public string Type { get; set; }
+			public CaseInsensitiveStringSet Keys { get; set; }
+			public CaseInsensitiveStringSet? RejectKeys { get; set; }
+			public CaseInsensitiveStringSet? Metadata { get; set; }
+
+			[BsonConstructor]
+			private IssueFingerprint()
+			{
+				Type = String.Empty;
+				Keys = new CaseInsensitiveStringSet();
+			}
+
+			public IssueFingerprint(IIssueFingerprint fingerprint)
+			{
+				Type = fingerprint.Type;
+				Keys = fingerprint.Keys;
+				RejectKeys = fingerprint.RejectKeys;
+				Metadata = fingerprint.Metadata;
+			}
+		}
+
+		class IssueSpan : IIssueSpan
+		{
+			public ObjectId Id { get; set; }
+
+			[BsonRequired]
+			public StreamId StreamId { get; set; }
+
+			[BsonRequired]
+			public string StreamName { get; set; }
+
+			[BsonRequired]
+			public TemplateId TemplateRefId { get; set; }
+
+			[BsonRequired]
+			public string NodeName { get; set; }
+			public DateTime? ResolvedAt { get; set; } // Propagated from the owning issue
+
+			[BsonRequired]
+			public IssueFingerprint Fingerprint { get; set; }
+
+			public int MinChange { get; set; }
+			public int MaxChange { get; set; } = Int32.MaxValue;
+
+			public IssueStep? LastSuccess { get; set; }
+
+			[BsonRequired]
+			public IssueStep FirstFailure { get; set; }
+
+			[BsonRequired]
+			public IssueStep LastFailure { get; set; }
+
+			public IssueStep? NextSuccess { get; set; }
+
+			public bool? PromoteByDefault { get; set; }
+
+			[BsonElement("NotifySuspects"), BsonIgnoreIfDefault(false)]
+			public bool NotifySuspectsDeprecated { get; set; }
+
+			bool IIssueSpan.PromoteByDefault => PromoteByDefault ?? NotifySuspectsDeprecated;
+
+			public List<IssueSpanSuspect> Suspects { get; set; }
+			public int IssueId { get; set; }
+			public int UpdateIndex { get; set; }
+
+			IIssueStep? IIssueSpan.LastSuccess => LastSuccess;
+			IIssueStep IIssueSpan.FirstFailure => FirstFailure;
+			IIssueStep IIssueSpan.LastFailure => LastFailure;
+			IIssueStep? IIssueSpan.NextSuccess => NextSuccess;
+			IReadOnlyList<IIssueSpanSuspect> IIssueSpan.Suspects => Suspects;
+			IIssueFingerprint IIssueSpan.Fingerprint => Fingerprint;
+
+			private IssueSpan()
+			{
+				StreamName = null!;
+				NodeName = null!;
+				Fingerprint = null!;
+				FirstFailure = null!;
+				LastFailure = null!;
+				Suspects = new List<IssueSpanSuspect>();
+			}
+
+			public IssueSpan(int issueId, NewIssueSpanData newSpan)
+			{
+				Id = ObjectId.GenerateNewId();
+				StreamId = newSpan.StreamId;
+				StreamName = newSpan.StreamName;
+				TemplateRefId = newSpan.TemplateRefId;
+				NodeName = newSpan.NodeName;
+				Fingerprint = new IssueFingerprint(newSpan.Fingerprint);
+				if (newSpan.LastSuccess != null)
+				{
+					MinChange = newSpan.LastSuccess.Change;
+					LastSuccess = new IssueStep(Id, newSpan.LastSuccess);
+				}
+				FirstFailure = new IssueStep(Id, newSpan.FirstFailure);
+				LastFailure = new IssueStep(Id, newSpan.FirstFailure);
+				if (newSpan.NextSuccess != null)
+				{
+					MaxChange = newSpan.NextSuccess.Change;
+					NextSuccess = new IssueStep(Id, newSpan.NextSuccess);
+				}
+				PromoteByDefault = newSpan.FirstFailure.PromoteByDefault;
+				Suspects = newSpan.Suspects.ConvertAll(x => new IssueSpanSuspect(x));
+				IssueId = issueId;
+			}
+		}
+
+		class IssueSpanSuspect : IIssueSpanSuspect
+		{
+			public int Change { get; set; }
+			public UserId AuthorId { get; set; }
+			public int? OriginatingChange { get; set; }
+
+			[BsonConstructor]
+			private IssueSpanSuspect()
+			{
+			}
+
+			public IssueSpanSuspect(NewIssueSpanSuspectData newSuspectData)
+			{
+				Change = newSuspectData.Change;
+				AuthorId = newSuspectData.AuthorId;
+				OriginatingChange = newSuspectData.OriginatingChange;
+			}
+		}
+
+		class IssueStep : IIssueStep
+		{
+			public ObjectId Id { get; set; }
+			public ObjectId SpanId { get; set; }
+
+			public int Change { get; set; }
+			public IssueSeverity Severity { get; set; }
+
+			[BsonRequired]
+			public string JobName { get; set; }
+
+			[BsonRequired]
+			public JobId JobId { get; set; }
+
+			[BsonRequired]
+			public SubResourceId BatchId { get; set; }
+
+			[BsonRequired]
+			public SubResourceId StepId { get; set; }
+
+			public DateTime StepTime { get; set; }
+
+			public LogId? LogId { get; set; }
+
+			[BsonIgnoreIfNull]
+			public NodeAnnotations? Annotations { get; set; }
+
+			public bool? PromoteByDefault { get; set; }
+
+			[BsonElement("NotifySuspects"), BsonIgnoreIfDefault(false)]
+			public bool NotifySuspectsDeprecated { get; set; }
+
+			IReadOnlyNodeAnnotations IIssueStep.Annotations => Annotations ?? NodeAnnotations.Empty;
+			bool IIssueStep.PromoteByDefault => PromoteByDefault ?? NotifySuspectsDeprecated;
+
+			[BsonConstructor]
+			private IssueStep()
+			{
+				JobName = null!;
+			}
+
+			public IssueStep(ObjectId spanId, NewIssueStepData stepData)
+			{
+				Id = ObjectId.GenerateNewId();
+				SpanId = spanId;
+				Change = stepData.Change;
+				Severity = stepData.Severity;
+				JobName = stepData.JobName;
+				JobId = stepData.JobId;
+				BatchId = stepData.BatchId;
+				StepId = stepData.StepId;
+				StepTime = stepData.StepTime;
+				LogId = stepData.LogId;
+				Annotations = stepData.Annotations;
+				PromoteByDefault = stepData.PromoteByDefault;
+			}
+		}
+
+		readonly RedisService _redisService;
+		readonly IUserCollection _userCollection;
+		readonly ISingletonDocument<IssueLedger> _ledgerSingleton;
+		readonly IMongoCollection<Issue> _issues;
+		readonly IMongoCollection<IssueSpan> _issueSpans;
+		readonly IMongoCollection<IssueStep> _issueSteps;
+		readonly IMongoCollection<IssueSuspect> _issueSuspects;
+		readonly IAuditLog<int> _auditLog;
+		readonly ITelemetrySink _telemetrySink;
+		readonly ILogger _logger;
+
+		public IssueCollection(MongoService mongoService, RedisService redisService, IUserCollection userCollection, IAuditLogFactory<int> auditLogFactory, ITelemetrySink telemetrySink, ILogger<IssueCollection> logger)
+		{
+			_redisService = redisService;
+			_userCollection = userCollection;
+			_telemetrySink = telemetrySink;
+			_logger = logger;
+
+			_ledgerSingleton = new SingletonDocument<IssueLedger>(mongoService);
+
+			List<MongoIndex<Issue>> issueIndexes = new List<MongoIndex<Issue>>();
+			issueIndexes.Add(keys => keys.Ascending(x => x.ResolvedAt));
+			issueIndexes.Add(keys => keys.Ascending(x => x.VerifiedAt));
+			_issues = mongoService.GetCollection<Issue>("IssuesV2");
+
+			List<MongoIndex<IssueSpan>> issueSpanIndexes = new List<MongoIndex<IssueSpan>>();
+			issueSpanIndexes.Add(keys => keys.Ascending(x => x.IssueId));
+			issueSpanIndexes.Add(keys => keys.Ascending(x => x.StreamId).Ascending(x => x.MinChange).Ascending(x => x.MaxChange));
+			issueSpanIndexes.Add("StreamChanges", keys => keys.Ascending(x => x.StreamId).Ascending(x => x.TemplateRefId).Ascending(x => x.NodeName).Ascending(x => x.MinChange).Ascending(x => x.MaxChange));
+			_issueSpans = mongoService.GetCollection<IssueSpan>("IssuesV2.Spans", issueSpanIndexes);
+
+			List<MongoIndex<IssueStep>> issueStepIndexes = new List<MongoIndex<IssueStep>>();
+			issueStepIndexes.Add(keys => keys.Ascending(x => x.SpanId));
+			issueStepIndexes.Add(keys => keys.Ascending(x => x.JobId).Ascending(x => x.BatchId).Ascending(x => x.StepId));
+			_issueSteps = mongoService.GetCollection<IssueStep>("IssuesV2.Steps", issueStepIndexes);
+
+			List<MongoIndex<IssueSuspect>> issueSuspectIndexes = new List<MongoIndex<IssueSuspect>>();
+			issueSuspectIndexes.Add(keys => keys.Ascending(x => x.Change));
+			issueSuspectIndexes.Add(keys => keys.Ascending(x => x.AuthorId).Ascending(x => x.ResolvedAt));
+			issueSuspectIndexes.Add(keys => keys.Ascending(x => x.IssueId).Ascending(x => x.Change), unique: true);
+			_issueSuspects = mongoService.GetCollection<IssueSuspect>("IssuesV2.Suspects", issueSuspectIndexes);
+
+			_auditLog = auditLogFactory.Create("IssuesV2.History", "IssueId");
+		}
+
+		/// <inheritdoc/>
+		public async Task<IAsyncDisposable> EnterCriticalSectionAsync()
+		{
+			Stopwatch timer = Stopwatch.StartNew();
+			TimeSpan nextNotifyTime = TimeSpan.FromSeconds(2.0);
+
+			RedisLock issueLock = new (_redisService.GetDatabase(), "issues/lock");
+			while (!await issueLock.AcquireAsync(TimeSpan.FromMinutes(1)))
+			{
+				if (timer.Elapsed > nextNotifyTime)
+				{
+					_logger.LogWarning("Waiting on lock over issue collection for {TimeSpan}", timer.Elapsed);
+					nextNotifyTime *= 2;
+				}
+				await Task.Delay(TimeSpan.FromMilliseconds(100));
+			}
+			return issueLock;
+		}
+
+		async Task<Issue?> TryUpdateIssueAsync(IIssue issue, UpdateDefinition<Issue> update)
+		{
+			Issue issueDocument = (Issue)issue;
+
+			int prevUpdateIndex = issueDocument.UpdateIndex;
+			update = update.Set(x => x.UpdateIndex, prevUpdateIndex + 1);
+
+			FindOneAndUpdateOptions<Issue, Issue> options = new FindOneAndUpdateOptions<Issue, Issue> { ReturnDocument = ReturnDocument.After };
+
+			Issue? newIssue = await _issues.FindOneAndUpdateAsync<Issue>(x => x.Id == issueDocument.Id && x.UpdateIndex == prevUpdateIndex, update, options);
+			if (newIssue != null)
+			{
+				SendTelemetry(newIssue);
+			}
+			return newIssue;
+		}
+
+		async Task<IssueSpan?> TryUpdateSpanAsync(IIssueSpan issueSpan, UpdateDefinition<IssueSpan> update)
+		{
+			IssueSpan issueSpanDocument = (IssueSpan)issueSpan;
+
+			int prevUpdateIndex = issueSpanDocument.UpdateIndex;
+			update = update.Set(x => x.UpdateIndex, prevUpdateIndex + 1);
+
+			FindOneAndUpdateOptions<IssueSpan, IssueSpan> options = new FindOneAndUpdateOptions<IssueSpan, IssueSpan> { ReturnDocument = ReturnDocument.After };
+
+			IssueSpan? newIssueSpan = await _issueSpans.FindOneAndUpdateAsync<IssueSpan>(x => x.Id == issueSpanDocument.Id && x.UpdateIndex == prevUpdateIndex, update, options);
+			if (newIssueSpan != null)
+			{
+				SendTelemetry(newIssueSpan);
+			}
+			return newIssueSpan;
+		}
+
+		#region Issues
+
+		/// <inheritdoc/>
+		public async Task<IIssue> AddIssueAsync(string summary)
+		{
+			IssueLedger ledger = await _ledgerSingleton.UpdateAsync(x => x.NextId++);
+
+			Issue newIssue = new Issue(ledger.NextId, summary);
+			await _issues.InsertOneAsync(newIssue);
+			SendTelemetry(newIssue);
+
+			ILogger issueLogger = GetLogger(newIssue.Id);
+			issueLogger.LogInformation("Created issue {IssueId}", newIssue.Id);
+
+			return newIssue;
+		}
+
+		void SendTelemetry(IIssue issue)
+		{
+			if (_telemetrySink.Enabled)
+			{
+				_telemetrySink.SendEvent("State.Issue", new
+				{
+					Id = issue.Id,
+					AcknowledgedAt = issue.AcknowledgedAt,
+					CreatedAt = issue.CreatedAt,
+					OwnerId = issue.OwnerId,
+					FixChange = issue.FixChange,
+					LastSeenAt = issue.LastSeenAt,
+					NominatedAt = issue.NominatedAt,
+					NominatedById = issue.NominatedById,
+					ResolvedAt = issue.ResolvedAt,
+					ResolvedById = issue.ResolvedById,
+					Severity = issue.Severity,
+					Summary = issue.Summary,
+					VerifiedAt = issue.VerifiedAt
+				});
+			}
+		}
+
+		async ValueTask<string> GetUserNameAsync(UserId? UserId)
+		{
+			if (UserId == null)
+			{
+				return "null";
+			}
+			else if (UserId == IIssue.ResolvedByUnknownId)
+			{
+				return "Horde (Unknown)";
+			}
+			else if (UserId == IIssue.ResolvedByTimeoutId)
+			{
+				return "Horde (Timeout)";
+			}
+
+			IUser? user = await _userCollection.GetCachedUserAsync(UserId);
+			if (user == null)
+			{
+				return "Unknown user";
+			}
+
+			return user.Name;
+		}
+
+		async Task LogIssueChangesAsync(UserId? initiatedByUserId, Issue oldIssue, Issue newIssue)
+		{
+			ILogger issueLogger = GetLogger(oldIssue.Id);
+			using IDisposable scope = issueLogger.BeginScope("User {UserName} ({UserId})", await GetUserNameAsync(initiatedByUserId), initiatedByUserId ?? UserId.Empty);
+			await LogIssueChangesImplAsync(issueLogger, oldIssue, newIssue);
+		}
+
+		async Task LogIssueChangesImplAsync(ILogger issueLogger, Issue oldIssue, Issue newIssue)
+		{
+			if (newIssue.Severity != oldIssue.Severity)
+			{
+				issueLogger.LogInformation("Changed severity to {Severity}", newIssue.Severity);
+			}
+			if (newIssue.Summary != oldIssue.Summary)
+			{
+				issueLogger.LogInformation("Changed summary to \"{Summary}\"", newIssue.Summary);
+			}
+			if (newIssue.Description != oldIssue.Description)
+			{
+				issueLogger.LogInformation("Description set to {Value}", newIssue.Description);
+			}
+			if (((IIssue)newIssue).Promoted != ((IIssue)oldIssue).Promoted)
+			{
+				issueLogger.LogInformation("Promoted set to {Value}", ((IIssue)newIssue).Promoted);
+			}
+			if (newIssue.OwnerId != oldIssue.OwnerId)
+			{
+				if (newIssue.NominatedById != null)
+				{					
+					issueLogger.LogInformation("User {UserName} ({UserId}) was nominated by {NominatedByUserName} ({NominatedByUserId})", await GetUserNameAsync(newIssue.OwnerId), newIssue.OwnerId, await GetUserNameAsync(newIssue.NominatedById), newIssue.NominatedById);
+				}
+				else
+				{
+					issueLogger.LogInformation("User {UserName} ({UserId}) was nominated by default", await GetUserNameAsync(newIssue.OwnerId), newIssue.OwnerId);
+				}
+			}
+			if (newIssue.AcknowledgedAt != oldIssue.AcknowledgedAt)
+			{
+				if (newIssue.AcknowledgedAt == null)
+				{
+					issueLogger.LogInformation("Issue was un-acknowledged by {UserName} ({UserId})", await GetUserNameAsync(oldIssue.OwnerId), oldIssue.OwnerId);
+				}
+				else
+				{
+					issueLogger.LogInformation("Issue was acknowledged by {UserName} ({UserId})", await GetUserNameAsync(newIssue.OwnerId), newIssue.OwnerId);
+				}
+			}
+			if (newIssue.FixChange != oldIssue.FixChange)
+			{
+				if (newIssue.FixChange == 0)
+				{
+					issueLogger.LogInformation("Issue was marked as not fixed");
+				}
+				else
+				{
+					issueLogger.LogInformation("Issue was marked as fixed in {Change}", newIssue.FixChange);
+				}
+			}
+			if (newIssue.ResolvedById != oldIssue.ResolvedById)
+			{
+				if (newIssue.ResolvedById == null)
+				{
+					issueLogger.LogInformation("Marking as unresolved");
+				}
+				else
+				{
+					issueLogger.LogInformation("Resolved by {UserName} ({UserId})", await GetUserNameAsync(newIssue.ResolvedById), newIssue.ResolvedById);
+				}
+			}
+
+			if (newIssue.ResolvedAt != oldIssue.ResolvedAt)
+			{
+				if (newIssue.ResolvedAt == null)
+				{
+					issueLogger.LogInformation("Clearing resolved at time");
+				}
+				else
+				{
+					issueLogger.LogInformation("Setting resolved at time");
+				}
+			}
+
+			if (newIssue.VerifiedAt != oldIssue.VerifiedAt)
+			{
+				if (newIssue.VerifiedAt == null)
+				{
+					issueLogger.LogInformation("Clearing verified at time");
+				}
+				else
+				{
+					issueLogger.LogInformation("Setting verified at time");
+				}
+			}
+
+			if (newIssue.ExternalIssueKey != oldIssue.ExternalIssueKey)
+			{
+				if (newIssue.ExternalIssueKey != null)
+				{
+					issueLogger.LogInformation("Linked to external issue {ExternalIssueKey}", newIssue.ExternalIssueKey);
+				}
+				else
+				{
+					issueLogger.LogInformation("Unlinked from external issue {ExternalIssueKey}", oldIssue.ExternalIssueKey);
+				}
+			}
+
+			if (newIssue.QuarantinedByUserId != oldIssue.QuarantinedByUserId)
+			{
+				if (newIssue.QuarantinedByUserId != null)
+				{
+					issueLogger.LogInformation("Quarantined by {UserName} ({UserId})", await GetUserNameAsync(newIssue.QuarantinedByUserId), newIssue.QuarantinedByUserId);
+				}
+				else
+				{
+					issueLogger.LogInformation("Quarantine cleared");
+				}
+			}
+
+			if (newIssue.ForceClosedByUserId != oldIssue.ForceClosedByUserId)
+			{
+				if (newIssue.ForceClosedByUserId != null)
+				{
+					issueLogger.LogInformation("Forced closed by {UserName} ({UserId})", await GetUserNameAsync(newIssue.ForceClosedByUserId), newIssue.ForceClosedByUserId);
+				}
+				else
+				{
+					issueLogger.LogInformation("Force closed cleared");
+				}
+			}
+
+			string oldFingerprints = oldIssue.FingerprintsDesc;
+			string newFingerprints = newIssue.FingerprintsDesc;
+			if (oldFingerprints != newFingerprints)
+			{
+				issueLogger.LogInformation("Fingerprints changed {Fingerprints}", newFingerprints);
+			}
+
+			HashSet<StreamId> oldFixStreams = new HashSet<StreamId>(oldIssue.Streams.Where(x => x.ContainsFix ?? false).Select(x => x.StreamId));
+			HashSet<StreamId> newFixStreams = new HashSet<StreamId>(newIssue.Streams.Where(x => x.ContainsFix ?? false).Select(x => x.StreamId));
+			foreach (StreamId streamId in newFixStreams.Where(x => !oldFixStreams.Contains(x)))
+			{
+				issueLogger.LogInformation("Marking stream {StreamId} as fixed", streamId);
+			}
+			foreach (StreamId streamId in oldFixStreams.Where(x => !newFixStreams.Contains(x)))
+			{
+				issueLogger.LogInformation("Marking stream {StreamId} as not fixed", streamId);
+			}
+		}
+
+		async Task LogIssueSuspectChangesAsync(ILogger issueLogger, List<IssueSuspect> oldIssueSuspects, List<IssueSuspect> newIssueSuspects)
+		{
+			HashSet<(UserId, int)> oldSuspects = new HashSet<(UserId, int)>(oldIssueSuspects.Select(x => (x.AuthorId, x.Change)));
+			HashSet<(UserId, int)> newSuspects = new HashSet<(UserId, int)>(newIssueSuspects.Select(x => (x.AuthorId, x.Change)));
+			foreach ((UserId userId, int change) in newSuspects.Where(x => !oldSuspects.Contains(x)))
+			{
+				issueLogger.LogInformation("Added suspect {UserName} ({UserId}) for change {Change}", await GetUserNameAsync(userId), userId, change);
+			}
+			foreach ((UserId userId, int change) in oldSuspects.Where(x => !newSuspects.Contains(x)))
+			{
+				issueLogger.LogInformation("Removed suspect {UserName} ({UserId}) for change {Change}", await GetUserNameAsync(userId), userId, change);
+			}
+
+			HashSet<UserId> oldDeclinedBy = new HashSet<UserId>(oldIssueSuspects.Where(x => x.DeclinedAt != null).Select(x => x.AuthorId));
+			HashSet<UserId> newDeclinedBy = new HashSet<UserId>(newIssueSuspects.Where(x => x.DeclinedAt != null).Select(x => x.AuthorId));
+			foreach (UserId addDeclinedBy in newDeclinedBy.Where(x => !oldDeclinedBy.Contains(x)))
+			{
+				issueLogger.LogInformation("Declined by {UserName} ({UserId})", await GetUserNameAsync(addDeclinedBy), addDeclinedBy);
+			}
+			foreach (UserId removeDeclinedBy in oldDeclinedBy.Where(x => !newDeclinedBy.Contains(x)))
+			{
+				issueLogger.LogInformation("Un-declined by {UserName} ({UserId})", await GetUserNameAsync(removeDeclinedBy), removeDeclinedBy);
+			}
+		}
+
+		/// <inheritdoc/>
+		public async Task<IIssue?> GetIssueAsync(int issueId)
+		{
+			Issue issue = await _issues.Find(x => x.Id == issueId).FirstOrDefaultAsync();
+			return issue;
+		}
+
+		/// <inheritdoc/>
+		public Task<List<IIssueSuspect>> FindSuspectsAsync(int issueId)
+		{
+			return _issueSuspects.Find(x => x.IssueId == issueId).ToListAsync<IssueSuspect, IIssueSuspect>();
+		}
+
+		class ProjectedIssueId
+		{
+			[System.Diagnostics.CodeAnalysis.SuppressMessage("Style", "IDE1006:Naming Styles")]
+			public int? _id { get; set; }
+		}
+
+		/// <inheritdoc/>
+		public async Task<List<IIssue>> FindIssuesAsync(IEnumerable<int>? ids = null, UserId? ownerId = null, StreamId? streamId = null, int? minChange = null, int? maxChange = null, bool? resolved = null, bool? promoted = null, int? index = null, int? count = null)
+		{
+			List<Issue> results;
+
+			if (ownerId == null)
+			{
+				results = await FilterIssuesByStreamIdAsync(ids, streamId, minChange, maxChange, resolved ?? false, promoted, index ?? 0, count);
+			}
+			else
+			{
+				results = await _issues.Find(x => x.OwnerId == ownerId).ToListAsync();
+			}
+
+			return results.ConvertAll<IIssue>(x => x);
+		}
+
+		async Task<List<Issue>> FilterIssuesByStreamIdAsync(IEnumerable<int>? ids, StreamId? streamId, int? minChange, int? maxChange, bool? resolved, bool? promoted, int index, int? count)
+		{
+			if (streamId == null)
+			{
+				return await FilterIssuesByOtherFieldsAsync(ids, minChange, maxChange, resolved, promoted, index, count);
+			}
+			else
+			{
+				FilterDefinition<IssueSpan> filter = Builders<IssueSpan>.Filter.Eq(x => x.StreamId, streamId.Value);
+				if (ids != null)
+				{
+					filter &= Builders<IssueSpan>.Filter.In(x => x.IssueId, ids.Select<int, int?>(x => x));
+				}
+				else
+				{
+					filter &= Builders<IssueSpan>.Filter.Exists(x => x.IssueId);
+				}
+
+				if (minChange != null)
+				{
+					filter &= Builders<IssueSpan>.Filter.Not(Builders<IssueSpan>.Filter.Lt(x => x.MaxChange, minChange.Value));
+				}
+				if (maxChange != null)
+				{
+					filter &= Builders<IssueSpan>.Filter.Not(Builders<IssueSpan>.Filter.Gt(x => x.MinChange, maxChange.Value));
+				}
+
+				if (resolved != null)
+				{
+					if (resolved.Value)
+					{
+						filter &= Builders<IssueSpan>.Filter.Ne(x => x.ResolvedAt, null);
+					}
+					else
+					{
+						filter &= Builders<IssueSpan>.Filter.Eq(x => x.ResolvedAt, null);
+					}
+				}
+
+				using (IAsyncCursor<ProjectedIssueId> cursor = await _issueSpans.Aggregate().Match(filter).Group(x => x.IssueId, x => new ProjectedIssueId { _id = x.Key }).SortByDescending(x => x._id).ToCursorAsync())
+				{
+					List<Issue> results = await PaginatedJoinAsync(cursor, (nextIds, nextIndex, nextCount) => FilterIssuesByOtherFieldsAsync(nextIds, null, null, null, promoted, nextIndex, nextCount), index, count);
+					if (resolved != null)
+					{
+						for (int idx = results.Count - 1; idx >= 0; idx--)
+						{
+							Issue issue = results[idx];
+							if ((issue.ResolvedAt != null) != resolved.Value)
+							{
+								_logger.LogWarning("Issue {IssueId} has resolved state out of sync with spans", issue.Id);
+								results.RemoveAt(idx);
+							}
+						}
+					}
+					return results;
+				}
+			}
+		}
+
+		async Task<List<Issue>> FilterIssuesByOtherFieldsAsync(IEnumerable<int>? ids, int? minChange, int? maxChange, bool? resolved, bool? promoted, int index, int? count)
+		{
+			FilterDefinition<Issue> filter = FilterDefinition<Issue>.Empty;
+			if (ids != null)
+			{
+				filter &= Builders<Issue>.Filter.In(x => x.Id, ids);
+			}
+			if (resolved != null)
+			{
+				if (resolved.Value)
+				{
+					filter &= Builders<Issue>.Filter.Ne(x => x.ResolvedAt, null);
+				}
+				else
+				{
+					filter &= Builders<Issue>.Filter.Eq(x => x.ResolvedAt, null);
+				}
+			}
+			if (minChange != null)
+			{
+				filter &= Builders<Issue>.Filter.Not(Builders<Issue>.Filter.Lt(x => x.MaxSuspectChange, minChange.Value));
+			}
+			if (maxChange != null)
+			{
+				filter &= Builders<Issue>.Filter.Not(Builders<Issue>.Filter.Gt(x => x.MinSuspectChange, maxChange.Value));
+			}
+			if (promoted != null)
+			{
+				if (promoted.Value)
+				{
+					filter &= Builders<Issue>.Filter.Eq(x => x.Promoted, true);
+				}
+				else
+				{
+					filter &= Builders<Issue>.Filter.Ne(x => x.Promoted, true); // Handle the field not existing as well as being set to false.
+				}
+			}
+			return await _issues.Find(filter).SortByDescending(x => x.Id).Range(index, count).ToListAsync();
+		}
+
+		/// <summary>
+		/// Performs a client-side join of a filtered set of issues against another query
+		/// </summary>
+		/// <param name="cursor"></param>
+		/// <param name="nextStageFunc"></param>
+		/// <param name="index"></param>
+		/// <param name="count"></param>
+		/// <returns></returns>
+		static async Task<List<Issue>> PaginatedJoinAsync(IAsyncCursor<ProjectedIssueId> cursor, Func<IEnumerable<int>, int, int?, Task<List<Issue>>> nextStageFunc, int index, int? count)
+		{
+			if (count == null)
+			{
+				List<ProjectedIssueId> issueIds = await cursor.ToListAsync();
+				return await nextStageFunc(issueIds.Where(x => x._id != null).Select(x => x._id!.Value), index, null);
+			}
+			else
+			{
+				List<Issue> results = new List<Issue>();
+				while (await cursor.MoveNextAsync() && results.Count < count.Value)
+				{
+					List<Issue> nextResults = await nextStageFunc(cursor.Current.Where(x => x._id != null).Select(x => x._id!.Value), 0, count.Value - results.Count);
+					int removeCount = Math.Min(index, nextResults.Count);
+					nextResults.RemoveRange(0, removeCount);
+					index -= removeCount;
+					results.AddRange(nextResults);
+				}
+				return results;
+			}
+		}
+
+		/// <inheritdoc/>
+		public async Task<List<IIssue>> FindIssuesForChangesAsync(List<int> changes)
+		{
+			List<int> issueIds = await (await _issueSuspects.DistinctAsync(x => x.IssueId, Builders<IssueSuspect>.Filter.In(x => x.Change, changes))).ToListAsync();
+			return await _issues.Find(Builders<Issue>.Filter.In(x => x.Id, issueIds)).ToListAsync<Issue, IIssue>();
+		}
+
+		/// <inheritdoc/>
+		public async Task<IIssue?> TryUpdateIssueAsync(IIssue issue, UserId? initiatedByUserId, IssueSeverity? newSeverity = null, string? newSummary = null, string? newUserSummary = null, string? newDescription = null, bool? newManuallyPromoted = null, UserId? newOwnerId = null, UserId? newNominatedById = null, bool? newAcknowledged = null, UserId? newDeclinedById = null, int? newFixChange = null, UserId? newResolvedById = null, List<ObjectId>? newExcludeSpanIds = null, DateTime? newLastSeenAt = null, string? newExternaIssueKey = null, UserId? newQuarantinedById = null, UserId? newForceClosedById = null, Uri? newWorkflowThreadUrl = null)
+		{
+			Issue issueDocument = (Issue)issue;
+
+			if (newDeclinedById != null && newDeclinedById == issueDocument.OwnerId)
+			{
+				newOwnerId = UserId.Empty;
+			}
+
+			if (issue.ResolvedById == null && (newForceClosedById != null && newResolvedById == null))
+			{
+				newResolvedById = newForceClosedById;
+			}
+
+			DateTime utcNow = DateTime.UtcNow;
+
+			List<UpdateDefinition<Issue>> updates = new List<UpdateDefinition<Issue>>();
+			if (newSeverity != null)
+			{
+				updates.Add(Builders<Issue>.Update.Set(x => x.Severity, newSeverity.Value));
+			}
+			if (newSummary != null)
+			{
+				updates.Add(Builders<Issue>.Update.Set(x => x.Summary, newSummary));
+			}
+			if (newUserSummary != null)
+			{
+				if (newUserSummary.Length == 0)
+				{
+					updates.Add(Builders<Issue>.Update.Unset(x => x.UserSummary!));
+				}
+				else
+				{
+					updates.Add(Builders<Issue>.Update.Set(x => x.UserSummary, newUserSummary));
+				}
+			}
+			if (newDescription != null)
+			{
+				if (newDescription.Length == 0)
+				{
+					updates.Add(Builders<Issue>.Update.Unset(x => x.Description));
+				}
+				else
+				{
+					updates.Add(Builders<Issue>.Update.Set(x => x.Description, newDescription));
+				}
+			}
+			if (newManuallyPromoted != null)
+			{
+				updates.Add(Builders<Issue>.Update.Set(x => x.ManuallyPromoted, newManuallyPromoted.Value));
+			}
+			if (newResolvedById != null)
+			{
+				newOwnerId ??= newResolvedById;
+				newAcknowledged ??= true;
+			}
+			if (newOwnerId != null)
+			{
+				if (newOwnerId.Value == UserId.Empty)
+				{
+					updates.Add(Builders<Issue>.Update.Unset(x => x.OwnerId!));
+					updates.Add(Builders<Issue>.Update.Unset(x => x.NominatedAt!));
+					updates.Add(Builders<Issue>.Update.Unset(x => x.NominatedById!));
+				}
+				else
+				{
+					updates.Add(Builders<Issue>.Update.Set(x => x.OwnerId!, newOwnerId.Value));
+
+					updates.Add(Builders<Issue>.Update.Set(x => x.NominatedAt, DateTime.UtcNow));
+					if (newNominatedById == null)
+					{
+						updates.Add(Builders<Issue>.Update.Unset(x => x.NominatedById!));
+					}
+					else
+					{
+						updates.Add(Builders<Issue>.Update.Set(x => x.NominatedById, newNominatedById.Value));
+					}
+					newAcknowledged ??= false;
+				}
+			}
+			if (newAcknowledged != null)
+			{
+				if (newAcknowledged.Value)
+				{
+					if (issueDocument.AcknowledgedAt == null)
+					{
+						updates.Add(Builders<Issue>.Update.Set(x => x.AcknowledgedAt, utcNow));
+					}
+				}
+				else
+				{
+					if (issueDocument.AcknowledgedAt != null)
+					{
+						updates.Add(Builders<Issue>.Update.Unset(x => x.AcknowledgedAt!));
+					}
+				}
+			}
+			if (newFixChange != null)
+			{
+				if (newFixChange == 0)
+				{
+					updates.Add(Builders<Issue>.Update.Unset(x => x.FixChange!));
+				}
+				else
+				{
+					updates.Add(Builders<Issue>.Update.Set(x => x.FixChange, newFixChange));
+				}
+			}
+			if (newResolvedById != null)
+			{
+				if (newResolvedById.Value != UserId.Empty)
+				{
+					if (issueDocument.ResolvedAt == null || issueDocument.ResolvedById != newResolvedById)
+					{
+						updates.Add(Builders<Issue>.Update.Set(x => x.ResolvedAt, utcNow));
+						updates.Add(Builders<Issue>.Update.Set(x => x.ResolvedById, newResolvedById.Value));
+					}
+				}
+				else
+				{
+					if (issueDocument.ResolvedAt != null)
+					{
+						updates.Add(Builders<Issue>.Update.Unset(x => x.ResolvedAt!));
+					}
+					if (issueDocument.ResolvedById != null)
+					{
+						updates.Add(Builders<Issue>.Update.Unset(x => x.ResolvedById!));
+					}
+				}
+			}
+			if (newExcludeSpanIds != null)
+			{
+				List<ObjectId> newCombinedExcludeSpanIds = newExcludeSpanIds;
+				if (issue.ExcludeSpans != null)
+				{
+					newCombinedExcludeSpanIds = newCombinedExcludeSpanIds.Union(issue.ExcludeSpans).ToList();
+				}
+				updates.Add(Builders<Issue>.Update.Set(x => x.ExcludeSpans, newCombinedExcludeSpanIds));
+			}
+			if (newLastSeenAt != null)
+			{
+				updates.Add(Builders<Issue>.Update.Set(x => x.LastSeenAt, newLastSeenAt.Value));
+			}
+
+			if (newDeclinedById != null)
+			{
+				GetLogger(issue.Id).LogInformation("Declined by {UserId}", newDeclinedById.Value);
+				await _issueSuspects.UpdateManyAsync(x => x.IssueId == issue.Id && x.AuthorId == newDeclinedById.Value, Builders<IssueSuspect>.Update.Set(x => x.DeclinedAt, DateTime.UtcNow));
+			}
+			if (newQuarantinedById != null)
+			{
+				if (newQuarantinedById.Value == UserId.Empty)
+				{
+					updates.Add(Builders<Issue>.Update.Unset(x => x.QuarantinedByUserId));
+					updates.Add(Builders<Issue>.Update.Unset(x => x.QuarantineTimeUtc));
+				}
+				else
+				{
+					updates.Add(Builders<Issue>.Update.Set(x => x.QuarantinedByUserId!, newQuarantinedById.Value));
+					updates.Add(Builders<Issue>.Update.Set(x => x.QuarantineTimeUtc, DateTime.UtcNow));
+				}
+			}
+			else if ((newResolvedById != null && newResolvedById.Value != UserId.Empty) || (newForceClosedById != null && newForceClosedById.Value != UserId.Empty))
+			{
+				// Clear quarantine if being resolved or if being force closed
+				updates.Add(Builders<Issue>.Update.Unset(x => x.QuarantinedByUserId));
+				updates.Add(Builders<Issue>.Update.Unset(x => x.QuarantineTimeUtc));
+			}
+
+			if (newForceClosedById != null)
+			{
+				if (newForceClosedById.Value == UserId.Empty)
+				{
+					updates.Add(Builders<Issue>.Update.Unset(x => x.ForceClosedByUserId));
+				}
+				else
+				{
+					updates.Add(Builders<Issue>.Update.Set(x => x.ForceClosedByUserId, newForceClosedById.Value));					
+				}
+			}
+
+			if (newExternaIssueKey != null)
+			{
+				updates.Add(Builders<Issue>.Update.Set(x => x.ExternalIssueKey, newExternaIssueKey.Length == 0 ? null : newExternaIssueKey));
+			}
+
+			if (newWorkflowThreadUrl != null)
+			{
+				updates.Add(Builders<Issue>.Update.Set(x => x.WorkflowThreadUrl, newWorkflowThreadUrl.ToString().Length == 0 ? null : newWorkflowThreadUrl));
+			}
+
+
+			if (updates.Count == 0)
+			{
+				return issueDocument;
+			}
+
+			Issue? newIssue = await TryUpdateIssueAsync(issue, Builders<Issue>.Update.Combine(updates));
+			if(newIssue == null)
+			{
+				return null;
+			}
+
+			await LogIssueChangesAsync(initiatedByUserId, issueDocument, newIssue);
+			return newIssue;
+		}
+
+		/// <inheritdoc/>
+		public async Task<IIssue?> TryUpdateIssueDerivedDataAsync(IIssue issue, string newSummary, IssueSeverity newSeverity, List<NewIssueFingerprint> newFingerprints, List<NewIssueStream> newStreams, List<NewIssueSuspectData> newSuspects, DateTime? newResolvedAt, DateTime? newVerifiedAt, DateTime newLastSeenAt)
+		{
+			Issue issueImpl = (Issue)issue;
+
+			// Update all the suspects for this issue
+			List<IssueSuspect> oldSuspectImpls = await _issueSuspects.Find(x => x.IssueId == issue.Id).ToListAsync();
+			List<IssueSuspect> newSuspectImpls = await UpdateIssueSuspectsAsync(issue.Id, oldSuspectImpls, newSuspects, newResolvedAt);
+
+			// Find the spans for this issue
+			List<IssueSpan> newSpans = await _issueSpans.Find(x => x.IssueId == issue.Id).ToListAsync();
+
+			// Update the resolved time on any issues
+			List<ObjectId> updateSpanIds = newSpans.Where(x => x.ResolvedAt != newResolvedAt).Select(x => x.Id).ToList();
+			if (updateSpanIds.Count > 0)
+			{
+				FilterDefinition<IssueSpan> filter = Builders<IssueSpan>.Filter.In(x => x.Id, updateSpanIds);
+				await _issueSpans.UpdateManyAsync(filter, Builders<IssueSpan>.Update.Set(x => x.ResolvedAt, newResolvedAt));
+			}
+
+			// Figure out if this issue should be promoted
+			bool newPromoted;
+			if (issueImpl.ManuallyPromoted.HasValue)
+			{
+				newPromoted = issueImpl.ManuallyPromoted.Value;
+			}
+			else if (issueImpl.ManuallyPromotedDeprecated.HasValue)
+			{
+				newPromoted = issueImpl.ManuallyPromotedDeprecated.Value;
+			}
+			else
+			{
+				newPromoted = newSpans.Any(x => ((IIssueSpan)x).PromoteByDefault);
+			}
+
+			// Find the default owner
+			UserId? newDefaultOwnerId = null;
+			string? autoAssignToUser = newSpans.Select(x => x.LastFailure.Annotations?.AutoAssignToUser).Where(x => x != null).FirstOrDefault();
+			if (autoAssignToUser != null)
+			{
+				IUser? user = await _userCollection.FindUserByLoginAsync(autoAssignToUser);
+				if(user != null)
+				{
+					newDefaultOwnerId = user.Id;
+				}
+			}
+			else
+			{
+				// Figure out if we can auto-assign an owner
+				bool canAutoAssign = newPromoted || newSpans.Any(x => x.LastFailure.Annotations?.AutoAssign ?? false);
+				if (canAutoAssign && newSuspectImpls.Count > 0)
+				{
+					UserId possibleOwnerId = newSuspectImpls[0].AuthorId;
+					if (newSuspectImpls.All(x => x.AuthorId == possibleOwnerId) && newSuspectImpls.Any(x => x.DeclinedAt == null))
+					{
+						newDefaultOwnerId = possibleOwnerId;
+					}
+				}
+			}
+
+			// Get the range of suspect changes
+			int newMinSuspectChange = (newSuspects.Count > 0) ? newSuspects.Min(x => x.Change) : 0;
+			int newMaxSuspectChange = (newSuspects.Count > 0) ? newSuspects.Min(x => x.Change) : 0;
+
+			// Perform the actual update with this data
+			List<UpdateDefinition<Issue>> updates = new List<UpdateDefinition<Issue>>();
+			if (!String.Equals(issue.Summary, newSummary, StringComparison.Ordinal))
+			{
+				updates.Add(Builders<Issue>.Update.Set(x => x.Summary, newSummary));
+			}
+			if (issue.Severity != newSeverity)
+			{
+				updates.Add(Builders<Issue>.Update.Set(x => x.Severity, newSeverity));
+			}
+			if (issue.Promoted != newPromoted)
+			{
+				updates.Add(Builders<Issue>.Update.Set(x => x.Promoted, newPromoted));
+			}
+			if (issue.Fingerprints.Count != newFingerprints.Count || !newFingerprints.Zip(issue.Fingerprints).All(x => x.First.Equals(x.Second)))
+			{
+				updates.Add(Builders<Issue>.Update.Set(x => x.Fingerprints, newFingerprints.Select(x => new IssueFingerprint(x))));
+			}
+			if (issue.Streams.Count != newStreams.Count || !newStreams.Zip(issue.Streams).All(x => x.First.StreamId == x.Second.StreamId && x.First.ContainsFix == x.Second.ContainsFix))
+			{
+				updates.Add(Builders<Issue>.Update.Set(x => x.Streams, newStreams.Select(x => new IssueStream(x))));
+			}
+			if (issueImpl.MinSuspectChange != newMinSuspectChange)
+			{
+				updates.Add(Builders<Issue>.Update.Set(x => x.MinSuspectChange, newMinSuspectChange));
+			}
+			if (issueImpl.MaxSuspectChange != newMaxSuspectChange)
+			{
+				updates.Add(Builders<Issue>.Update.Set(x => x.MaxSuspectChange, newMaxSuspectChange));
+			}
+			if (issueImpl.DefaultOwnerId != newDefaultOwnerId)
+			{
+				updates.Add(Builders<Issue>.Update.Set(x => x.DefaultOwnerId, newDefaultOwnerId));
+			}
+			if (issue.ResolvedAt != newResolvedAt)
+			{
+				updates.Add(Builders<Issue>.Update.SetOrUnsetNull(x => x.ResolvedAt, newResolvedAt));
+			}
+			if (newResolvedAt == null && issue.ResolvedById != null)
+			{
+				updates.Add(Builders<Issue>.Update.Unset(x => x.ResolvedById));
+			}
+			if (issue.VerifiedAt != newVerifiedAt)
+			{
+				updates.Add(Builders<Issue>.Update.SetOrUnsetNull(x => x.VerifiedAt, newVerifiedAt));
+			}
+			if (issue.LastSeenAt != newLastSeenAt)
+			{
+				updates.Add(Builders<Issue>.Update.Set(x => x.LastSeenAt, newLastSeenAt));
+			}
+
+			Issue? newIssue = await TryUpdateIssueAsync(issue, Builders<Issue>.Update.Combine(updates));
+			if(newIssue != null)
+			{
+				await LogIssueChangesAsync(null, issueImpl, newIssue);
+				await LogIssueSuspectChangesAsync(GetLogger(issue.Id), oldSuspectImpls, newSuspectImpls);
+				return newIssue;
+			}
+			return null;
+		}
+
+		async Task<List<IssueSuspect>> UpdateIssueSuspectsAsync(int issueId, List<IssueSuspect> oldSuspectImpls, List<NewIssueSuspectData> newSuspects, DateTime? resolvedAt)
+		{
+			List<IssueSuspect> newSuspectImpls = new List<IssueSuspect>(oldSuspectImpls);
+
+			// Find the current list of suspects
+			HashSet<(UserId, int)> curSuspectKeys = new HashSet<(UserId, int)>(oldSuspectImpls.Select(x => (x.AuthorId, x.Change)));
+			List<IssueSuspect> createSuspects = newSuspects.Where(x => !curSuspectKeys.Contains((x.AuthorId, x.Change))).Select(x => new IssueSuspect(issueId, x, resolvedAt)).ToList();
+
+			HashSet<(UserId, int)> newSuspectKeys = new HashSet<(UserId, int)>(newSuspects.Select(x => (x.AuthorId, x.Change)));
+			List<IssueSuspect> deleteSuspects = oldSuspectImpls.Where(x => !newSuspectKeys.Contains((x.AuthorId, x.Change))).ToList();
+
+			// Apply the suspect changes
+			if (createSuspects.Count > 0)
+			{
+				await _issueSuspects.InsertManyIgnoreDuplicatesAsync(createSuspects);
+				newSuspectImpls.AddRange(createSuspects);
+			}
+			if (deleteSuspects.Count > 0)
+			{
+				await _issueSuspects.DeleteManyAsync(Builders<IssueSuspect>.Filter.In(x => x.Id, deleteSuspects.Select(y => y.Id)));
+				newSuspectImpls.RemoveAll(x => !newSuspectKeys.Contains((x.AuthorId, x.Change)));
+			}
+
+			// Make sure all the remaining suspects have the correct resolved time
+			if (newSuspectImpls.Any(x => x.ResolvedAt != resolvedAt))
+			{
+				await _issueSuspects.UpdateManyAsync(Builders<IssueSuspect>.Filter.Eq(x => x.IssueId, issueId), Builders<IssueSuspect>.Update.Set(x => x.ResolvedAt, resolvedAt));
+			}
+			return newSuspectImpls;
+		}
+
+		#endregion
+
+		#region Spans
+
+		/// <inheritdoc/>
+		public async Task<IIssueSpan> AddSpanAsync(int issueId, NewIssueSpanData newSpan)
+		{
+			IssueSpan span = new IssueSpan(issueId, newSpan);
+			await _issueSpans.InsertOneAsync(span);
+			SendTelemetry(span);
+			return span;
+		}
+
+		void SendTelemetry(IIssueSpan issueSpan)
+		{
+			if (_telemetrySink.Enabled)
+			{
+				_telemetrySink.SendEvent("State.IssueSpan", new
+				{
+					Id = issueSpan.Id,
+					IssueId = issueSpan.IssueId,
+					Fingerprint = new { Type = issueSpan.Fingerprint.Type, Keys = issueSpan.Fingerprint.Keys },
+					FirstFailure = new { JobId = issueSpan.FirstFailure.JobId, JobName = issueSpan.FirstFailure.JobName, Change = issueSpan.FirstFailure.Change, StepId = issueSpan.FirstFailure.SpanId },
+					LastFailure = (issueSpan.LastFailure != null) ? new { JobId = issueSpan.LastFailure.JobId, JobName = issueSpan.LastFailure.JobName, Change = issueSpan.LastFailure.Change, StepId = issueSpan.LastFailure.SpanId } : null,
+					StreamId = issueSpan.StreamId,
+					StreamName = issueSpan.StreamName,
+					TemplateRefId = issueSpan.TemplateRefId
+				});
+			}
+		}
+
+		/// <inheritdoc/>
+		public async Task<IIssueSpan?> GetSpanAsync(ObjectId spanId)
+		{
+			return await _issueSpans.Find(Builders<IssueSpan>.Filter.Eq(x => x.Id, spanId)).FirstOrDefaultAsync();
+		}
+
+		/// <inheritdoc/>
+		public async Task<IIssueSpan?> TryUpdateSpanAsync(IIssueSpan span, NewIssueStepData? newLastSuccess = null, NewIssueStepData? newFailure = null, NewIssueStepData? newNextSuccess = null, List<NewIssueSpanSuspectData>? newSuspects = null, int? newIssueId = null)
+		{
+			List<UpdateDefinition<IssueSpan>> updates = new List<UpdateDefinition<IssueSpan>>();
+			if (newLastSuccess != null)
+			{
+				updates.Add(Builders<IssueSpan>.Update.Set(x => x.MinChange, newLastSuccess.Change));
+				updates.Add(Builders<IssueSpan>.Update.Set(x => x.LastSuccess, new IssueStep(span.Id, newLastSuccess)));
+			}
+			if (newFailure != null)
+			{
+				if (newFailure.Change < span.FirstFailure.Change)
+				{
+					updates.Add(Builders<IssueSpan>.Update.Set(x => x.FirstFailure, new IssueStep(span.Id, newFailure)));
+				}
+				if (newFailure.Change >= span.LastFailure.Change)
+				{
+					updates.Add(Builders<IssueSpan>.Update.Set(x => x.LastFailure, new IssueStep(span.Id, newFailure)));
+				}
+				if (newFailure.PromoteByDefault != span.PromoteByDefault && newFailure.Change >= span.LastFailure.Change)
+				{
+					updates.Add(Builders<IssueSpan>.Update.Set(x => x.PromoteByDefault, newFailure.PromoteByDefault));
+				}
+			}
+			if (newNextSuccess != null)
+			{
+				updates.Add(Builders<IssueSpan>.Update.Set(x => x.MaxChange, newNextSuccess.Change));
+				updates.Add(Builders<IssueSpan>.Update.Set(x => x.NextSuccess, new IssueStep(span.Id, newNextSuccess)));
+			}
+			if (newSuspects != null)
+			{
+				updates.Add(Builders<IssueSpan>.Update.Set(x => x.Suspects, newSuspects.ConvertAll(x => new IssueSpanSuspect(x))));
+			}
+			if (newIssueId != null)
+			{
+				updates.Add(Builders<IssueSpan>.Update.Set(x => x.IssueId, newIssueId.Value));
+			}
+
+			if (updates.Count == 0)
+			{
+				return span;
+			}
+
+			IssueSpan? newSpan = await TryUpdateSpanAsync(span, Builders<IssueSpan>.Update.Combine(updates));
+			if (newSpan != null)
+			{
+				ILogger logger = GetLogger(newSpan.IssueId);
+				if (newLastSuccess != null)
+				{
+					logger.LogInformation("Set last success for span {SpanId} to job {JobId} at CL {Change}", newSpan.Id, newLastSuccess.JobId, newLastSuccess.Change);
+				}
+				if (newNextSuccess != null)
+				{
+					logger.LogInformation("Set next success for span {SpanId} to job {JobId} at CL {Change}", newSpan.Id, newNextSuccess.JobId, newNextSuccess.Change);
+				}
+				if (newFailure != null)
+				{
+					logger.LogInformation("Added failure for span {SpanId} in job {JobId} at CL {Change}", newSpan.Id, newFailure.JobId, newFailure.Change);
+				}
+				SendTelemetry(newSpan);
+			}
+			return newSpan;
+		}
+
+		/// <inheritdoc/>
+		public async Task<List<IIssueSpan>> FindSpansAsync(int issueId)
+		{
+			return await _issueSpans.Find(x => x.IssueId == issueId).ToListAsync<IssueSpan, IIssueSpan>();
+		}
+
+		/// <inheritdoc/>
+		public Task<List<IIssueSpan>> FindSpansAsync(IEnumerable<ObjectId> spanIds)
+		{
+			return _issueSpans.Find(Builders<IssueSpan>.Filter.In(x => x.Id, spanIds)).ToListAsync<IssueSpan, IIssueSpan>();
+		}
+
+		/// <inheritdoc/>
+		public async Task<List<IIssueSpan>> FindOpenSpansAsync(StreamId streamId, TemplateId templateId, string nodeName, int change)
+		{
+			List<IssueSpan> spans = await _issueSpans.Find(x => x.StreamId == streamId && x.TemplateRefId == templateId && x.NodeName == nodeName && change >= x.MinChange && change <= x.MaxChange).ToListAsync();
+			return spans.ConvertAll<IIssueSpan>(x => x);
+		}
+
+		/// <inheritdoc/>
+		public Task<List<IIssueSpan>> FindSpansAsync(IEnumerable<ObjectId>? spanIds, IEnumerable<int>? issueIds, StreamId? streamId, int? minChange, int? maxChange, bool? resolved, int? index, int? count)
+		{
+			FilterDefinition<IssueSpan> filter = FilterDefinition<IssueSpan>.Empty;
+
+			if(spanIds != null)
+			{
+				filter &= Builders<IssueSpan>.Filter.In(x => x.Id, spanIds);
+			}
+
+			if (streamId != null)
+			{
+				filter &= Builders<IssueSpan>.Filter.Eq(x => x.StreamId, streamId);
+			}
+
+			if (issueIds != null)
+			{
+				filter &= Builders<IssueSpan>.Filter.In(x => x.IssueId, issueIds.Select<int, int?>(x => x));
+			}
+			if (minChange != null)
+			{
+				filter &= Builders<IssueSpan>.Filter.Not(Builders<IssueSpan>.Filter.Lt(x => x.MaxChange, minChange.Value));
+			}
+			if (maxChange != null)
+			{
+				filter &= Builders<IssueSpan>.Filter.Not(Builders<IssueSpan>.Filter.Gt(x => x.MinChange, maxChange.Value));
+			}
+			if (resolved != null)
+			{
+				if (resolved.Value)
+				{
+					filter &= Builders<IssueSpan>.Filter.Ne(x => x.ResolvedAt, null);
+				}
+				else
+				{
+					filter &= Builders<IssueSpan>.Filter.Eq(x => x.ResolvedAt, null);
+				}
+			}
+
+			return _issueSpans.Find(filter).Range(index, count).ToListAsync<IssueSpan, IIssueSpan>();
+		}
+
+		#endregion
+
+		#region Steps
+
+		/// <inheritdoc/>
+		public async Task<IIssueStep> AddStepAsync(ObjectId spanId, NewIssueStepData newStep)
+		{
+			IssueStep step = new IssueStep(spanId, newStep);
+			await _issueSteps.InsertOneAsync(step);
+			return step;
+		}
+
+		/// <inheritdoc/>
+		public Task<List<IIssueStep>> FindStepsAsync(IEnumerable<ObjectId> spanIds)
+		{
+			FilterDefinition<IssueStep> filter = Builders<IssueStep>.Filter.In(x => x.SpanId, spanIds);
+			return _issueSteps.Find(filter).ToListAsync<IssueStep, IIssueStep>();
+		}
+
+		/// <inheritdoc/>
+		public Task<List<IIssueStep>> FindStepsAsync(JobId jobId, SubResourceId? batchId, SubResourceId? stepId)
+		{
+			FilterDefinition<IssueStep> filter = Builders<IssueStep>.Filter.Eq(x => x.JobId, jobId);
+			if (batchId != null)
+			{
+				filter &= Builders<IssueStep>.Filter.Eq(x => x.BatchId, batchId.Value);
+			}
+			if (stepId != null)
+			{
+				filter &= Builders<IssueStep>.Filter.Eq(x => x.StepId, stepId.Value);
+			}
+			return _issueSteps.Find(filter).ToListAsync<IssueStep, IIssueStep>();
+		}
+
+		#endregion
+
+		/// <inheritdoc/>
+		public IAuditLogChannel<int> GetLogger(int issueId)
+		{
+			return _auditLog[issueId];
+		}
+	}
+}
