Horde: Rename the Horde.Build project to Horde.Server.

diff --git a/Engine/Source/Programs/Horde/Horde.Server/Jobs/Graphs/GraphMessages.cs b/Engine/Source/Programs/Horde/Horde.Server/Jobs/Graphs/GraphMessages.cs
new file mode 100644
--- /dev/null
+++ b/Engine/Source/Programs/Horde/Horde.Server/Jobs/Graphs/GraphMessages.cs
@@ -0,0 +1,264 @@
+// Copyright Epic Games, Inc. All Rights Reserved.
+
+using System.Collections.Generic;
+using System.Linq;
+using Amazon.S3.Model;
+using EpicGames.Core;
+using HordeCommon;
+
+namespace Horde.Server.Jobs.Graphs
+{
+	/// <summary>
+	/// Information required to create a node
+	/// </summary>
+	public class GetNodeResponse
+	{
+		/// <summary>
+		/// The name of this node 
+		/// </summary>
+		public string Name { get; set; }
+
+		/// <summary>
+		/// Inputs for this node
+		/// </summary>
+		public List<string> Inputs { get; set; } = new List<string>();
+
+		/// <summary>
+		/// Output from this node
+		/// </summary>
+		public List<string> Outputs { get; set; } = new List<string>();
+
+		/// <summary>
+		/// Indices of nodes which must have succeeded for this node to run
+		/// </summary>
+		public List<string> InputDependencies { get; set; }
+
+		/// <summary>
+		/// Indices of nodes which must have completed for this node to run
+		/// </summary>
+		public List<string> OrderDependencies { get; set; }
+
+		/// <summary>
+		/// The priority of this node
+		/// </summary>
+		public Priority Priority { get; set; }
+
+		/// <summary>
+		/// Whether this node can be retried
+		/// </summary>
+		public bool AllowRetry { get; set; }
+
+		/// <summary>
+		/// This node can start running early, before dependencies of other nodes in the same group are complete
+		/// </summary>
+		public bool RunEarly { get; set; }
+
+		/// <summary>
+		/// Whether to include warnings in diagnostic output
+		/// </summary>
+		public bool Warnings { get; set; }
+
+		/// <summary>
+		/// Average time to execute this node based on historical trends
+		/// </summary>
+		public float? AverageDuration { get; set; }
+
+		/// <summary>
+		/// Credentials required for this node to run. This dictionary maps from environment variable names to a credential property in the format 'CredentialName.PropertyName'.
+		/// </summary>
+		public IReadOnlyDictionary<string, string>? Credentials { get; set; }
+
+		/// <summary>
+		/// Properties for this node
+		/// </summary>
+		public IReadOnlyDictionary<string, string>? Properties { get; set; }
+
+		/// <summary>
+		/// Annotations for this node
+		/// </summary>
+		public IReadOnlyDictionary<string, string>? Annotations { get; set; }
+
+		/// <summary>
+		/// Constructor
+		/// </summary>
+		/// <param name="node">The node to construct from</param>
+		/// <param name="groups">The groups in the graph</param>
+		public GetNodeResponse(INode node, IReadOnlyList<INodeGroup> groups)
+		{
+			Name = node.Name;
+			Inputs.AddRange(node.Inputs.Select(x => groups[x.NodeRef.GroupIdx].Nodes[x.NodeRef.NodeIdx].OutputNames[x.OutputIdx]));
+			Outputs.AddRange(node.OutputNames);
+			InputDependencies = new List<string>(node.InputDependencies.Select(x => groups[x.GroupIdx].Nodes[x.NodeIdx].Name));
+			OrderDependencies = new List<string>(node.OrderDependencies.Select(x => groups[x.GroupIdx].Nodes[x.NodeIdx].Name));
+			Priority = node.Priority;
+			AllowRetry = node.AllowRetry;
+			RunEarly = node.RunEarly;
+			Warnings = node.Warnings;
+			Credentials = node.Credentials;
+			Properties = node.Properties;
+			Annotations = node.Annotations;
+		}
+	}
+
+	/// <summary>
+	/// Information about a group of nodes
+	/// </summary>
+	public class GetGroupResponse
+	{
+		/// <summary>
+		/// The type of agent to execute this group
+		/// </summary>
+		public string AgentType { get; set; }
+
+		/// <summary>
+		/// Nodes in the group
+		/// </summary>
+		public List<GetNodeResponse> Nodes { get; set; }
+
+		/// <summary>
+		/// Constructor
+		/// </summary>
+		/// <param name="group">The group to construct from</param>
+		/// <param name="groups">Other groups in this graph</param>
+		public GetGroupResponse(INodeGroup group, IReadOnlyList<INodeGroup> groups)
+		{
+			AgentType = group.AgentType;
+			Nodes = group.Nodes.ConvertAll(x => new GetNodeResponse(x, groups));
+		}
+	}
+
+	/// <summary>
+	/// Information about an aggregate
+	/// </summary>
+	public class GetAggregateResponse
+	{
+		/// <summary>
+		/// Name of the aggregate
+		/// </summary>
+		public string Name { get; set; }
+
+		/// <summary>
+		/// Nodes which must be part of the job for the aggregate to be shown
+		/// </summary>
+		public List<string> Nodes { get; set; }
+
+		/// <summary>
+		/// Constructor
+		/// </summary>
+		/// <param name="aggregate">The aggregate to construct from</param>
+		/// <param name="groups">List of groups in this graph</param>
+		public GetAggregateResponse(IAggregate aggregate, IReadOnlyList<INodeGroup> groups)
+		{
+			Name = aggregate.Name;
+			Nodes = aggregate.Nodes.ConvertAll(x => x.ToNode(groups).Name);
+		}
+	}
+
+	/// <summary>
+	/// Information about a label
+	/// </summary>
+	public class GetLabelResponse
+	{
+		/// <summary>
+		/// Category of the aggregate
+		/// </summary>
+		public string? Category => DashboardCategory;
+
+		/// <summary>
+		/// Label for this aggregate
+		/// </summary>
+		public string? Name => DashboardName;
+
+		/// <summary>
+		/// Name to show for this label on the dashboard
+		/// </summary>
+		public string? DashboardName { get; set; }
+
+		/// <summary>
+		/// Category to show this label in on the dashboard
+		/// </summary>
+		public string? DashboardCategory { get; set; }
+
+		/// <summary>
+		/// Name to show for this label in UGS
+		/// </summary>
+		public string? UgsName { get; set; }
+
+		/// <summary>
+		/// Project to display this label for in UGS
+		/// </summary>
+		public string? UgsProject { get; set; }
+
+		/// <summary>
+		/// Nodes which must be part of the job for the aggregate to be shown
+		/// </summary>
+		public List<string> RequiredNodes { get; set; }
+
+		/// <summary>
+		/// Nodes to include in the status of this aggregate, if present in the job
+		/// </summary>
+		public List<string> IncludedNodes { get; set; }
+
+		/// <summary>
+		/// Constructor
+		/// </summary>
+		/// <param name="label">The label to construct from</param>
+		/// <param name="groups">List of groups in this graph</param>
+		public GetLabelResponse(ILabel label, IReadOnlyList<INodeGroup> groups)
+		{
+			DashboardName = label.DashboardName;
+			DashboardCategory = label.DashboardCategory;
+			UgsName = label.UgsName;
+			UgsProject = label.UgsProject;
+			RequiredNodes = label.RequiredNodes.ConvertAll(x => x.ToNode(groups).Name);
+			IncludedNodes = label.IncludedNodes.ConvertAll(x => x.ToNode(groups).Name);
+		}
+	}
+
+	/// <summary>
+	/// Information about a graph
+	/// </summary>
+	public class GetGraphResponse
+	{
+		/// <summary>
+		/// The hash of the graph
+		/// </summary>
+		public string Hash { get; set; }
+
+		/// <summary>
+		/// Array of nodes for this job
+		/// </summary>
+		public List<GetGroupResponse>? Groups { get; set; }
+
+		/// <summary>
+		/// List of aggregates
+		/// </summary>
+		public List<GetAggregateResponse>? Aggregates { get; set; }
+
+		/// <summary>
+		/// List of labels for the graph
+		/// </summary>
+		public List<GetLabelResponse>? Labels { get; set; }
+
+		/// <summary>
+		/// Constructor
+		/// </summary>
+		/// <param name="graph">The graph to construct from</param>
+		public GetGraphResponse(IGraph graph)
+		{
+			Hash = graph.Id.ToString();
+			if (graph.Groups.Count > 0)
+			{
+				Groups = graph.Groups.ConvertAll(x => new GetGroupResponse(x, graph.Groups));
+			}
+			if (graph.Aggregates.Count > 0)
+			{
+				Aggregates = graph.Aggregates.ConvertAll(x => new GetAggregateResponse(x, graph.Groups));
+			}
+			if (graph.Labels.Count > 0)
+			{
+				Labels = graph.Labels.ConvertAll(x => new GetLabelResponse(x, graph.Groups));
+			}
+		}
+	}
+}
diff --git a/Engine/Source/Programs/Horde/Horde.Server/Jobs/Graphs/IGraphCollection.cs b/Engine/Source/Programs/Horde/Horde.Server/Jobs/Graphs/IGraphCollection.cs
new file mode 100644
--- /dev/null
+++ b/Engine/Source/Programs/Horde/Horde.Server/Jobs/Graphs/IGraphCollection.cs
@@ -0,0 +1,49 @@
+// Copyright Epic Games, Inc. All Rights Reserved.
+
+using System.Collections.Generic;
+using System.Threading.Tasks;
+using EpicGames.Core;
+using Horde.Server.Jobs.Templates;
+
+namespace Horde.Server.Jobs.Graphs
+{
+	/// <summary>
+	/// Interface for a collection of graph documents
+	/// </summary>
+	public interface IGraphCollection
+	{
+		/// <summary>
+		/// Adds a graph from a template
+		/// </summary>
+		/// <param name="template">The template</param>
+		/// <param name="streamInitialAgentType">Default agent type for the stream</param>
+		/// <returns>New graph</returns>
+		Task<IGraph> AddAsync(ITemplate template, string? streamInitialAgentType);
+
+		/// <summary>
+		/// Creates a graph by appending groups and aggregates to an existing graph.
+		/// </summary>
+		/// <param name="baseGraph">The base graph</param>
+		/// <param name="newGroupRequests">List of group requests</param>
+		/// <param name="newAggregateRequests">List of aggregate requests</param>
+		/// <param name="newLabelRequests">List of label requests</param>
+		/// <returns>The new graph definition</returns>
+		Task<IGraph> AppendAsync(IGraph? baseGraph, List<NewGroup>? newGroupRequests = null, List<NewAggregate>? newAggregateRequests = null, List<NewLabel>? newLabelRequests = null);
+
+		/// <summary>
+		/// Gets the graph for a job
+		/// </summary>
+		/// <param name="hash">Hash of the graph to retrieve</param>
+		/// <returns>The graph for this job</returns>
+		Task<IGraph> GetAsync(ContentHash hash);
+
+		/// <summary>
+		/// Finds all graphs stored in the collection
+		/// </summary>
+		/// <param name="hashes">Hashes to filter by</param>
+		/// <param name="index">Starting index of the graph to return</param>
+		/// <param name="count">Number of results to return</param>
+		/// <returns>List of graphs</returns>
+		Task<List<IGraph>> FindAllAsync(ContentHash[]? hashes = null, int? index = null, int? count = null);
+	}
+}
diff --git a/Engine/Source/Programs/Horde/Horde.Server/Jobs/IJobCollection.cs b/Engine/Source/Programs/Horde/Horde.Server/Jobs/IJobCollection.cs
new file mode 100644
--- /dev/null
+++ b/Engine/Source/Programs/Horde/Horde.Server/Jobs/IJobCollection.cs
@@ -0,0 +1,347 @@
+// Copyright Epic Games, Inc. All Rights Reserved.
+
+using System;
+using System.Collections.Generic;
+using System.Threading.Tasks;
+using EpicGames.Core;
+using Horde.Server.Agents;
+using Horde.Server.Agents.Leases;
+using Horde.Server.Agents.Pools;
+using Horde.Server.Agents.Sessions;
+using Horde.Server.Jobs.Graphs;
+using Horde.Server.Logs;
+using Horde.Server.Streams;
+using Horde.Server.Users;
+using Horde.Server.Utilities;
+using Horde.Common;
+using HordeCommon;
+using HordeCommon.Rpc.Tasks;
+using MongoDB.Bson;
+
+namespace Horde.Server.Jobs
+{
+	/// <summary>
+	/// Options for creating a new job
+	/// </summary>
+	public class CreateJobOptions
+	{
+		/// <summary>
+		/// Optional changelist to preflight
+		/// </summary>
+		public int? PreflightChange { get; set; }
+
+		/// <summary>
+		/// Optional cloned preflight changelist
+		/// </summary>
+		public int? ClonedPreflightChange { get; set; }
+
+		/// <summary>
+		/// Description for the change being preflighted
+		/// </summary>
+		public string? PreflightDescription { get; set; }
+
+		/// <summary>
+		/// User that started the job
+		/// </summary>
+		public UserId? StartedByUserId { get; set; }
+
+		/// <summary>
+		/// Priority of the job
+		/// </summary>
+		public Priority? Priority { get; set; }
+
+		/// <summary>
+		/// Whether to automatically submit the preflighted change on completion
+		/// </summary>
+		public bool? AutoSubmit { get; set; }
+
+		/// <summary>
+		/// Whether to update issues based on the outcome of this job
+		/// </summary>
+		public bool? UpdateIssues { get; set; }
+
+		/// <summary>
+		/// Whether to promote issues based on the outcome of this job by default
+		/// </summary>
+		public bool? PromoteIssuesByDefault { get; set; }
+
+		/// <summary>
+		/// Settings for agents executing the job
+		/// </summary>
+		public JobOptions? JobOptions { get; set; }
+
+		/// <summary>
+		/// List of downstream job triggers
+		/// </summary>
+		public List<ChainedJobTemplateConfig> JobTriggers { get; } = new List<ChainedJobTemplateConfig>();
+
+		/// <summary>
+		/// Whether to show badges in UGS for this job
+		/// </summary>
+		public bool ShowUgsBadges { get; set; }
+
+		/// <summary>
+		/// Whether to show alerts in UGS for this job
+		/// </summary>
+		public bool ShowUgsAlerts { get; set; }
+
+		/// <summary>
+		/// Notification channel for this job
+		/// </summary>
+		public string? NotificationChannel { get; set; }
+
+		/// <summary>
+		/// Notification channel filter for this job
+		/// </summary>
+		public string? NotificationChannelFilter { get; set; }
+
+		/// <summary>
+		/// Arguments for the job
+		/// </summary>
+		public List<string> Arguments { get; } = new List<string>();
+
+		/// <summary>
+		/// Environment variables for the job
+		/// </summary>
+		public Dictionary<string, string> Environment { get; } = new Dictionary<string, string>(StringComparer.OrdinalIgnoreCase);
+
+		/// <summary>
+		/// Default constructor
+		/// </summary>
+		public CreateJobOptions()
+		{
+		}
+
+		/// <summary>
+		/// Constructor
+		/// </summary>
+		public CreateJobOptions(TemplateRefConfig templateRefConfig)
+		{
+			JobOptions = templateRefConfig.JobOptions;
+			PromoteIssuesByDefault = templateRefConfig.PromoteIssuesByDefault;
+			if (templateRefConfig.ChainedJobs != null)
+			{
+				JobTriggers.AddRange(templateRefConfig.ChainedJobs);
+			}
+			ShowUgsBadges = templateRefConfig.ShowUgsBadges;
+			ShowUgsAlerts = templateRefConfig.ShowUgsAlerts;
+			NotificationChannel = templateRefConfig.NotificationChannel;
+			NotificationChannelFilter = templateRefConfig.NotificationChannelFilter;
+		}
+	}
+
+	/// <summary>
+	/// Interface for a collection of job documents
+	/// </summary>
+	public interface IJobCollection
+	{
+		/// <summary>
+		/// Creates a new job
+		/// </summary>
+		/// <param name="jobId">A requested job id</param>
+		/// <param name="streamId">Unique id of the stream that this job belongs to</param>
+		/// <param name="templateRefId">Name of the template ref</param>
+		/// <param name="templateHash">Template for this job</param>
+		/// <param name="graph">The graph for the new job</param>
+		/// <param name="name">Name of the job</param>
+		/// <param name="change">The change to build</param>
+		/// <param name="codeChange">The corresponding code changelist number</param>
+		/// <param name="options">Additional options for the new job</param>
+		/// <returns>The new job document</returns>
+		Task<IJob> AddAsync(JobId jobId, StreamId streamId, TemplateId templateRefId, ContentHash templateHash, IGraph graph, string name, int change, int codeChange, CreateJobOptions options);
+
+		/// <summary>
+		/// Gets a job with the given unique id
+		/// </summary>
+		/// <param name="jobId">Job id to search for</param>
+		/// <returns>Information about the given job</returns>
+		Task<IJob?> GetAsync(JobId jobId);
+
+		/// <summary>
+		/// Deletes a job
+		/// </summary>
+		/// <param name="job">The job to remove</param>
+		Task<bool> RemoveAsync(IJob job);
+
+		/// <summary>
+		/// Delete all the jobs for a stream
+		/// </summary>
+		/// <param name="streamId">Unique id of the stream</param>
+		/// <returns>Async task</returns>
+		Task RemoveStreamAsync(StreamId streamId);
+
+		/// <summary>
+		/// Searches for jobs matching the given criteria
+		/// </summary>
+		/// <param name="jobIds">List of job ids to return</param>
+		/// <param name="streamId">The stream containing the job</param>
+		/// <param name="name">Name of the job</param>
+		/// <param name="templates">Templates to look for</param>
+		/// <param name="minChange">The minimum changelist number</param>
+		/// <param name="maxChange">The maximum changelist number</param>
+		/// <param name="preflightChange">Preflight change to find</param>
+		/// <param name="preflightOnly">Whether to only include preflights</param>
+		/// <param name="startedByUser">User id for which to include jobs</param>
+		/// <param name="preflightStartedByUser">User for which to include preflight jobs</param>
+		/// <param name="minCreateTime">The minimum creation time</param>
+		/// <param name="maxCreateTime">The maximum creation time</param>
+		/// <param name="modifiedBefore">Filter the results by modified time</param>
+		/// <param name="modifiedAfter">Filter the results by modified time</param>
+		/// <param name="index">Index of the first result to return</param>
+		/// <param name="count">Number of results to return</param>
+		/// <param name="consistentRead">If the database read should be made to the replica server</param>
+		/// <param name="indexHint">Name of index to be specified as a hint to the database query planner</param>
+		/// <param name="excludeUserJobs">Whether to exclude user jobs from the find</param>
+		/// <returns>List of jobs matching the given criteria</returns>
+		Task<List<IJob>> FindAsync(JobId[]? jobIds = null, StreamId? streamId = null, string? name = null, TemplateId[]? templates = null, int? minChange = null, int? maxChange = null, int? preflightChange = null, bool? preflightOnly = null, UserId? preflightStartedByUser = null, UserId? startedByUser = null, DateTimeOffset? minCreateTime = null, DateTimeOffset? maxCreateTime = null, DateTimeOffset? modifiedBefore = null, DateTimeOffset? modifiedAfter = null, int? index = null, int? count = null, bool consistentRead = true, string? indexHint = null, bool? excludeUserJobs = null);
+
+		/// <summary>
+		/// Searches for jobs in a specified stream and templates
+		/// </summary>
+		/// <param name="streamId">The stream containing the job</param>
+		/// <param name="templates">Templates to look for</param>
+		/// <param name="preflightStartedByUser">User for which to include preflight jobs</param>
+		/// <param name="maxCreateTime">The maximum creation time</param>
+		/// <param name="modifiedAfter">Filter the results by modified time</param>
+		/// <param name="index">Index of the first result to return</param>
+		/// <param name="count">Number of results to return</param>
+		/// <param name="consistentRead">If the database read should be made to the replica server</param>
+		Task<List<IJob>> FindLatestByStreamWithTemplatesAsync(StreamId streamId, TemplateId[] templates, UserId? preflightStartedByUser = null, DateTimeOffset? maxCreateTime = null, DateTimeOffset? modifiedAfter = null, int? index = null, int? count = null, bool consistentRead = false);
+
+		/// <summary>
+		/// Updates a new job
+		/// </summary>
+		/// <param name="job">The job document to update</param>
+		/// <param name="graph">Graph for the job</param>
+		/// <param name="name">Name of the job</param>
+		/// <param name="priority">Priority of the job</param>
+		/// <param name="autoSubmit">Automatically submit the job on completion</param>
+		/// <param name="autoSubmitChange">Changelist that was automatically submitted</param>
+		/// <param name="autoSubmitMessage"></param>
+		/// <param name="abortedByUserId">Name of the user that aborted the job</param>
+		/// <param name="notificationTriggerId">Id for a notification trigger</param>
+		/// <param name="reports">New reports</param>
+		/// <param name="arguments">New arguments for the job</param>
+		/// <param name="labelIdxToTriggerId">New trigger ID for a label in the job</param>
+		/// <param name="jobTrigger">New downstream job id</param>
+		Task<IJob?> TryUpdateJobAsync(IJob job, IGraph graph, string? name = null, Priority? priority = null, bool? autoSubmit = null, int? autoSubmitChange = null, string? autoSubmitMessage = null, UserId? abortedByUserId = null, ObjectId? notificationTriggerId = null, List<Report>? reports = null, List<string>? arguments = null, KeyValuePair<int, ObjectId>? labelIdxToTriggerId = null, KeyValuePair<TemplateId, JobId>? jobTrigger = null);
+
+		/// <summary>
+		/// Updates the state of a batch
+		/// </summary>
+		/// <param name="job">Job to update</param>
+		/// <param name="graph">The graph for this job</param>
+		/// <param name="batchId">Unique id of the batch to update</param>
+		/// <param name="newLogId">The new log file id</param>
+		/// <param name="newState">New state of the jobstep</param>
+		/// <param name="newError">Error code for the batch</param>
+		/// <returns>True if the job was updated, false if it was deleted</returns>
+		Task<IJob?> TryUpdateBatchAsync(IJob job, IGraph graph, SubResourceId batchId, LogId? newLogId, JobStepBatchState? newState, JobStepBatchError? newError);
+
+		/// <summary>
+		/// Update a jobstep state
+		/// </summary>
+		/// <param name="job">Job to update</param>
+		/// <param name="graph">The graph for this job</param>
+		/// <param name="batchId">Unique id of the batch containing the step</param>
+		/// <param name="stepId">Unique id of the step to update</param>
+		/// <param name="newState">New state of the jobstep</param>
+		/// <param name="newOutcome">New outcome of the jobstep</param>
+		/// <param name="newError">New error annotation for this jobstep</param>
+		/// <param name="newAbortRequested">New state of request abort</param>
+		/// <param name="newAbortByUserId">New name of user that requested the abort</param>
+		/// <param name="newLogId">New log id for the jobstep</param>
+		/// <param name="newNotificationTriggerId">New id for a notification trigger</param>
+		/// <param name="newRetryByUserId">Whether the step should be retried</param>
+		/// <param name="newPriority">New priority for this step</param>
+		/// <param name="newReports">New report documents</param>
+		/// <param name="newProperties">Property changes. Any properties with a null value will be removed.</param>
+		/// <returns>True if the job was updated, false if it was deleted in the meantime</returns>
+		Task<IJob?> TryUpdateStepAsync(IJob job, IGraph graph, SubResourceId batchId, SubResourceId stepId, JobStepState newState = default, JobStepOutcome newOutcome = default, JobStepError? newError = null, bool? newAbortRequested = null, UserId? newAbortByUserId = null, LogId? newLogId = null, ObjectId? newNotificationTriggerId = null, UserId? newRetryByUserId = null, Priority? newPriority = null, List<Report>? newReports = null, Dictionary<string, string?>? newProperties = null);
+
+		/// <summary>
+		/// Attempts to update the node groups to be executed for a job. Fails if another write happens in the meantime.
+		/// </summary>
+		/// <param name="job">The job to update</param>
+		/// <param name="newGraph">New graph for this job</param>
+		/// <returns>True if the groups were updated to the given list. False if another write happened first.</returns>
+		Task<IJob?> TryUpdateGraphAsync(IJob job, IGraph newGraph);
+
+		/// <summary>
+		/// Removes a job from the dispatch queue. Ignores the state of any batches still remaining to execute. Should only be used to correct for inconsistent state.
+		/// </summary>
+		/// <param name="job">Job</param>
+		/// <returns></returns>
+		Task<IJob?> TryRemoveFromDispatchQueueAsync(IJob job);
+
+		/// <summary>
+		/// Adds an issue to a job
+		/// </summary>
+		/// <param name="jobId">The job id</param>
+		/// <param name="issueId">The issue to add</param>
+		/// <returns>Async task</returns>
+		Task AddIssueToJobAsync(JobId jobId, int issueId);
+
+		/// <summary>
+		/// Gets a queue of jobs to consider for execution
+		/// </summary>
+		/// <returns>Sorted list of jobs to execute</returns>
+		Task<List<IJob>> GetDispatchQueueAsync();
+
+		/// <summary>
+		/// Marks a job as skipped
+		/// </summary>
+		/// <param name="job">The job to update</param>
+		/// <param name="graph">Graph for the job</param>
+		/// <param name="reason">Reason for this batch being failed</param>
+		/// <returns>Updated version of the job</returns>
+		Task<IJob?> SkipAllBatchesAsync(IJob? job, IGraph graph, JobStepBatchError reason);
+
+		/// <summary>
+		/// Marks a batch as skipped
+		/// </summary>
+		/// <param name="job">The job to update</param>
+		/// <param name="batchId">The batch to mark as skipped</param>
+		/// <param name="graph">Graph for the job</param>
+		/// <param name="reason">Reason for this batch being failed</param>
+		/// <returns>Updated version of the job</returns>
+		Task<IJob?> SkipBatchAsync(IJob? job, SubResourceId batchId, IGraph graph, JobStepBatchError reason);
+
+		/// <summary>
+		/// Abort an agent's lease, and update the payload accordingly
+		/// </summary>
+		/// <param name="job">The job to update</param>
+		/// <param name="batchIdx">Index of the batch to cancel</param>
+		/// <param name="graph">Graph for the job</param>
+		/// <param name="reason">Reason for this batch being failed</param>
+		/// <returns>True if the job is updated</returns>
+		Task<IJob?> TryFailBatchAsync(IJob job, int batchIdx, IGraph graph, JobStepBatchError reason);
+
+		/// <summary>
+		/// Attempt to assign a lease to execute a batch
+		/// </summary>
+		/// <param name="job">The job containing the batch</param>
+		/// <param name="batchIdx">Index of the batch</param>
+		/// <param name="poolId">The pool id</param>
+		/// <param name="agentId">New agent to execute the batch</param>
+		/// <param name="sessionId">Session of the agent that is to execute the batch</param>
+		/// <param name="leaseId">The lease unique id</param>
+		/// <param name="logId">Unique id of the log for the batch</param>
+		/// <returns>True if the batch is updated</returns>
+		Task<IJob?> TryAssignLeaseAsync(IJob job, int batchIdx, PoolId poolId, AgentId agentId, SessionId sessionId, LeaseId leaseId, LogId logId);
+
+		/// <summary>
+		/// Cancel a lease reservation on a batch (before it has started)
+		/// </summary>
+		/// <param name="job">The job containing the lease</param>
+		/// <param name="batchIdx">Index of the batch to cancel</param>
+		/// <returns>True if the job is updated</returns>
+		Task<IJob?> TryCancelLeaseAsync(IJob job, int batchIdx);
+
+		/// <summary>
+		/// Upgrade all documents in the collection
+		/// </summary>
+		/// <returns>Async task</returns>
+		Task UpgradeDocumentsAsync();
+	}
+}
diff --git a/Engine/Source/Programs/Horde/Horde.Server/Jobs/IJobStepRef.cs b/Engine/Source/Programs/Horde/Horde.Server/Jobs/IJobStepRef.cs
new file mode 100644
--- /dev/null
+++ b/Engine/Source/Programs/Horde/Horde.Server/Jobs/IJobStepRef.cs
@@ -0,0 +1,193 @@
+// Copyright Epic Games, Inc. All Rights Reserved.
+
+using System;
+using Horde.Server.Agents;
+using Horde.Server.Agents.Pools;
+using Horde.Server.Logs;
+using Horde.Server.Streams;
+using Horde.Server.Utilities;
+using HordeCommon;
+using MongoDB.Bson.Serialization;
+using MongoDB.Bson.Serialization.Attributes;
+
+namespace Horde.Server.Jobs
+{
+	/// <summary>
+	/// Unique id struct for JobStepRef objects. Includes a job id, batch id, and step id to uniquely identify the step.
+	/// </summary>
+	[BsonSerializer(typeof(JobStepRefIdSerializer))]
+	public struct JobStepRefId
+	{
+		/// <summary>
+		/// The job id
+		/// </summary>
+		public JobId JobId { get; set; }
+
+		/// <summary>
+		/// The batch id within the job
+		/// </summary>
+		public SubResourceId BatchId { get; set; }
+
+		/// <summary>
+		/// The step id
+		/// </summary>
+		public SubResourceId StepId { get; set; }
+
+		/// <summary>
+		/// Constructor
+		/// </summary>
+		/// <param name="jobId">The job id</param>
+		/// <param name="batchId">The batch id within the job</param>
+		/// <param name="stepId">The step id</param>
+		public JobStepRefId(JobId jobId, SubResourceId batchId, SubResourceId stepId)
+		{
+			JobId = jobId;
+			BatchId = batchId;
+			StepId = stepId;
+		}
+
+		/// <summary>
+		/// Parse a job step id from a string
+		/// </summary>
+		/// <param name="text">Text to parse</param>
+		/// <returns>The parsed id</returns>
+		public static JobStepRefId Parse(string text)
+		{
+			string[] components = text.Split(':');
+			return new JobStepRefId(JobId.Parse(components[0]), SubResourceId.Parse(components[1]), SubResourceId.Parse(components[2]));
+		}
+
+		/// <summary>
+		/// Formats this id as a string
+		/// </summary>
+		/// <returns>Formatted id</returns>
+		public override string ToString()
+		{
+			return $"{JobId}:{BatchId}:{StepId}";
+		}
+	}
+
+	/// <summary>
+	/// Serializer for JobStepRefId objects
+	/// </summary>
+	public sealed class JobStepRefIdSerializer : IBsonSerializer<JobStepRefId>
+	{
+		/// <inheritdoc/>
+		public Type ValueType => typeof(JobStepRefId);
+
+		/// <inheritdoc/>
+		void IBsonSerializer.Serialize(BsonSerializationContext context, BsonSerializationArgs args, object value)
+		{
+			Serialize(context, args, (JobStepRefId)value);
+		}
+
+		/// <inheritdoc/>
+		object IBsonSerializer.Deserialize(BsonDeserializationContext context, BsonDeserializationArgs args)
+		{
+			return ((IBsonSerializer<JobStepRefId>)this).Deserialize(context, args);
+		}
+
+		/// <inheritdoc/>
+		public JobStepRefId Deserialize(BsonDeserializationContext context, BsonDeserializationArgs args)
+		{
+			return JobStepRefId.Parse(context.Reader.ReadString());
+		}
+
+		/// <inheritdoc/>
+		public void Serialize(BsonSerializationContext context, BsonSerializationArgs args, JobStepRefId value)
+		{
+			context.Writer.WriteString(((JobStepRefId)value).ToString());
+		}
+	}
+
+	/// <summary>
+	/// Searchable reference to a jobstep
+	/// </summary>
+	public interface IJobStepRef
+	{
+		/// <summary>
+		/// Globally unique identifier for the jobstep being referenced
+		/// </summary>
+		public JobStepRefId Id { get; }
+
+		/// <summary>
+		/// Name of the job
+		/// </summary>
+		public string JobName { get; }
+
+		/// <summary>
+		/// Name of the name
+		/// </summary>
+		public string NodeName { get; }
+
+		/// <summary>
+		/// Unique id of the stream containing the job
+		/// </summary>
+		public StreamId StreamId { get; }
+
+		/// <summary>
+		/// Template for the job being executed
+		/// </summary>
+		public TemplateId TemplateId { get; }
+
+		/// <summary>
+		/// The change number being built
+		/// </summary>
+		public int Change { get; }
+
+		/// <summary>
+		/// Log for this step
+		/// </summary>
+		public LogId? LogId { get; }
+
+		/// <summary>
+		/// The agent type
+		/// </summary>
+		public PoolId? PoolId { get; }
+
+		/// <summary>
+		/// The agent id
+		/// </summary>
+		public AgentId? AgentId { get; }
+
+		/// <summary>
+		/// Outcome of the step, once complete.
+		/// </summary>
+		public JobStepOutcome? Outcome { get; }
+
+		/// <summary>
+		/// Whether this step should update issues
+		/// </summary>
+		public bool UpdateIssues { get; }
+
+		/// <summary>
+		/// The last change that succeeded. Note that this is only set when the ref is updated; it is not necessarily consistent with steps run later.
+		/// </summary>
+		public int? LastSuccess { get; }
+
+		/// <summary>
+		/// The last change that succeeded, or completed a warning. See <see cref="LastSuccess"/>.
+		/// </summary>
+		public int? LastWarning { get; }
+
+		/// <summary>
+		/// Time taken for the batch containing this batch to start after it became ready
+		/// </summary>
+		public float BatchWaitTime { get; }
+
+		/// <summary>
+		/// Time taken for this batch to initialize
+		/// </summary>
+		public float BatchInitTime { get; }
+
+		/// <summary>
+		/// Time at which the step started.
+		/// </summary>
+		public DateTime StartTimeUtc { get; }
+
+		/// <summary>
+		/// Time at which the step finished.
+		/// </summary>
+		public DateTime? FinishTimeUtc { get; }
+	}
+}
diff --git a/Engine/Source/Programs/Horde/Horde.Server/Jobs/IJobStepRefCollection.cs b/Engine/Source/Programs/Horde/Horde.Server/Jobs/IJobStepRefCollection.cs
new file mode 100644
--- /dev/null
+++ b/Engine/Source/Programs/Horde/Horde.Server/Jobs/IJobStepRefCollection.cs
@@ -0,0 +1,128 @@
+// Copyright Epic Games, Inc. All Rights Reserved.
+
+using System;
+using System.Collections.Generic;
+using System.Threading.Tasks;
+using Horde.Server.Agents;
+using Horde.Server.Agents.Pools;
+using Horde.Server.Jobs.Graphs;
+using Horde.Server.Logs;
+using Horde.Server.Streams;
+using Horde.Server.Utilities;
+using HordeCommon;
+using Microsoft.Extensions.Logging;
+
+namespace Horde.Server.Jobs
+{
+	/// <summary>
+	/// Interface for a collection of JobStepRef documents
+	/// </summary>
+	public interface IJobStepRefCollection
+	{
+		/// <summary>
+		/// Constructor
+		/// </summary>
+		/// <param name="id">Unique id for the step being referenced</param>
+		/// <param name="jobName">Name of the job</param>
+		/// <param name="stepName">Name of the step</param>
+		/// <param name="streamId">Unique id for the stream containing the job</param>
+		/// <param name="templateId"></param>
+		/// <param name="change">The change number being built</param>
+		/// <param name="logId">The log file id</param>
+		/// <param name="poolId">The pool id</param>
+		/// <param name="agentId">The agent id</param>
+		/// <param name="outcome">Outcome of this step, if known</param>
+		/// <param name="updateIssues">Whether this step ref is included for issue updates</param>
+		/// <param name="lastSuccess">The last change that completed with success</param>
+		/// <param name="lastWarning">The last change that completed with a warning (or success)</param>
+		/// <param name="waitTime">Time taken for the batch containing this step to start</param>
+		/// <param name="initTime">Time taken for the batch containing this step to initializer</param>
+		/// <param name="jobStartTimeUtc">Start time of the job</param>
+		/// <param name="startTimeUtc">Start time</param>
+		/// <param name="finishTimeUtc">Finish time for the step, if known</param>
+		Task<IJobStepRef> InsertOrReplaceAsync(JobStepRefId id, string jobName, string stepName, StreamId streamId, TemplateId templateId, int change, LogId? logId, PoolId? poolId, AgentId? agentId, JobStepOutcome? outcome, bool updateIssues, int? lastSuccess, int? lastWarning, float waitTime, float initTime, DateTime jobStartTimeUtc, DateTime startTimeUtc, DateTime? finishTimeUtc);
+
+		/// <summary>
+		/// Gets the history of a given node
+		/// </summary>
+		/// <param name="streamId">Unique id for a stream</param>
+		/// <param name="templateId"></param>
+		/// <param name="nodeName">Name of the node</param>
+		/// <param name="change">The current change</param>
+		/// <param name="includeFailed">Whether to include failed nodes</param>
+		/// <param name="count">Number of results to return</param>
+		/// <returns>List of step references</returns>
+		Task<List<IJobStepRef>> GetStepsForNodeAsync(StreamId streamId, TemplateId templateId, string nodeName, int? change, bool includeFailed, int count);
+
+		/// <summary>
+		/// Gets the previous job that ran a given step
+		/// </summary>
+		/// <param name="streamId">Id of the stream to search</param>
+		/// <param name="templateId">The template id</param>
+		/// <param name="nodeName">Name of the step to find</param>
+		/// <param name="change">The current changelist number</param>
+		/// <param name="outcome">The outcome to filter by or include all outcomes if null</param>
+		/// <param name="updateIssues">If true, constrain to steps which update issues</param>		 
+		/// <returns>The previous job, or null.</returns>
+		Task<IJobStepRef?> GetPrevStepForNodeAsync(StreamId streamId, TemplateId templateId, string nodeName, int change, JobStepOutcome? outcome = null, bool? updateIssues = null);
+
+		/// <summary>
+		/// Gets the next job that ran a given step
+		/// </summary>
+		/// <param name="streamId">Id of the stream to search</param>
+		/// <param name="templateId">The template id</param>
+		/// <param name="nodeName">Name of the step to find</param>
+		/// <param name="change">The current changelist number</param>
+		/// <param name="outcome">The outcome to filter by or include all outcomes if null</param>
+		/// <param name="updateIssues">If true, constrain to steps which update issues</param>
+		/// <returns>The previous job, or null.</returns>
+		Task<IJobStepRef?> GetNextStepForNodeAsync(StreamId streamId, TemplateId templateId, string nodeName, int change, JobStepOutcome? outcome = null, bool? updateIssues = null);
+	}
+
+	static class JobStepRefCollectionExtensions
+	{
+		public static async Task UpdateAsync(this IJobStepRefCollection jobStepRefs, IJob job, IJobStepBatch batch, IJobStep step, IGraph graph, ILogger? logger = null)
+		{
+			if (job.PreflightChange == 0)
+			{
+				float waitTime = (float)(batch.GetWaitTime() ?? TimeSpan.Zero).TotalSeconds;
+				float initTime = (float)(batch.GetInitTime() ?? TimeSpan.Zero).TotalSeconds;
+
+				string nodeName = graph.Groups[batch.GroupIdx].Nodes[step.NodeIdx].Name;
+				JobStepOutcome? outcome = step.IsPending() ? (JobStepOutcome?)null : step.Outcome;
+
+				int? lastSuccess = null;
+				int? lastWarning = null;
+				if (outcome != JobStepOutcome.Success)
+				{
+					IJobStepRef? prevStep = await jobStepRefs.GetPrevStepForNodeAsync(job.StreamId, job.TemplateId, nodeName, job.Change);
+					if (prevStep != null)
+					{
+						lastSuccess = prevStep.LastSuccess;
+						if (outcome != JobStepOutcome.Warnings)
+						{
+							lastWarning = prevStep.LastWarning;
+						}
+					}
+				}
+
+				if (job.AbortedByUserId != null && outcome == null)
+				{
+					outcome = JobStepOutcome.Unspecified;
+				}
+
+				if (batch.Error != JobStepBatchError.None && outcome == null)
+				{
+					outcome = JobStepOutcome.Unspecified;
+				}
+				
+				if (logger != null)
+				{
+					logger.LogInformation("Updating step reference {StepId} for job {JobId}, batch {BatchId}, with outcome {JobStepOutcome}", step.Id, job.Id, batch.Id, outcome);
+				}				
+
+				await jobStepRefs.InsertOrReplaceAsync(new JobStepRefId(job.Id, batch.Id, step.Id), job.Name, nodeName, job.StreamId, job.TemplateId, job.Change, step.LogId, batch.PoolId, batch.AgentId, outcome, job.UpdateIssues, lastSuccess, lastWarning, waitTime, initTime, job.CreateTimeUtc, step.StartTimeUtc ?? DateTime.UtcNow, step.FinishTimeUtc);
+			}
+		}
+	}
+}
