Horde: Rename the Horde.Build project to Horde.Server.

diff --git a/Engine/Source/Programs/Horde/Horde.Server/Jobs/JobTaskSource.cs b/Engine/Source/Programs/Horde/Horde.Server/Jobs/JobTaskSource.cs
new file mode 100644
--- /dev/null
+++ b/Engine/Source/Programs/Horde/Horde.Server/Jobs/JobTaskSource.cs
@@ -0,0 +1,1122 @@
+// Copyright Epic Games, Inc. All Rights Reserved.
+
+using System;
+using System.Collections.Generic;
+using System.Diagnostics;
+using System.Diagnostics.CodeAnalysis;
+using System.Globalization;
+using System.Linq;
+using System.Text;
+using System.Threading;
+using System.Threading.Tasks;
+using EpicGames.Core;
+using EpicGames.Horde.Common;
+using Google.Protobuf;
+using Google.Protobuf.WellKnownTypes;
+using Horde.Server.Acls;
+using Horde.Server.Agents;
+using Horde.Server.Agents.Leases;
+using Horde.Server.Agents.Pools;
+using Horde.Server.Jobs.Graphs;
+using Horde.Server.Logs;
+using Horde.Server.Perforce;
+using Horde.Server.Server;
+using Horde.Server.Storage;
+using Horde.Server.Streams;
+using Horde.Server.Tasks;
+using Horde.Server.Ugs;
+using Horde.Server.Utilities;
+using HordeCommon;
+using HordeCommon.Rpc.Tasks;
+using Microsoft.Extensions.Hosting;
+using Microsoft.Extensions.Logging;
+using Microsoft.Extensions.Options;
+using MongoDB.Driver;
+
+namespace Horde.Server.Jobs
+{
+	/// <summary>
+	/// Background service to dispatch pending work to agents in priority order.
+	/// </summary>
+	public sealed class JobTaskSource : TaskSourceBase<ExecuteJobTask>, IHostedService, IDisposable
+	{
+		/// <inheritdoc/>
+		public override string Type => "Job";
+
+		/// <inheritdoc/>
+		public override TaskSourceFlags Flags => TaskSourceFlags.None;
+
+		/// <summary>
+		/// An item in the queue to be executed
+		/// </summary>
+		[DebuggerDisplay("{_job.Id}:{Batch.Id} ({_poolId})")]
+		internal class QueueItem
+		{
+			/// <summary>
+			/// The job instance
+			/// </summary>
+			public IJob _job;
+
+			/// <summary>
+			/// Index of the batch within this job to be executed
+			/// </summary>
+			public int _batchIdx;
+
+			/// <summary>
+			/// The pool of machines to allocate from
+			/// </summary>
+			public PoolId _poolId;
+
+			/// <summary>
+			/// The type of workspace that this item should run in
+			/// </summary>
+			public AgentWorkspace _workspace;
+
+			/// <summary>
+			/// Whether or not to use the AutoSDK.
+			/// </summary>
+			public bool _useAutoSdk;
+
+			/// <summary>
+			/// Task for creating a lease and assigning to a waiter
+			/// </summary>
+			public Task? _assignTask;
+
+			/// <summary>
+			/// Accessor for the batch referenced by this item
+			/// </summary>
+			public IJobStepBatch Batch => _job.Batches[_batchIdx];
+
+			/// <summary>
+			/// Returns an identifier describing this unique batch
+			/// </summary>
+			public (JobId, SubResourceId) Id => (_job.Id, Batch.Id);
+
+			/// <summary>
+			/// Constructor
+			/// </summary>
+			/// <param name="job">The job instance</param>
+			/// <param name="batchIdx">The batch index to execute</param>
+			/// <param name="poolId">Unique id of the pool of machines to allocate from</param>
+			/// <param name="workspace">The workspace that this job should run in</param>
+			/// <param name="useAutoSdk">Whether or not to use the AutoSDK</param>
+			public QueueItem(IJob job, int batchIdx, PoolId poolId, AgentWorkspace workspace, bool useAutoSdk)
+			{
+				_job = job;
+				_batchIdx = batchIdx;
+				_poolId = poolId;
+				_workspace = workspace;
+				_useAutoSdk = useAutoSdk;
+			}
+		}
+
+		/// <summary>
+		/// Comparer for items in the queue
+		/// </summary>
+		class QueueItemComparer : IComparer<QueueItem>
+		{
+			/// <summary>
+			/// Compare two items
+			/// </summary>
+			/// <param name="x">First item to compare</param>
+			/// <param name="y">Second item to compare</param>
+			/// <returns>Negative value if X is a higher priority than Y</returns>
+			public int Compare([AllowNull] QueueItem x, [AllowNull] QueueItem y)
+			{
+				if (x == null)
+				{
+					return (y == null) ? 0 : -1;
+				}
+				else if (y == null)
+				{
+					return 1;
+				}
+
+				int delta = y.Batch.SchedulePriority - x.Batch.SchedulePriority;
+				if (delta == 0)
+				{
+					delta = x._job.Id.Id.CompareTo(y._job.Id.Id);
+					if (delta == 0)
+					{
+						delta = (int)x.Batch.Id.Value - (int)y.Batch.Id.Value;
+					}
+				}
+				return delta;
+			}
+		}
+
+		/// <summary>
+		/// Information about an agent waiting for work
+		/// </summary>
+		class QueueWaiter
+		{
+			/// <summary>
+			/// The agent performing the wait
+			/// </summary>
+			public IAgent Agent { get; }
+
+			/// <summary>
+			/// Task to wait for a lease to be assigned
+			/// </summary>
+			public Task<AgentLease?> Task => LeaseSource.Task;
+
+			/// <summary>
+			/// Completion source for the waiting agent. If a new queue item becomes available, the result will be passed through 
+			/// </summary>
+			public TaskCompletionSource<AgentLease?> LeaseSource { get; } = new TaskCompletionSource<AgentLease?>();
+
+			/// <summary>
+			/// Constructor
+			/// </summary>
+			/// <param name="agent">The agent waiting for a task</param>
+			public QueueWaiter(IAgent agent)
+			{
+				Agent = agent;
+			}
+		}
+
+		readonly AclService _aclService;
+		readonly IStreamCollection _streamCollection;
+		readonly ILogFileService _logFileService;
+		readonly IAgentCollection _agentsCollection;
+		readonly IJobCollection _jobs;
+		readonly IJobStepRefCollection _jobStepRefs;
+		readonly IGraphCollection _graphs;
+		readonly IPoolCollection _poolCollection;
+		readonly IUgsMetadataCollection _ugsMetadataCollection;
+		readonly PerforceLoadBalancer _perforceLoadBalancer;
+		readonly IClock _clock;
+		readonly IOptionsMonitor<ServerSettings> _settings;
+		readonly IOptionsMonitor<GlobalConfig> _globalConfig;
+		readonly ILogger<JobTaskSource> _logger;
+		readonly ITicker _ticker;
+
+		// Object used for ensuring mutual exclusion to the queues
+		readonly object _lockObject = new object();
+
+		// List of items waiting to be executed
+		SortedSet<QueueItem> _queue = new SortedSet<QueueItem>(new QueueItemComparer());
+
+		// Map from batch id to the corresponding queue item
+		Dictionary<(JobId, SubResourceId), QueueItem> _batchIdToQueueItem = new Dictionary<(JobId, SubResourceId), QueueItem>();
+
+		// Set of long-poll tasks waiting to be satisfied 
+		readonly HashSet<QueueWaiter> _waiters = new HashSet<QueueWaiter>();
+
+		// During a background queue refresh operation, any updated batches are added to this dictionary for merging into the updated queue.
+		List<QueueItem>? _newQueueItemsDuringUpdate;
+
+		/// <summary>
+		/// Delegate for job schedule events
+		/// </summary>
+		public delegate void JobScheduleEvent(IPool pool, bool hasAgentsOnline, IJob job, IGraph graph, SubResourceId batchId);
+		
+		/// <summary>
+		/// Event triggered when a job is scheduled
+		/// </summary>
+		public event JobScheduleEvent? OnJobScheduled;
+
+		// Interval between querying the database for jobs to execute
+		static readonly TimeSpan s_refreshInterval = TimeSpan.FromSeconds(5.0);
+
+		/// <summary>
+		/// Constructor
+		/// </summary>
+		public JobTaskSource(AclService aclService, IAgentCollection agents, IJobCollection jobs, IJobStepRefCollection jobStepRefs, IGraphCollection graphs, IPoolCollection pools, IUgsMetadataCollection ugsMetadataCollection, IStreamCollection streamCollection, ILogFileService logFileService, PerforceLoadBalancer perforceLoadBalancer, IClock clock, IOptionsMonitor<ServerSettings> settings, IOptionsMonitor<GlobalConfig> globalConfig, ILogger<JobTaskSource> logger)
+		{
+			_aclService = aclService;
+			_agentsCollection = agents;
+			_jobs = jobs;
+			_jobStepRefs = jobStepRefs;
+			_graphs = graphs;
+			_poolCollection = pools;
+			_ugsMetadataCollection = ugsMetadataCollection;
+			_streamCollection = streamCollection;
+			_logFileService = logFileService;
+			_perforceLoadBalancer = perforceLoadBalancer;
+			_clock = clock;
+			_ticker = clock.AddTicker<JobTaskSource>(s_refreshInterval, TickAsync, logger);
+			_globalConfig = globalConfig;
+			_settings = settings;
+			_logger = logger;
+
+			OnLeaseStartedProperties.Add(nameof(ExecuteJobTask.JobId), x => JobId.Parse(x.JobId)).Add(nameof(ExecuteJobTask.BatchId), x => SubResourceId.Parse(x.BatchId)).Add(nameof(ExecuteJobTask.LogId), x => LogId.Parse(x.LogId));
+		}
+
+		/// <inheritdoc/>
+		public Task StartAsync(CancellationToken cancellationToken) => _ticker.StartAsync();
+
+		/// <inheritdoc/>
+		public Task StopAsync(CancellationToken cancellationToken) => _ticker.StopAsync();
+
+		/// <inheritdoc/>
+		public void Dispose() => _ticker.Dispose();
+
+		/// <summary>
+		/// Gets an object containing the stats of the queue for diagnostic purposes.
+		/// </summary>
+		/// <returns>Status object</returns>
+		public object GetStatus()
+		{
+			lock (_lockObject)
+			{
+				List<object> outputItems = new List<object>();
+				foreach (QueueItem queueItem in _queue)
+				{
+					outputItems.Add(new { JobId = queueItem._job.Id.ToString(), BatchId = queueItem.Batch.Id.ToString(), PoolId = queueItem._poolId.ToString(), Workspace = queueItem._workspace });
+				}
+
+				List<object> outputWaiters = new List<object>();
+				foreach (QueueWaiter waiter in _waiters)
+				{
+					outputWaiters.Add(new { Id = waiter.Agent.Id.ToString(), Pools = waiter.Agent.GetPools().Select(x => x.ToString()).ToList(), waiter.Agent.Workspaces });
+				}
+
+				return new { Items = outputItems, Waiters = outputWaiters };
+			}
+		}
+
+		/// <summary>
+		/// Cancel any pending wait for an agent, allowing it to cycle its session state immediately
+		/// </summary>
+		/// <param name="agentId">The agent id</param>
+		public void CancelLongPollForAgent(AgentId agentId)
+		{
+			QueueWaiter? waiter;
+			lock (_lockObject)
+			{
+				waiter = _waiters.FirstOrDefault(x => x.Agent.Id == agentId);
+			}
+			if(waiter != null)
+			{
+				waiter.LeaseSource.TrySetCanceled();
+			}
+		}
+
+		internal class PoolStatus
+		{
+			public readonly IPool Pool;
+			public readonly bool HasAgents;
+			public readonly bool HasEnabledAgents;
+			public readonly bool HasOnlineAgents;
+			public bool IsAutoScaled => Pool.EnableAutoscaling;
+
+			public PoolStatus(IPool pool, bool hasAgents, bool hasEnabledAgents, bool hasOnlineAgents)
+			{
+				Pool = pool;
+				HasAgents = hasAgents;
+				HasEnabledAgents = hasEnabledAgents;
+				HasOnlineAgents = hasOnlineAgents;
+			}
+		}
+
+		/// <summary>
+		/// Calculate status of agents for a pool
+		/// </summary>
+		/// <param name="utcNow">Current time</param>
+		/// <param name="pools">List of all available pools</param>
+		/// <param name="agents">List of all available agents</param>
+		/// <returns></returns>
+		internal static Dictionary<PoolId, PoolStatus> GetPoolStatus(DateTime utcNow, List<IPool> pools, List<IAgent> agents)
+		{
+			Dictionary<PoolId, PoolStatus> poolStatus = new ();
+
+			foreach (IPool pool in pools)
+			{
+				Condition poolCondition = pool.Condition ?? Condition.Parse("false");
+				List<IAgent> poolAgents = agents.Where(x => x.ExplicitPools.Contains(pool.Id) || x.SatisfiesCondition(poolCondition)).ToList();
+
+				bool hasAgents = poolAgents.Count > 0;
+				bool hasEnabledAgents = poolAgents.Any(x => x.Enabled);
+				bool hasOnlineAgents = poolAgents.Any(x => x.Enabled && x.IsSessionValid(utcNow));
+
+				poolStatus[pool.Id] = new PoolStatus(pool, hasAgents, hasEnabledAgents, hasOnlineAgents);
+			}
+
+			return poolStatus;
+		}
+		
+		/// <summary>
+		/// Background task
+		/// </summary>
+		/// <param name="stoppingToken">Token that indicates that the service should shut down</param>
+		/// <returns>Async task</returns>
+		internal async ValueTask TickAsync(CancellationToken stoppingToken)
+		{
+			// Set the NewBatchIdToQueueItem member, so we capture any updated jobs during the DB query.
+			lock (_lockObject)
+			{
+				_newQueueItemsDuringUpdate = new List<QueueItem>();
+			}
+
+			// Query all the current streams
+			GlobalConfig globalConfig = _globalConfig.CurrentValue;
+			List<IStream> streamsList = await _streamCollection.GetAsync(globalConfig.Streams);
+			Dictionary<StreamId, IStream> streams = streamsList.ToDictionary(x => x.Id, x => x);
+
+			// Find all the pools which are valid (ie. have at least one online agent)
+			List<IAgent> agents = await _agentsCollection.FindAsync();
+			List<IPool> pools = await _poolCollection.GetAsync();
+			Dictionary<PoolId, PoolStatus> poolStatus = GetPoolStatus(_clock.UtcNow, pools, agents);
+
+			// New list of queue items
+			SortedSet<QueueItem> newQueue = new SortedSet<QueueItem>(_queue.Comparer);
+			Dictionary<(JobId, SubResourceId), QueueItem> newBatchIdToQueueItem = new Dictionary<(JobId, SubResourceId), QueueItem>();
+
+			// Query for a new list of jobs for the queue
+			List<IJob> newJobs = await _jobs.GetDispatchQueueAsync();
+			for (int idx = 0; idx < newJobs.Count; idx++)
+			{
+				IJob? newJob = newJobs[idx];
+
+				if (newJob.GraphHash == null)
+				{
+					_logger.LogError("Job {JobId} has a null graph hash and can't be started.", newJob.Id);
+					await _jobs.TryRemoveFromDispatchQueueAsync(newJob);
+					continue;
+				}
+				if (newJob.AbortedByUserId != null)
+				{
+					_logger.LogError("Job {JobId} was aborted but not removed from dispatch queue", newJob.Id);
+					await _jobs.TryRemoveFromDispatchQueueAsync(newJob);
+					continue;
+				}
+
+				// Get the graph for this job
+				IGraph graph = await _graphs.GetAsync(newJob.GraphHash);
+
+				// Get the stream. If it fails, skip the whole job.
+				IStream? stream;
+				if (!streams.TryGetValue(newJob.StreamId, out stream))
+				{
+					newJob = await _jobs.SkipAllBatchesAsync(newJob, graph, JobStepBatchError.UnknownStream);
+					continue;
+				}
+
+				// Update all the batches
+				HashSet<SubResourceId> checkedBatchIds = new HashSet<SubResourceId>();
+				while (newJob != null)
+				{
+					// Find the next batch within this job that is ready
+					int batchIdx = newJob.Batches.FindIndex(x => x.State == JobStepBatchState.Ready && checkedBatchIds.Add(x.Id));
+					if (batchIdx == -1)
+					{
+						break;
+					}
+
+					// Validate the agent type and workspace settings
+					IJobStepBatch batch = newJob.Batches[batchIdx];
+					if (!stream.Config.AgentTypes.TryGetValue(graph.Groups[batch.GroupIdx].AgentType, out AgentConfig? agentType))
+					{
+						newJob = await SkipBatchAsync(newJob, batch.Id, graph, JobStepBatchError.UnknownAgentType);
+					}
+					else if (!poolStatus.ContainsKey(agentType.Pool))
+					{
+						newJob = await SkipBatchAsync(newJob, batch.Id, graph, JobStepBatchError.UnknownPool);
+					}
+					else if (!poolStatus[agentType.Pool].HasAgents)
+					{
+						newJob = await SkipBatchAsync(newJob, batch.Id, graph, JobStepBatchError.NoAgentsInPool);
+					}
+					else if (!stream.Config.TryGetAgentWorkspace(agentType, out (AgentWorkspace, bool)? workspaceResult))
+					{
+						newJob = await SkipBatchAsync(newJob, batch.Id, graph, JobStepBatchError.UnknownWorkspace);
+					}
+					else
+					{
+						ITemplateRef? templateRef;
+						if (stream.Templates.TryGetValue(newJob.TemplateId, out templateRef))
+						{							
+							if (templateRef.StepStates != null)
+							{
+								for (int i = 0; i < templateRef.StepStates.Count; i++)
+								{
+									ITemplateStep state = templateRef.StepStates[i];
+
+									IJobStep? step = batch.Steps.FirstOrDefault(x => graph.Groups[batch.GroupIdx].Nodes[x.NodeIdx].Name.Equals(state.Name, StringComparison.Ordinal));
+									if (step != null)
+									{
+										JobId jobId = newJob.Id;
+										newJob = await _jobs.TryUpdateStepAsync(newJob, graph, batch.Id, step.Id, JobStepState.Skipped, newError: JobStepError.Paused);
+										if (newJob == null)
+										{
+											_logger.LogError("Job {JobId} failed to update step {StepName} pause state", jobId, state.Name);
+											break;
+										}
+										else
+										{
+											_logger.LogInformation("Job {JobId} step {StepName} has been skipped due to being paused", jobId, state.Name);
+										}
+									}
+								}
+							}
+						}
+
+						if (newJob != null)
+						{
+							(AgentWorkspace workspace, bool useAutoSdk) = workspaceResult.Value;
+							QueueItem newQueueItem = new QueueItem(newJob, batchIdx, agentType.Pool, workspace, useAutoSdk);
+							newQueue.Add(newQueueItem);
+							newBatchIdToQueueItem[(newJob.Id, batch.Id)] = newQueueItem;
+
+							IPool? newJobPool = pools.Find(p => p.Id == agentType.Pool);
+							if (newJobPool != null)
+							{
+								OnJobScheduled?.Invoke(newJobPool, poolStatus[agentType.Pool].HasOnlineAgents, newJob, graph, batch.Id);
+							}
+						}
+					}
+				}
+
+				if (newJob != null)
+				{
+					if (!newJob.Batches.Any(batch => batch.State == JobStepBatchState.Ready || batch.State == JobStepBatchState.Starting || batch.State == JobStepBatchState.Running || batch.State == JobStepBatchState.Stopping))
+					{
+						_logger.LogError("Job {JobId} is in dispatch queue but not currently executing", newJob.Id);
+						await _jobs.TryRemoveFromDispatchQueueAsync(newJob);
+					}
+				}
+			}
+
+			// Update the queue
+			lock (_lockObject)
+			{
+				_queue = newQueue;
+				_batchIdToQueueItem = newBatchIdToQueueItem;
+
+				// Merge the new queue items with the queue
+				foreach (QueueItem newQueueItem in _newQueueItemsDuringUpdate)
+				{
+					QueueItem? existingQueueItem;
+					if (!newBatchIdToQueueItem.TryGetValue((newQueueItem._job.Id, newQueueItem.Batch.Id), out existingQueueItem))
+					{
+						// Always just add this item
+						_queue.Add(newQueueItem);
+						_batchIdToQueueItem[newQueueItem.Id] = newQueueItem;
+					}
+					else if (newQueueItem._job.UpdateIndex > existingQueueItem._job.UpdateIndex)
+					{
+						// Replace the existing item
+						_queue.Remove(existingQueueItem);
+						_queue.Add(newQueueItem);
+						_batchIdToQueueItem[newQueueItem.Id] = newQueueItem;
+					}
+				}
+
+				// Clear out the list to capture queue items during an update
+				_newQueueItemsDuringUpdate = null;
+			}
+		}
+
+		private async Task<IJob?> SkipBatchAsync(IJob job, SubResourceId batchId, IGraph graph, JobStepBatchError reason)
+		{
+			_logger.LogInformation("Skipping batch {BatchId} for job {JobId} (reason: {Reason})", batchId, job.Id, reason);
+
+			IReadOnlyList<(LabelState, LabelOutcome)> oldLabelStates = job.GetLabelStates(graph);
+			IJob? newJob = await _jobs.SkipBatchAsync(job, batchId, graph, reason);
+			if(newJob != null)
+			{
+				IReadOnlyList<(LabelState, LabelOutcome)> newLabelStates = newJob.GetLabelStates(graph);
+				await UpdateUgsBadges(newJob, graph, oldLabelStates, newLabelStates);
+			}
+			return newJob;
+		}
+
+		/// <summary>
+		/// Get the queue, for internal testing only
+		/// </summary>
+		/// <returns>A copy of the queue</returns>
+		internal SortedSet<QueueItem> GetQueueForTesting()
+		{
+			lock (_lockObject)
+			{
+				return new SortedSet<QueueItem>(_queue);
+			}
+		}
+
+		void AssignAnyQueueItemToWaiter(QueueWaiter waiter)
+		{
+			lock (_waiters)
+			{
+				foreach(QueueItem item in _batchIdToQueueItem.Values)
+				{
+					if (TryAssignItemToWaiter(item, waiter))
+					{
+						break;
+					}
+				}
+			}
+		}
+
+		/// <summary>
+		/// Attempt to find a waiter that can handle the given queue item
+		/// </summary>
+		/// <param name="item">The queue item</param>
+		/// <returns></returns>
+		[SuppressMessage("Maintainability", "CA1508:Avoid dead conditional code", Justification = "<Pending>")]
+		void AssignQueueItemToAnyWaiter(QueueItem item)
+		{
+			if (item._assignTask == null && item.Batch.SessionId == null)
+			{
+				lock (_waiters)
+				{
+					if (item._assignTask == null && item.Batch.SessionId == null)
+					{
+						foreach (QueueWaiter waiter in _waiters)
+						{
+							if (TryAssignItemToWaiter(item, waiter))
+							{
+								break;
+							}
+						}
+					}
+				}
+			}
+		}
+
+		/// <summary>
+		/// Attempts to assign a queue item to an agent waiting for work
+		/// </summary>
+		/// <param name="item"></param>
+		/// <param name="waiter"></param>
+		/// <returns></returns>
+		bool TryAssignItemToWaiter(QueueItem item, QueueWaiter waiter)
+		{
+			if (item._assignTask == null && item.Batch.SessionId == null && waiter.Agent.Enabled && waiter.Agent.Leases.Count == 0 && waiter.Agent.IsInPool(item._poolId))
+			{
+				Task startTask = new Task<Task>(() => TryCreateLeaseAsync(item, waiter));
+				Task executeTask = startTask.ContinueWith(task => task, TaskScheduler.Default);
+
+				if (Interlocked.CompareExchange(ref item._assignTask, executeTask, null) == null)
+				{
+					startTask.Start(TaskScheduler.Default);
+					return true;
+				}
+			}
+			return false;
+		}
+
+		/// <summary>
+		/// Updates the current state of a job
+		/// </summary>
+		/// <param name="job">The job that has been updated</param>
+		/// <param name="graph">Graph for the job</param>
+		public void UpdateQueuedJob(IJob job, IGraph graph)
+		{
+			StreamConfig? streamConfig;
+			_globalConfig.CurrentValue.TryGetStream(job.StreamId, out streamConfig);
+
+			List<TaskCompletionSource<bool>> completeWaiters = new List<TaskCompletionSource<bool>>();
+			lock (_lockObject)
+			{
+				for (int batchIdx = 0; batchIdx < job.Batches.Count; batchIdx++)
+				{
+					IJobStepBatch batch = job.Batches[batchIdx];
+					if (batch.State == JobStepBatchState.Ready && streamConfig != null && batch.AgentId == null)
+					{
+						// Check if this item is already in the list.
+						QueueItem? existingItem;
+						if (_batchIdToQueueItem.TryGetValue((job.Id, batch.Id), out existingItem))
+						{
+							// Make sure this is a newer version of the job. There's no guarantee that this is the latest revision.
+							if (job.UpdateIndex > existingItem._job.UpdateIndex)
+							{
+								if (batch.SchedulePriority == existingItem.Batch.SchedulePriority)
+								{
+									existingItem._job = job;
+									existingItem._batchIdx = batchIdx;
+								}
+								else
+								{
+									RemoveQueueItem(existingItem);
+									InsertQueueItem(job, batchIdx, existingItem._poolId, existingItem._workspace, existingItem._useAutoSdk);
+								}
+							}
+							continue;
+						}
+
+						// Get the group being executed by this batch
+						INodeGroup group = graph.Groups[batch.GroupIdx];
+
+						// Get the requirements for the new queue item
+						AgentConfig? agentType;
+						if (streamConfig.AgentTypes.TryGetValue(group.AgentType, out agentType))
+						{
+							(AgentWorkspace, bool)? result;
+							if (streamConfig.TryGetAgentWorkspace(agentType, out result))
+							{
+								(AgentWorkspace agentWorkspace, bool useAutoSdk) = result.Value;
+								InsertQueueItem(job, batchIdx, agentType.Pool, agentWorkspace, useAutoSdk);
+							}
+						}
+					}
+					else
+					{
+						// Check if this item is already in the list. Remove it if it is.
+						QueueItem? existingItem;
+						if (_batchIdToQueueItem.TryGetValue((job.Id, batch.Id), out existingItem))
+						{
+							if (job.UpdateIndex > existingItem._job.UpdateIndex)
+							{
+								RemoveQueueItem(existingItem);
+							}
+						}
+					}
+				}
+			}
+
+			// Awake all the threads that have been assigned new work items. Has do be done outside the lock to prevent continuations running within it (see Waiter.CompletionSource for more info).
+			foreach (TaskCompletionSource<bool> completeWaiter in completeWaiters)
+			{
+				completeWaiter.TrySetResult(true);
+			}
+		}
+
+		/// <inheritdoc/>
+		public override Task<Task<AgentLease?>> AssignLeaseAsync(IAgent agent, CancellationToken cancellationToken)
+		{
+			QueueWaiter waiter = new QueueWaiter(agent);
+			lock (_lockObject)
+			{
+				AssignAnyQueueItemToWaiter(waiter);
+				if (waiter.LeaseSource.Task.TryGetResult(out AgentLease? result))
+				{
+					if (result == null)
+					{
+						return Task.FromResult(Skip(cancellationToken));
+					}
+					return Task.FromResult(Lease(result));
+				}
+				_waiters.Add(waiter);
+			}
+			return Task.FromResult(WaitForLeaseAsync(waiter, cancellationToken));
+		}
+
+		private async Task<AgentLease?> WaitForLeaseAsync(QueueWaiter waiter, CancellationToken cancellationToken)
+		{
+			try
+			{
+				using (cancellationToken.Register(() => waiter.LeaseSource.TrySetResult(null)))
+				{
+					return await waiter.Task;
+				}
+			}
+			finally
+			{
+				lock (_lockObject)
+				{
+					_waiters.Remove(waiter);
+				}
+			}
+		}
+
+		/// <inheritdoc/>
+		public override Task CancelLeaseAsync(IAgent agent, LeaseId leaseId, ExecuteJobTask task)
+		{
+			return CancelLeaseAsync(agent, JobId.Parse(task.JobId), task.BatchId.ToSubResourceId());
+		}
+
+		/// <summary>
+		/// Assign a new batch to be executed by the given agent
+		/// </summary>
+		/// <param name="item">The item to create a lease for</param>
+		/// <param name="waiter">The agent waiting for work</param>
+		/// <returns>New work to execute</returns>
+		private async Task<AgentLease?> TryCreateLeaseAsync(QueueItem item, QueueWaiter waiter)
+		{
+			IJob job = item._job;
+			IJobStepBatch batch = item.Batch;
+			IAgent agent = waiter.Agent;
+			_logger.LogDebug("Assigning job {JobId}, batch {BatchId} to waiter (agent {AgentId})", job.Id, batch.Id, agent.Id);
+
+			// Generate a new unique id for the lease
+			LeaseId leaseId = LeaseId.GenerateNewId();
+
+			// The next time to try assigning to another agent
+			DateTime backOffTime = DateTime.UtcNow + TimeSpan.FromMinutes(1.0);
+
+			// Allocate a log ID but hold off creating the actual log file until the lease has been accepted
+			LogId logId = LogId.GenerateNewId();
+			
+			// Try to update the job with this agent id
+			IJob? newJob = await _jobs.TryAssignLeaseAsync(item._job, item._batchIdx, item._poolId, agent.Id, agent.SessionId!.Value, leaseId, logId);
+			if (newJob != null)
+			{
+				job = newJob;
+
+				StreamConfig? streamConfig;
+				if (!_globalConfig.CurrentValue.TryGetStream(newJob.StreamId, out streamConfig))
+				{
+					return null;
+				}
+
+				// Get the lease name
+				StringBuilder leaseName = new StringBuilder($"{streamConfig.Name} - ");
+				if (job.PreflightChange > 0)
+				{
+					leaseName.Append((job.Change > 0) ? $"Preflight CL {job.PreflightChange} against CL {job.Change}" : $"Preflight CL {job.PreflightChange} against latest");
+				}
+				else
+				{
+					leaseName.Append((job.Change > 0) ? $"CL {job.Change}" : "Latest CL");
+				}
+				leaseName.Append(CultureInfo.InvariantCulture, $" - {job.Name}");
+
+				// Encode the payload
+				ExecuteJobTask? task = await CreateExecuteJobTaskAsync(leaseId, streamConfig, job, batch, agent, item._workspace, item._useAutoSdk, logId);
+				if (task != null)
+				{
+					byte[] payload = Any.Pack(task).ToByteArray();
+
+					// Create the lease and try to set it on the waiter. If this fails, the waiter has already moved on, and the lease can be cancelled.
+					AgentLease lease = new AgentLease(leaseId, leaseName.ToString(), job.StreamId, item._poolId, logId, LeaseState.Pending, null, true, payload);
+					if (waiter.LeaseSource.TrySetResult(lease))
+					{
+						_logger.LogDebug("Assigned lease {LeaseId} to agent {AgentId}", leaseId, agent.Id);
+						await _logFileService.CreateLogFileAsync(job.Id, leaseId, agent.SessionId, LogType.Json, job.JobOptions?.UseNewLogStorage ?? false, logId);
+						return lease;
+					}
+				}
+
+				// Cancel the lease
+				_logger.LogDebug("Unable to assign lease {LeaseId} to agent {AgentId}, cancelling", leaseId, agent.Id);
+				await CancelLeaseAsync(waiter.Agent, job.Id, batch.Id);
+			}
+			else
+			{
+				// Unable to assign job
+				_logger.LogDebug("Failed to assign job {JobId}, batch {BatchId} to agent {AgentId}. Refreshing queue entries.", job.Id, batch.Id, agent.Id);
+
+				// Get the new copy of the job
+				newJob = await _jobs.GetAsync(job.Id);
+				if (newJob == null)
+				{
+					lock (_lockObject)
+					{
+						List<QueueItem> removeItems = _queue.Where(x => x._job == job).ToList();
+						foreach (QueueItem removeItem in removeItems)
+						{
+							RemoveQueueItem(removeItem);
+						}
+					}
+				}
+				else
+				{
+					IGraph graph = await _graphs.GetAsync(newJob.GraphHash);
+					UpdateQueuedJob(newJob, graph);
+				}
+			}
+
+			// Clear out the assignment for this item, and try to reassign it
+			item._assignTask = null;
+			AssignQueueItemToAnyWaiter(item);
+			return null;
+		}
+
+		async Task<ExecuteJobTask?> CreateExecuteJobTaskAsync(LeaseId leaseId, StreamConfig streamConfig, IJob job, IJobStepBatch batch, IAgent agent, AgentWorkspace workspace, bool useAutoSdk, LogId logId)
+		{
+			// Get the lease name
+			StringBuilder leaseName = new StringBuilder($"{streamConfig.Name} - ");
+			if (job.PreflightChange > 0)
+			{
+				leaseName.Append((job.Change > 0) ? $"Preflight CL {job.PreflightChange} against CL {job.Change}" : $"Preflight CL {job.PreflightChange} against latest");
+			}
+			else
+			{
+				leaseName.Append((job.Change > 0) ? $"CL {job.Change}" : "Latest CL");
+			}
+			leaseName.Append(CultureInfo.InvariantCulture, $" - {job.Name}");
+
+			// Get the global settings
+			GlobalConfig globalConfig = _globalConfig.CurrentValue;
+
+			// Create a bearer token for the job executor
+			List<AclClaimConfig> claims = new List<AclClaimConfig>();
+			claims.Add(new AclClaimConfig(HordeClaimTypes.Lease, leaseId.ToString()));
+
+			string storagePrefix = $"{job.StreamId}/{job.Change}-{job.Id}";
+			claims.Add(new AclClaimConfig(HordeClaimTypes.ReadNamespace, $"{Namespace.Artifacts}:{storagePrefix}"));
+			claims.Add(new AclClaimConfig(HordeClaimTypes.WriteNamespace, $"{Namespace.Artifacts}:{storagePrefix}"));
+		
+			// Encode the payload
+			ExecuteJobTask task = new ExecuteJobTask();
+			task.JobId = job.Id.ToString();
+			task.BatchId = batch.Id.ToString();
+			task.LogId = logId.ToString();
+			task.JobName = leaseName.ToString();
+			task.JobOptions = job.JobOptions;
+			task.NamespaceId = Namespace.Artifacts.ToString();
+			task.StoragePrefix = storagePrefix;
+			task.Token = await _aclService.IssueBearerTokenAsync(claims, null);
+
+			List<HordeCommon.Rpc.Messages.AgentWorkspace> workspaces = new ();
+
+			PerforceCluster? cluster = globalConfig.FindPerforceCluster(workspace.Cluster);
+			if (cluster == null)
+			{
+				return null;
+			}
+
+			AgentWorkspace? autoSdkWorkspace = useAutoSdk ? agent.GetAutoSdkWorkspace(cluster) : null;
+			if (autoSdkWorkspace != null)
+			{
+				autoSdkWorkspace.Method = workspace.Method;
+				
+				if (!await agent.TryAddWorkspaceMessage(autoSdkWorkspace, cluster, _perforceLoadBalancer, workspaces))
+				{
+					return null;
+				}
+				task.AutoSdkWorkspace = workspaces.Last();
+			}
+
+			if (!await agent.TryAddWorkspaceMessage(workspace, cluster, _perforceLoadBalancer, workspaces))
+			{
+				return null;
+			}
+
+			task.Workspace = workspaces.Last();
+
+			return task;
+		}
+
+		/// <summary>
+		/// Send any badge updates for this job
+		/// </summary>
+		/// <param name="job">The job being updated</param>
+		/// <param name="graph">Graph for the job</param>
+		/// <param name="oldLabelStates">Previous badge states for the job</param>
+		/// <returns>Async task</returns>
+		public async Task UpdateUgsBadges(IJob job, IGraph graph, IReadOnlyList<(LabelState, LabelOutcome)> oldLabelStates)
+		{
+			await UpdateUgsBadges(job, graph, oldLabelStates, job.GetLabelStates(graph));
+		}
+
+		/// <summary>
+		/// Send any badge updates for this job
+		/// </summary>
+		/// <param name="job">The job being updated</param>
+		/// <param name="graph">Graph for the job</param>
+		/// <param name="oldLabelStates">Previous badge states for the job</param>
+		/// <param name="newLabelStates">The new badge states for the job</param>
+		/// <returns>Async task</returns>
+		public async Task UpdateUgsBadges(IJob job, IGraph graph, IReadOnlyList<(LabelState, LabelOutcome)> oldLabelStates, IReadOnlyList<(LabelState, LabelOutcome)> newLabelStates)
+		{
+			if (!job.ShowUgsBadges || job.PreflightChange != 0)
+			{
+				return;
+			}
+
+			IReadOnlyDictionary<int, UgsBadgeState> oldStates = job.GetUgsBadgeStates(graph, oldLabelStates);
+			IReadOnlyDictionary<int, UgsBadgeState> newStates = job.GetUgsBadgeStates(graph, newLabelStates);
+
+			// Figure out a list of all the badges that have been modified
+			List<int> updateLabels = new List<int>();
+			foreach (KeyValuePair<int, UgsBadgeState> pair in oldStates)
+			{
+				if (!newStates.ContainsKey(pair.Key))
+				{
+					updateLabels.Add(pair.Key);
+				}
+			}
+			foreach (KeyValuePair<int, UgsBadgeState> pair in newStates)
+			{
+				if (!oldStates.TryGetValue(pair.Key, out UgsBadgeState prevState) || prevState != pair.Value)
+				{
+					updateLabels.Add(pair.Key);
+				}
+			}
+
+			// Cached stream for this job
+			StreamConfig? streamConfig = null;
+
+			// Send all the updates
+			Dictionary<int, IUgsMetadata> metadataCache = new Dictionary<int, IUgsMetadata>();
+			foreach (int labelIdx in updateLabels)
+			{
+				ILabel label = graph.Labels[labelIdx];
+
+				// Skip if this label has no UGS name.
+				if (label.UgsName == null)
+				{
+					continue;
+				}
+
+				// Get the new state
+				if (!newStates.TryGetValue(labelIdx, out UgsBadgeState newState))
+				{
+					newState = UgsBadgeState.Skipped;
+				}
+
+				// Get the stream
+				if (streamConfig == null)
+				{
+					if (!_globalConfig.CurrentValue.TryGetStream(job.StreamId, out streamConfig))
+					{
+						_logger.LogError("Unable to fetch definition for stream {StreamId}", job.StreamId);
+						break;
+					}
+				}
+
+				// The changelist number to display the badge for
+				int change;
+				if (label.Change == LabelChange.Code)
+				{
+					change = job.CodeChange;
+				}
+				else
+				{
+					change = job.Change;
+				}
+
+				// Get the current metadata state
+				IUgsMetadata? metadata;
+				if (!metadataCache.TryGetValue(change, out metadata))
+				{
+					metadata = await _ugsMetadataCollection.FindOrAddAsync(streamConfig.Name, change, label.UgsProject);
+					metadataCache[change] = metadata;
+				}
+
+				// Apply the update
+				Uri labelUrl = new Uri(_settings.CurrentValue.DashboardUrl, $"job/{job.Id}?label={labelIdx}");
+				_logger.LogInformation("Updating state of badge {BadgeName} at {Change} to {NewState} ({LabelUrl})", label.UgsName, change, newState, labelUrl);
+				metadata = await _ugsMetadataCollection.UpdateBadgeAsync(metadata, label.UgsName!, labelUrl, newState);
+				metadataCache[change] = metadata;
+			}
+		}
+
+		/// <inheritdoc/>
+		public override async Task OnLeaseFinishedAsync(IAgent agent, LeaseId leaseId, ExecuteJobTask task, LeaseOutcome outcome, ReadOnlyMemory<byte> output, ILogger logger)
+		{
+			await base.OnLeaseFinishedAsync(agent, leaseId, task, outcome, output, logger);
+
+			if (outcome != LeaseOutcome.Success)
+			{
+				AgentId agentId = agent.Id;
+				JobId jobId = JobId.Parse(task.JobId);
+				SubResourceId batchId = task.BatchId.ToSubResourceId();
+
+				// Update the batch
+				for (; ; )
+				{
+					IJob? job = await _jobs.GetAsync(jobId);
+					if (job == null)
+					{
+						break;
+					}
+
+					int batchIdx = job.Batches.FindIndex(x => x.Id == batchId);
+					if (batchIdx == -1)
+					{
+						break;
+					}
+
+					IJobStepBatch batch = job.Batches[batchIdx];
+					if (batch.AgentId != agentId)
+					{
+						break;
+					}
+
+					int runningStepIdx = batch.Steps.FindIndex(x => x.State == JobStepState.Running);
+
+					JobStepBatchError error;
+					if (outcome == LeaseOutcome.Cancelled)
+					{
+						error = JobStepBatchError.Cancelled;
+					}
+					else
+					{
+						error = JobStepBatchError.ExecutionError;
+					}
+					
+					IGraph graph = await _graphs.GetAsync(job.GraphHash);
+					job = await _jobs.TryFailBatchAsync(job, batchIdx, graph, error);
+
+					if (job != null)
+					{
+						if (batch.Error != JobStepBatchError.None)
+						{
+							logger.LogInformation("Failed lease {LeaseId}, job {JobId}, batch {BatchId} with error {Error}", leaseId, job.Id, batch.Id, batch.Error);
+						}
+						if (runningStepIdx != -1)
+						{
+							await _jobStepRefs.UpdateAsync(job, batch, batch.Steps[runningStepIdx], graph, logger);
+						}
+						break;
+					}
+				}
+			}
+		}
+
+		/// <summary>
+		/// 
+		/// </summary>
+		/// <param name="agent"></param>
+		/// <param name="jobId"></param>
+		/// <param name="batchId"></param>
+		/// <returns></returns>
+		async Task CancelLeaseAsync(IAgent agent, JobId jobId, SubResourceId batchId)
+		{
+			_logger.LogDebug("Cancelling lease for job {JobId}, batch {BatchId}", jobId, batchId);
+
+			// Update the batch
+			for (; ; )
+			{
+				IJob? job = await _jobs.GetAsync(jobId);
+				if (job == null)
+				{
+					break;
+				}
+
+				int batchIdx = job.Batches.FindIndex(x => x.Id == batchId);
+				if (batchIdx == -1)
+				{
+					break;
+				}
+
+				IJobStepBatch batch = job.Batches[batchIdx];
+				if (batch.AgentId != agent.Id)
+				{
+					break;
+				}
+
+				IJob? newJob = await _jobs.TryCancelLeaseAsync(job, batchIdx);
+				if (newJob != null)
+				{
+					break;
+				}
+			}
+		}
+
+		/// <summary>
+		/// Inserts an item into the queue
+		/// </summary>
+		/// <param name="job"></param>
+		/// <param name="batchIdx"></param>
+		/// <param name="poolId">The pool to use</param>
+		/// <param name="workspace">The workspace for this item to run in</param>
+		/// <param name="useAutoSdk">Whether or not to use the AutoSDK</param>
+		/// <returns></returns>
+		void InsertQueueItem(IJob job, int batchIdx, PoolId poolId, AgentWorkspace workspace, bool useAutoSdk)
+		{
+			_logger.LogDebug("Adding queued job {JobId}, batch {BatchId} [Pool: {Pool}, Workspace: {Workspace}]", job.Id, job.Batches[batchIdx].Id, poolId, workspace.Identifier);
+
+			QueueItem newItem = new QueueItem(job, batchIdx, poolId, workspace, useAutoSdk);
+			_batchIdToQueueItem[newItem.Id] = newItem;
+			_queue.Add(newItem);
+
+			AssignQueueItemToAnyWaiter(newItem);
+		}
+
+		/// <summary>
+		/// Removes an item from the queue
+		/// </summary>
+		/// <param name="item">Item to remove</param>
+		void RemoveQueueItem(QueueItem item)
+		{
+			_logger.LogDebug("Removing queued job {JobId}, batch {BatchId}", item._job.Id, item.Batch.Id);
+
+			_queue.Remove(item);
+			_batchIdToQueueItem.Remove(item.Id);
+		}
+	}
+}
diff --git a/Engine/Source/Programs/Horde/Horde.Server/Jobs/JobsController.cs b/Engine/Source/Programs/Horde/Horde.Server/Jobs/JobsController.cs
new file mode 100644
--- /dev/null
+++ b/Engine/Source/Programs/Horde/Horde.Server/Jobs/JobsController.cs
@@ -0,0 +1,1622 @@
+// Copyright Epic Games, Inc. All Rights Reserved.
+
+using System;
+using System.Collections.Generic;
+using System.Linq;
+using System.Threading;
+using System.Threading.Tasks;
+using EpicGames.Core;
+using Horde.Server.Acls;
+using Horde.Server.Agents;
+using Horde.Server.Jobs.Artifacts;
+using Horde.Server.Jobs.Graphs;
+using Horde.Server.Jobs.Templates;
+using Horde.Server.Jobs.Timing;
+using Horde.Server.Logs;
+using Horde.Server.Notifications;
+using Horde.Server.Perforce;
+using Horde.Server.Server;
+using Horde.Server.Streams;
+using Horde.Server.Users;
+using Horde.Server.Utilities;
+using HordeCommon;
+using Microsoft.AspNetCore.Authorization;
+using Microsoft.AspNetCore.Mvc;
+using Microsoft.Extensions.Logging;
+using Microsoft.Extensions.Options;
+using MongoDB.Bson;
+using OpenTracing;
+using OpenTracing.Util;
+
+namespace Horde.Server.Jobs
+{
+	/// <summary>
+	/// Controller for the /api/v1/jobs endpoing
+	/// </summary>
+	[ApiController]
+	[Authorize]
+	[Route("[controller]")]
+	public class JobsController : HordeControllerBase
+	{
+		private readonly IGraphCollection _graphs;
+		private readonly ICommitService _commitService;
+		private readonly IPerforceService _perforce;
+		private readonly IStreamCollection _streamCollection;
+		private readonly JobService _jobService;
+		private readonly ITemplateCollection _templateCollection;
+		private readonly IArtifactCollectionV1 _artifactCollection;
+		private readonly IUserCollection _userCollection;
+		private readonly INotificationService _notificationService;
+		private readonly AgentService _agentService;
+		private readonly IOptionsSnapshot<GlobalConfig> _globalConfig;
+		private readonly ILogger<JobsController> _logger;
+
+		/// <summary>
+		/// Constructor
+		/// </summary>
+		public JobsController(IGraphCollection graphs, ICommitService commitService, IPerforceService perforce, IStreamCollection streamCollection, JobService jobService, ITemplateCollection templateCollection, IArtifactCollectionV1 artifactCollection, IUserCollection userCollection, INotificationService notificationService, AgentService agentService, IOptionsSnapshot<GlobalConfig> globalConfig, ILogger<JobsController> logger)
+		{
+			_graphs = graphs;
+			_commitService = commitService;
+			_perforce = perforce;
+			_streamCollection = streamCollection;
+			_jobService = jobService;
+			_templateCollection = templateCollection;
+			_artifactCollection = artifactCollection;
+			_userCollection = userCollection;
+			_notificationService = notificationService;
+			_agentService = agentService;
+			_globalConfig = globalConfig;
+			_logger = logger;
+		}
+
+		/// <summary>
+		/// Creates a new job
+		/// </summary>
+		/// <param name="create">Properties of the new job</param>
+		/// <returns>Id of the new job</returns>
+		[HttpPost]
+		[Route("/api/v1/jobs")]
+		public async Task<ActionResult<CreateJobResponse>> CreateJobAsync([FromBody] CreateJobRequest create)
+		{
+			StreamConfig? streamConfig;
+			if (!_globalConfig.Value.TryGetStream(new StreamId(create.StreamId), out streamConfig))
+			{
+				return NotFound(create.StreamId);
+			}
+			if (!streamConfig.Authorize(AclAction.CreateJob, User))
+			{
+				return Forbid(AclAction.CreateJob, streamConfig.Id);
+			}
+
+			// Get the name of the template ref
+			TemplateId templateRefId = create.TemplateId;
+
+			// Augment the request with template properties
+			TemplateRefConfig? templateRefConfig;
+			if (!streamConfig.TryGetTemplate(templateRefId, out templateRefConfig))
+			{
+				return BadRequest($"Template {create.TemplateId} is not available for stream {streamConfig.Id}");
+			}
+			if (!templateRefConfig.Authorize(AclAction.CreateJob, User))
+			{
+				return Forbid(AclAction.CreateJob, streamConfig.Id);
+			}
+
+			ITemplate? template = await _templateCollection.GetOrAddAsync(templateRefConfig);
+			if (template == null)
+			{
+				return BadRequest("Missing template referenced by {TemplateId}", create.TemplateId);
+			}
+			if (!template.AllowPreflights && create.PreflightChange > 0)
+			{
+				return BadRequest("Template {TemplateId} does not allow preflights", create.TemplateId);
+			}
+
+			// Get the name of the new job
+			string name = create.Name ?? template.Name;
+			if (create.TemplateId == new TemplateId("stage-to-marketplace") && create.Arguments != null)
+			{
+				foreach (string argument in create.Arguments)
+				{
+					const string Prefix = "-set:UserContentItems=";
+					if (argument.StartsWith(Prefix, StringComparison.Ordinal))
+					{
+						name += $" - {argument.Substring(Prefix.Length)}";
+						break;
+					}
+				}
+			}
+
+			// Environment variables for the job
+			Dictionary<string, string> environment = new Dictionary<string, string>();
+
+			// Check the preflight change is valid
+			ShelfInfo? shelfInfo = null;
+			if (create.PreflightChange != null)
+			{
+				(CheckShelfResult result, shelfInfo) = await _perforce.CheckShelfAsync(streamConfig, create.PreflightChange.Value, HttpContext.RequestAborted);
+				switch (result)
+				{
+					case CheckShelfResult.Ok:
+						break;
+					case CheckShelfResult.NoChange:
+						return BadRequest(KnownLogEvents.Horde_InvalidPreflight, "CL {Change} does not exist", create.PreflightChange);
+					case CheckShelfResult.NoShelvedFiles:
+						return BadRequest(KnownLogEvents.Horde_InvalidPreflight, "CL {Change} does not contain any shelved files", create.PreflightChange);
+					case CheckShelfResult.WrongStream:
+						return BadRequest(KnownLogEvents.Horde_InvalidPreflight, "CL {Change} does not contain files in {Stream}", create.PreflightChange, streamConfig.Name);
+					case CheckShelfResult.MixedStream:
+						return BadRequest(KnownLogEvents.Horde_InvalidPreflight, "CL {Change} contains files from multiple streams", create.PreflightChange);
+					default:
+						return BadRequest(KnownLogEvents.Horde_InvalidPreflight, "CL {Change} cannot be preflighted ({Result})", create.PreflightChange, result);
+				}
+
+				if (shelfInfo!.Tags != null)
+				{
+					string tagList = String.Join(";", shelfInfo.Tags.Select(x => x.ToString()));
+					environment.Add("UE_HORDE_PREFLIGHT_TAGS", tagList);
+				}
+			}
+
+			// Get the priority of the new job
+			Priority priority = create.Priority ?? template.Priority ?? Priority.Normal;
+
+			// New groups for the job
+			IGraph graph = await _graphs.AddAsync(template, streamConfig.InitialAgentType);
+
+			// Get the commits for this stream
+			ICommitCollection commits = _commitService.GetCollection(streamConfig);
+
+			// Get the change to build
+			int change = await GetChangeToBuildAsync(create, streamConfig.Id, template, shelfInfo, commits, HttpContext.RequestAborted);
+
+			// And get the matching code changelist
+			ICommit? lastCodeCommit = await commits.GetLastCodeChangeAsync(change, HttpContext.RequestAborted);
+			int codeChange = lastCodeCommit?.Number ?? change;
+
+			// New properties for the job
+			List<string> arguments = create.Arguments ?? template.GetDefaultArguments();
+
+			bool? updateIssues = null;
+			if (template.UpdateIssues)
+			{
+				updateIssues = true;
+			}
+			else if (create.UpdateIssues.HasValue)
+			{
+				updateIssues = create.UpdateIssues.Value;
+			}
+
+			// Create options for the new job
+			CreateJobOptions options = new CreateJobOptions(templateRefConfig);
+			options.PreflightChange = create.PreflightChange;
+			options.PreflightDescription = shelfInfo?.Description;
+			options.StartedByUserId = User.GetUserId();
+			options.Priority = priority;
+			options.AutoSubmit = create.AutoSubmit;
+			options.UpdateIssues = updateIssues;
+			options.Arguments.AddRange(arguments);
+
+			foreach ((string key, string value) in environment)
+			{
+				options.Environment[key] = value;
+			}
+
+			// Create the job
+			IJob job = await _jobService.CreateJobAsync(null, streamConfig, templateRefId, template.Hash, graph, name, change, codeChange, options);
+			await UpdateNotificationsAsync(job.Id, new UpdateNotificationsRequest { Slack = true });
+			return new CreateJobResponse(job.Id.ToString());
+		}
+
+		async ValueTask<int> GetChangeToBuildAsync(CreateJobRequest create, StreamId streamId, ITemplate template, ShelfInfo? shelfInfo, ICommitCollection commits, CancellationToken cancellationToken)
+		{
+			// If there's an explicit change specified, use that
+			if (create.Change.HasValue && create.Change.Value != -1)
+			{
+				return create.Change.Value;
+			}
+
+			// Evaluate the change queries
+			if (create.ChangeQueries != null && create.ChangeQueries.Count > 0)
+			{
+				int? change = await _jobService.EvaluateChangeQueriesAsync(streamId, create.ChangeQueries, shelfInfo?.Tags, commits, cancellationToken);
+				if (change != null)
+				{
+					return change.Value;
+				}
+			}
+
+			// If we need to submit a new change, do that
+			if (create.PreflightChange == null && template.SubmitNewChange != null)
+			{
+				return await commits.CreateNewAsync(template, cancellationToken);
+			}
+
+			// Otherwise return the latest change
+			return await commits.GetLatestNumberAsync(cancellationToken);
+		}
+
+		/// <summary>
+		/// Deletes a specific job.
+		/// </summary>
+		/// <param name="jobId">Id of the job to delete</param>
+		/// <returns>Async task</returns>
+		[HttpDelete]
+		[Route("/api/v1/jobs/{jobId}")]
+		public async Task<ActionResult> DeleteJobAsync(JobId jobId)
+		{
+			IJob? job = await _jobService.GetJobAsync(jobId);
+			if (job == null)
+			{
+				return NotFound(jobId);
+			}
+			if (!_globalConfig.Value.Authorize(job, AclAction.DeleteJob, User))
+			{
+				return Forbid(AclAction.DeleteJob, jobId);
+			}
+			if (!await _jobService.DeleteJobAsync(job))
+			{
+				return NotFound(jobId);
+			}
+			return Ok();
+		}
+
+		/// <summary>
+		/// Updates a specific job.
+		/// </summary>
+		/// <param name="jobId">Id of the job to find</param>
+		/// <param name="request">Settings to update in the job</param>
+		/// <returns>Async task</returns>
+		[HttpPut]
+		[Route("/api/v1/jobs/{jobId}")]
+		public async Task<ActionResult> UpdateJobAsync(JobId jobId, [FromBody] UpdateJobRequest request)
+		{
+			IJob? job = await _jobService.GetJobAsync(jobId);
+			if (job == null)
+			{
+				return NotFound(jobId);
+			}
+
+			StreamConfig? streamConfig;
+			if (!_globalConfig.Value.TryGetStream(job.StreamId, out streamConfig))
+			{
+				return NotFound(job.StreamId);
+			}
+			if (!streamConfig.Authorize(AclAction.UpdateJob, User))
+			{
+				return Forbid(AclAction.UpdateJob, jobId);
+			}
+
+			// Convert legacy behavior of clearing out the argument to setting the aborted flag
+			if (request.Arguments != null && request.Arguments.Count == 0)
+			{
+				request.Aborted = true;
+				request.Arguments = null;
+			}
+
+			UserId? abortedByUserId = null;
+			if (request.Aborted ?? false)
+			{
+				abortedByUserId = User.GetUserId();
+			}
+
+			IJob? newJob = await _jobService.UpdateJobAsync(job, name: request.Name, priority: request.Priority, autoSubmit: request.AutoSubmit, abortedByUserId: abortedByUserId, arguments: request.Arguments);
+			if (newJob == null)
+			{
+				return NotFound(jobId);
+			}
+			return Ok();
+		}
+
+		/// <summary>
+		/// Updates notifications for a specific job.
+		/// </summary>
+		/// <param name="jobId">Id of the job to find</param>
+		/// <param name="request">The notification request</param>
+		/// <returns>Information about the requested job</returns>
+		[HttpPut]
+		[Route("/api/v1/jobs/{jobId}/notifications")]
+		public async Task<ActionResult> UpdateNotificationsAsync(JobId jobId, [FromBody] UpdateNotificationsRequest request)
+		{
+			IJob? job = await _jobService.GetJobAsync(jobId);
+			if (job == null)
+			{
+				return NotFound(jobId);
+			}
+
+			StreamConfig? streamConfig;
+			if (!_globalConfig.Value.TryGetStream(job.StreamId, out streamConfig))
+			{
+				return NotFound(job.StreamId);
+			}
+			if (!streamConfig.Authorize(AclAction.CreateSubscription, User))
+			{
+				return Forbid(AclAction.CreateSubscription, jobId);
+			}
+
+			ObjectId triggerId = job.NotificationTriggerId ?? ObjectId.GenerateNewId();
+
+			job = await _jobService.UpdateJobAsync(job, null, null, null, null, triggerId, null, null);
+			if (job == null)
+			{
+				return NotFound(jobId);
+			}
+
+			await _notificationService.UpdateSubscriptionsAsync(triggerId, User, request.Email, request.Slack);
+			return Ok();
+		}
+
+		/// <summary>
+		/// Gets information about a specific job.
+		/// </summary>
+		/// <param name="jobId">Id of the job to find</param>
+		/// <returns>Information about the requested job</returns>
+		[HttpGet]
+		[Route("/api/v1/jobs/{jobId}/notifications")]
+		public async Task<ActionResult<GetNotificationResponse>> GetNotificationsAsync(JobId jobId)
+		{
+			IJob? job = await _jobService.GetJobAsync(jobId);
+			if (job == null)
+			{
+				return NotFound(jobId);
+			}
+
+			StreamConfig? streamConfig;
+			if (!_globalConfig.Value.TryGetStream(job.StreamId, out streamConfig))
+			{
+				return NotFound(job.StreamId);
+			}
+			if (!streamConfig.Authorize(AclAction.CreateSubscription, User))
+			{
+				return Forbid(AclAction.CreateSubscription, jobId);
+			}
+
+			INotificationSubscription? subscription;
+			if (job.NotificationTriggerId == null)
+			{
+				subscription = null;
+			}
+			else
+			{
+				subscription = await _notificationService.GetSubscriptionsAsync(job.NotificationTriggerId.Value, User);
+			}
+			return new GetNotificationResponse(subscription);
+		}
+
+		/// <summary>
+		/// Gets information about a specific job.
+		/// </summary>
+		/// <param name="jobId">Id of the job to find</param>
+		/// <param name="modifiedAfter">If specified, returns an empty response unless the job's update time is equal to or less than the given value</param>
+		/// <param name="filter">Filter for the fields to return</param>
+		/// <returns>Information about the requested job</returns>
+		[HttpGet]
+		[Route("/api/v1/jobs/{jobId}")]
+		[ProducesResponseType(typeof(GetJobResponse), 200)]
+		public async Task<ActionResult<object>> GetJobAsync(JobId jobId, [FromQuery] DateTimeOffset? modifiedAfter = null, [FromQuery] PropertyFilter? filter = null)
+		{
+			IJob? job = await _jobService.GetJobAsync(jobId);
+			if (job == null)
+			{
+				return NotFound(jobId);
+			}
+
+			StreamConfig? streamConfig;
+			if (!_globalConfig.Value.TryGetStream(job.StreamId, out streamConfig))
+			{
+				return NotFound(job.StreamId);
+			}
+			if (!streamConfig.Authorize(AclAction.ViewJob, User))
+			{
+				return Forbid(AclAction.ViewJob, job.StreamId);
+			}
+			if (modifiedAfter != null && job.UpdateTimeUtc <= modifiedAfter.Value)
+			{
+				return new Dictionary<string, object>();
+			}
+
+			IGraph graph = await _jobService.GetGraphAsync(job);
+			bool includeCosts = streamConfig.Authorize(AclAction.ViewCosts, User);
+			return await CreateJobResponseAsync(job, graph, includeCosts, filter);
+		}
+
+		/// <summary>
+		/// Creates a response containing information about this object
+		/// </summary>
+		/// <param name="job">The job document</param>
+		/// <param name="graph">The graph for this job</param>
+		/// <param name="includeCosts">Whether to include costs in the response</param>
+		/// <param name="filter">Filter for the properties to return</param>
+		/// <returns>Object containing the requested properties</returns>
+		async Task<object> CreateJobResponseAsync(IJob job, IGraph graph, bool includeCosts, PropertyFilter? filter)
+		{
+			if (filter == null)
+			{
+				return await CreateJobResponseAsync(job, graph, true, true, includeCosts);
+			}
+			else
+			{
+				return filter.ApplyTo(await CreateJobResponseAsync(job, graph, filter.Includes(nameof(GetJobResponse.Batches)), filter.Includes(nameof(GetJobResponse.Labels)) || filter.Includes(nameof(GetJobResponse.DefaultLabel)), includeCosts));
+			}
+		}
+
+		/// <summary>
+		/// Creates a response containing information about this object
+		/// </summary>
+		/// <param name="job">The job document</param>
+		/// <param name="graph">The graph definition</param>
+		/// <param name="includeBatches">Whether to include the job batches in the response</param>
+		/// <param name="includeLabels">Whether to include the job aggregates in the response</param>
+		/// <param name="includeCosts">Whether to include costs of running particular agents</param>
+		/// <returns>The response object</returns>
+		async ValueTask<GetJobResponse> CreateJobResponseAsync(IJob job, IGraph graph, bool includeBatches, bool includeLabels, bool includeCosts)
+		{
+			GetThinUserInfoResponse? startedByUserInfo = null;
+			if (job.StartedByUserId != null)
+			{
+				startedByUserInfo = new GetThinUserInfoResponse(await _userCollection.GetCachedUserAsync(job.StartedByUserId.Value));
+			}
+
+			GetThinUserInfoResponse? abortedByUserInfo = null;
+			if (job.AbortedByUserId != null)
+			{
+				abortedByUserInfo = new GetThinUserInfoResponse(await _userCollection.GetCachedUserAsync(job.AbortedByUserId.Value));
+			}
+
+			GetJobResponse response = new GetJobResponse(job, startedByUserInfo, abortedByUserInfo);
+			if (includeBatches || includeLabels)
+			{
+				if (includeBatches)
+				{
+					response.Batches = new List<GetBatchResponse>();
+					foreach (IJobStepBatch batch in job.Batches)
+					{
+						response.Batches.Add(await CreateBatchResponseAsync(batch, includeCosts));
+					}
+				}
+				if (includeLabels)
+				{
+					response.Labels = new List<GetLabelStateResponse>();
+					response.DefaultLabel = job.GetLabelStateResponses(graph, response.Labels);
+				}
+			}
+			return response;
+		}
+
+		/// <summary>
+		/// Get the response object for a batch
+		/// </summary>
+		/// <param name="batch"></param>
+		/// <param name="includeCosts"></param>
+		/// <returns></returns>
+		async ValueTask<GetBatchResponse> CreateBatchResponseAsync(IJobStepBatch batch, bool includeCosts)
+		{
+			List<GetStepResponse> steps = new List<GetStepResponse>();
+			foreach (IJobStep step in batch.Steps)
+			{
+				steps.Add(await CreateStepResponseAsync(step));
+			}
+
+			double? agentRate = null;
+			if (batch.AgentId != null && includeCosts)
+			{
+				agentRate = await _agentService.GetRateAsync(batch.AgentId.Value);
+			}
+
+			return new GetBatchResponse(batch, steps, agentRate);
+		}
+
+		/// <summary>
+		/// Get the response object for a step
+		/// </summary>
+		/// <param name="step"></param>
+		/// <returns></returns>
+		async ValueTask<GetStepResponse> CreateStepResponseAsync(IJobStep step)
+		{
+			GetThinUserInfoResponse? abortedByUserInfo = null;
+			if (step.AbortedByUserId != null)
+			{
+				abortedByUserInfo = new GetThinUserInfoResponse(await _userCollection.GetCachedUserAsync(step.AbortedByUserId.Value));
+			}
+
+			GetThinUserInfoResponse? retriedByUserInfo = null;
+			if (step.RetriedByUserId != null)
+			{
+				retriedByUserInfo = new GetThinUserInfoResponse(await _userCollection.GetCachedUserAsync(step.RetriedByUserId.Value));
+			}
+
+			return new GetStepResponse(step, abortedByUserInfo, retriedByUserInfo);
+		}
+
+		/// <summary>
+		/// Gets information about the graph for a specific job.
+		/// </summary>
+		/// <param name="jobId">Id of the job to find</param>
+		/// <param name="filter">Filter for the fields to return</param>
+		/// <returns>Information about the requested job</returns>
+		[HttpGet]
+		[Route("/api/v1/jobs/{jobId}/graph")]
+		[ProducesResponseType(typeof(GetGraphResponse), 200)]
+		public async Task<ActionResult<object>> GetJobGraphAsync(JobId jobId, [FromQuery] PropertyFilter? filter = null)
+		{
+			IJob? job = await _jobService.GetJobAsync(jobId);
+			if (job == null)
+			{
+				return NotFound(jobId);
+			}
+
+			StreamConfig? streamConfig;
+			if (!_globalConfig.Value.TryGetStream(job.StreamId, out streamConfig))
+			{
+				return NotFound(job.StreamId);
+			}
+			if (!streamConfig.Authorize(AclAction.ViewJob, User))
+			{
+				return Forbid(AclAction.ViewJob, jobId);
+			}
+
+			IGraph graph = await _jobService.GetGraphAsync(job);
+			return PropertyFilter.Apply(new GetGraphResponse(graph), filter);
+		}
+
+		/// <summary>
+		/// Gets timing information about the graph for a specific job.
+		/// </summary>
+		/// <param name="jobId">Id of the job to find</param>
+		/// <param name="filter">Filter for the fields to return</param>
+		/// <returns>Information about the requested job</returns>
+		[HttpGet]
+		[Route("/api/v1/jobs/{jobId}/timing")]
+		[ProducesResponseType(typeof(GetJobTimingResponse), 200)]
+		public async Task<ActionResult<object>> GetJobTimingAsync(JobId jobId, [FromQuery] PropertyFilter? filter = null)
+		{
+			IJob? job = await _jobService.GetJobAsync(jobId);
+			if (job == null)
+			{
+				return NotFound(jobId);
+			}
+
+			StreamConfig? streamConfig;
+			if (!_globalConfig.Value.TryGetStream(job.StreamId, out streamConfig))
+			{
+				return NotFound(job.StreamId);
+			}
+			if (!streamConfig.Authorize(AclAction.ViewJob, User))
+			{
+				return Forbid(AclAction.ViewJob, jobId);
+			}
+
+			IJobTiming jobTiming = await _jobService.GetJobTimingAsync(job);
+			IGraph graph = await _jobService.GetGraphAsync(job);
+			return PropertyFilter.Apply(await CreateJobTimingResponse(job, graph, jobTiming), filter);
+		}
+
+		private async Task<GetJobTimingResponse> CreateJobTimingResponse(IJob job, IGraph graph, IJobTiming jobTiming, bool includeJobResponse = false)
+		{
+			Dictionary<INode, TimingInfo> nodeToTimingInfo = job.GetTimingInfo(graph, jobTiming);
+
+			Dictionary<string, GetStepTimingInfoResponse> steps = new Dictionary<string, GetStepTimingInfoResponse>();
+			foreach (IJobStepBatch batch in job.Batches)
+			{
+				foreach (IJobStep step in batch.Steps)
+				{
+					INode node = graph.Groups[batch.GroupIdx].Nodes[step.NodeIdx];
+					steps[step.Id.ToString()] = new GetStepTimingInfoResponse(node.Name, nodeToTimingInfo[node]);
+				}
+			}
+
+			List<GetLabelTimingInfoResponse> labels = new List<GetLabelTimingInfoResponse>();
+			foreach (ILabel label in graph.Labels)
+			{
+				TimingInfo timingInfo = TimingInfo.Max(label.GetDependencies(graph.Groups).Select(x => nodeToTimingInfo[x]));
+				labels.Add(new GetLabelTimingInfoResponse(label, timingInfo));
+			}
+
+			GetJobResponse? jobResponse = null;
+			if (includeJobResponse)
+			{
+				jobResponse = await CreateJobResponseAsync(job, graph, true, true, true);
+			}
+
+			return new GetJobTimingResponse(jobResponse, steps, labels);
+		}
+		
+		/// <summary>
+		/// Find timing information about the graph for multiple jobs
+		/// </summary>
+		/// <param name="streamId">The stream to search in</param>
+		/// <param name="templates">List of templates to find</param>
+		/// <param name="filter">Filter for the fields to return</param>
+		/// <param name="count">Number of results to return</param>
+		/// <returns>Job timings for each job ID</returns>
+		[HttpGet]
+		[Route("/api/v1/jobs/timing")]
+		[ProducesResponseType(typeof(FindJobTimingsResponse), 200)]
+		public async Task<ActionResult<object>> FindJobTimingsAsync(
+			[FromQuery] string? streamId = null,
+			[FromQuery(Name = "template")] string[]? templates = null,
+			[FromQuery] PropertyFilter? filter = null,
+			[FromQuery] int count = 100)
+		{
+			if (streamId == null)
+			{
+				return BadRequest("Missing/invalid query parameter streamId");
+			}
+
+			TemplateId[] templateRefIds = templates switch
+			{
+				{ Length: > 0 } => templates.Select(x => new TemplateId(x)).ToArray(),
+				_ => Array.Empty<TemplateId>()
+			};
+
+			List<IJob> jobs = await _jobService.FindJobsByStreamWithTemplatesAsync(new StreamId(streamId), templateRefIds, count: count, consistentRead: false);
+
+			Dictionary<string, GetJobTimingResponse> jobTimings = await jobs.ToAsyncEnumerable()
+				.Where(job => _globalConfig.Value.Authorize(job, AclAction.ViewJob, User))
+				.ToDictionaryAwaitAsync(x => ValueTask.FromResult(x.Id.ToString()), async job =>
+				{
+					IJobTiming jobTiming = await _jobService.GetJobTimingAsync(job);
+					IGraph graph = await _jobService.GetGraphAsync(job);
+					return await CreateJobTimingResponse(job, graph, jobTiming, true);
+				});
+			
+			return PropertyFilter.Apply(new FindJobTimingsResponse(jobTimings), filter);
+		}
+
+		/// <summary>
+		/// Gets information about the template for a specific job.
+		/// </summary>
+		/// <param name="jobId">Id of the job to find</param>
+		/// <param name="filter">Filter for the fields to return</param>
+		/// <returns>Information about the requested job</returns>
+		[HttpGet]
+		[Route("/api/v1/jobs/{jobId}/template")]
+		[ProducesResponseType(typeof(GetTemplateResponse), 200)]
+		public async Task<ActionResult<object>> GetJobTemplateAsync(JobId jobId, [FromQuery] PropertyFilter? filter = null)
+		{
+			IJob? job = await _jobService.GetJobAsync(jobId);
+			if (job == null || job.TemplateHash == null)
+			{
+				return NotFound(jobId);
+			}
+
+			StreamConfig? streamConfig;
+			if (!_globalConfig.Value.TryGetStream(job.StreamId, out streamConfig))
+			{
+				return NotFound(job.StreamId);
+			}
+			if (!streamConfig.Authorize(AclAction.ViewJob, User))
+			{
+				return Forbid(AclAction.ViewJob, jobId);
+			}
+
+			ITemplate? template = await _templateCollection.GetAsync(job.TemplateHash);
+			if(template == null)
+			{
+				return NotFound(job.StreamId, job.TemplateId);
+			}
+
+			return new GetTemplateResponse(template).ApplyFilter(filter);
+		}
+
+		/// <summary>
+		/// Find jobs matching a criteria
+		/// </summary>
+		/// <param name="ids">The job ids to return</param>
+		/// <param name="name">Name of the job to find</param>
+		/// <param name="templates">List of templates to find</param>
+		/// <param name="streamId">The stream to search for</param>
+		/// <param name="minChange">The minimum changelist number</param>
+		/// <param name="maxChange">The maximum changelist number</param>
+		/// <param name="includePreflight">Whether to include preflight jobs</param>
+		/// <param name="preflightOnly">Whether to only include preflight jobs</param>
+		/// <param name="preflightChange">The preflighted changelist</param>
+		/// <param name="startedByUserId">User id for which to include jobs</param>
+		/// <param name="preflightStartedByUserId">User id for which to include preflight jobs</param>
+		/// <param name="minCreateTime">Minimum creation time</param>
+		/// <param name="maxCreateTime">Maximum creation time</param>
+		/// <param name="modifiedBefore">If specified, only jobs updated before the give time will be returned</param>
+		/// <param name="modifiedAfter">If specified, only jobs updated after the give time will be returned</param>
+		/// <param name="target">Target to filter the returned jobs by</param>
+		/// <param name="state">Filter state of the returned jobs</param>
+		/// <param name="outcome">Filter outcome of the returned jobs</param>
+		/// <param name="filter">Filter for properties to return</param>
+		/// <param name="index">Index of the first result to be returned</param>
+		/// <param name="count">Number of results to return</param>
+		/// <returns>List of jobs</returns>
+		[HttpGet]
+		[Route("/api/v1/jobs")]
+		[ProducesResponseType(typeof(List<GetJobResponse>), 200)]
+		public async Task<ActionResult<List<object>>> FindJobsAsync(
+			[FromQuery(Name = "Id")] string[]? ids = null,
+			[FromQuery] string? streamId = null,
+			[FromQuery] string? name = null,
+			[FromQuery(Name = "template")] string[]? templates = null,
+			[FromQuery] int? minChange = null,
+			[FromQuery] int? maxChange = null,
+			[FromQuery] bool includePreflight = true,
+			[FromQuery] bool? preflightOnly = null,
+			[FromQuery] int? preflightChange = null,
+			[FromQuery] string? preflightStartedByUserId = null,
+			[FromQuery] string? startedByUserId = null,
+			[FromQuery] DateTimeOffset? minCreateTime = null,
+			[FromQuery] DateTimeOffset? maxCreateTime = null,
+			[FromQuery] DateTimeOffset? modifiedBefore = null,
+			[FromQuery] DateTimeOffset? modifiedAfter = null,
+			[FromQuery] string? target = null,
+			[FromQuery] JobStepState[]? state = null,
+			[FromQuery] JobStepOutcome[]? outcome = null, 
+			[FromQuery] PropertyFilter? filter = null,
+			[FromQuery] int index = 0,
+			[FromQuery] int count = 100)
+		{
+			JobId[]? jobIdValues = (ids == null) ? (JobId[]?)null : Array.ConvertAll(ids, x => JobId.Parse(x));
+			StreamId? streamIdValue = (streamId == null)? (StreamId?)null : new StreamId(streamId);
+			
+			TemplateId[]? templateRefIds = (templates != null && templates.Length > 0) ? templates.Select(x => new TemplateId(x)).ToArray() : null;
+
+			if (includePreflight == false)
+			{
+				preflightChange = 0;
+			}
+
+			UserId? preflightStartedByUserIdValue = null;
+
+			if (preflightStartedByUserId != null)
+			{
+				preflightStartedByUserIdValue = UserId.Parse(preflightStartedByUserId);
+			}
+
+			UserId? startedByUserIdValue = null;
+
+			if (startedByUserId != null)
+			{
+				startedByUserIdValue = UserId.Parse(startedByUserId);
+			}
+
+			List<IJob> jobs;
+			using (IScope _ = GlobalTracer.Instance.BuildSpan("FindJobs").StartActive())
+			{
+				jobs = await _jobService.FindJobsAsync(jobIdValues, streamIdValue, name, templateRefIds, minChange,
+					maxChange, preflightChange, preflightOnly, preflightStartedByUserIdValue, startedByUserIdValue, minCreateTime?.UtcDateTime, maxCreateTime?.UtcDateTime, target, state, outcome,
+					modifiedBefore, modifiedAfter, index, count, false);
+			}
+
+			return await CreateAuthorizedJobResponsesAsync(jobs, filter);
+		}
+
+		/// <summary>
+		/// Find jobs for a stream with given templates, sorted by creation date
+		/// </summary>
+		/// <param name="streamId">The stream to search for</param>
+		/// <param name="templates">List of templates to find</param>
+		/// <param name="preflightStartedByUserId">User id for which to include preflight jobs</param>
+		/// <param name="maxCreateTime">Maximum creation time</param>
+		/// <param name="modifiedAfter">If specified, only jobs updated after the given time will be returned</param>
+		/// <param name="filter">Filter for properties to return</param>
+		/// <param name="index">Index of the first result to be returned</param>
+		/// <param name="count">Number of results to return</param>
+		/// <param name="consistentRead">If a read to the primary database is required, for read consistency. Usually not required.</param>
+		/// <returns>List of jobs</returns>
+		[HttpGet]
+		[Route("/api/v1/jobs/streams/{streamId}")]
+		[ProducesResponseType(typeof(List<GetJobResponse>), 200)]
+		public async Task<ActionResult<List<object>>> FindJobsByStreamWithTemplatesAsync(
+			string streamId,
+			[FromQuery(Name = "template")] string[] templates,
+			[FromQuery] string? preflightStartedByUserId = null,
+			[FromQuery] DateTimeOffset? maxCreateTime = null,
+			[FromQuery] DateTimeOffset? modifiedAfter = null,
+			[FromQuery] PropertyFilter? filter = null,
+			[FromQuery] int index = 0,
+			[FromQuery] int count = 100,
+			[FromQuery] bool consistentRead = false)
+		{
+			StreamId streamIdValue = new StreamId(streamId);
+			TemplateId[] templateRefIds = templates.Select(x => new TemplateId(x)).ToArray();
+			UserId? preflightStartedByUserIdValue = preflightStartedByUserId != null ? UserId.Parse(preflightStartedByUserId) : null;
+			count = Math.Min(1000, count);
+
+			List<IJob> jobs = await _jobService.FindJobsByStreamWithTemplatesAsync(streamIdValue, templateRefIds, preflightStartedByUserIdValue, maxCreateTime, modifiedAfter, index, count, consistentRead);
+			return await CreateAuthorizedJobResponsesAsync(jobs, filter);
+		}
+
+		private async Task<List<object>> CreateAuthorizedJobResponsesAsync(List<IJob> jobs, PropertyFilter? filter = null)
+		{
+			List<object> responses = new ();
+			foreach(IGrouping<StreamId, IJob> grouping in jobs.GroupBy(x => x.StreamId))
+			{
+				StreamConfig? streamConfig;
+				if (_globalConfig.Value.TryGetStream(grouping.Key, out streamConfig) && streamConfig.Authorize(AclAction.ViewJob, User))
+				{
+					bool includeCosts = streamConfig.Authorize(AclAction.ViewCosts, User);
+					foreach (IJob job in grouping)
+					{
+						IGraph graph = await _jobService.GetGraphAsync(job);
+						responses.Add(await CreateJobResponseAsync(job, graph, includeCosts, filter));
+					}
+				}
+			}
+			return responses;
+		}
+
+		/// <summary>
+		/// Adds an array of nodes to be executed for a job
+		/// </summary>
+		/// <param name="jobId">Unique id for the job</param>
+		/// <param name="requests">Properties of the new nodes</param>
+		/// <returns>Id of the new job</returns>
+		[HttpPost]
+		[Route("/api/v1/jobs/{jobId}/groups")]
+		public async Task<ActionResult> CreateGroupsAsync(JobId jobId, [FromBody] List<NewGroup> requests)
+		{
+			for (; ; )
+			{
+				IJob? job = await _jobService.GetJobAsync(jobId);
+				if (job == null)
+				{
+					return NotFound(jobId);
+				}
+
+				StreamConfig? streamConfig;
+				if (!_globalConfig.Value.TryGetStream(job.StreamId, out streamConfig))
+				{
+					return NotFound(job.StreamId);
+				}
+				if (!streamConfig.Authorize(AclAction.ExecuteJob, User))
+				{
+					return Forbid(AclAction.ExecuteJob, jobId);
+				}
+
+				IGraph graph = await _jobService.GetGraphAsync(job);
+				graph = await _graphs.AppendAsync(graph, requests, null, null);
+
+				IJob? newJob = await _jobService.TryUpdateGraphAsync(job, graph);
+				if (newJob != null)
+				{
+					return Ok();
+				}
+			}
+		}
+
+		/// <summary>
+		/// Gets the nodes to be executed for a job
+		/// </summary>
+		/// <param name="jobId">Unique id for the job</param>
+		/// <param name="filter">Filter for the properties to return</param>
+		/// <returns>List of nodes to be executed</returns>
+		[HttpGet]
+		[Route("/api/v1/jobs/{jobId}/groups")]
+		[ProducesResponseType(typeof(List<GetGroupResponse>), 200)]
+		public async Task<ActionResult<List<object>>> GetGroupsAsync(JobId jobId, [FromQuery] PropertyFilter? filter = null)
+		{
+			IJob? job = await _jobService.GetJobAsync(jobId);
+			if (job == null)
+			{
+				return NotFound(jobId);
+			}
+
+			StreamConfig? streamConfig;
+			if (!_globalConfig.Value.TryGetStream(job.StreamId, out streamConfig))
+			{
+				return NotFound(job.StreamId);
+			}
+			if (!streamConfig.Authorize(AclAction.ViewJob, User))
+			{
+				return Forbid(AclAction.ViewJob, jobId);
+			}
+
+			IGraph graph = await _jobService.GetGraphAsync(job);
+			return graph.Groups.ConvertAll(x => new GetGroupResponse(x, graph.Groups).ApplyFilter(filter));
+		}
+
+		/// <summary>
+		/// Gets the nodes in a group to be executed for a job
+		/// </summary>
+		/// <param name="jobId">Unique id for the job</param>
+		/// <param name="groupIdx">The group index</param>
+		/// <param name="filter">Filter for the properties to return</param>
+		/// <returns>List of nodes to be executed</returns>
+		[HttpGet]
+		[Route("/api/v1/jobs/{jobId}/groups/{groupIdx}")]
+		[ProducesResponseType(typeof(GetGroupResponse), 200)]
+		public async Task<ActionResult<object>> GetGroupAsync(JobId jobId, int groupIdx, [FromQuery] PropertyFilter? filter = null)
+		{
+			IJob? job = await _jobService.GetJobAsync(jobId);
+			if (job == null)
+			{
+				return NotFound(jobId);
+			}
+
+			StreamConfig? streamConfig;
+			if (!_globalConfig.Value.TryGetStream(job.StreamId, out streamConfig))
+			{
+				return NotFound(job.StreamId);
+			}
+			if (!streamConfig.Authorize(AclAction.ViewJob, User))
+			{
+				return Forbid(AclAction.ViewJob, jobId);
+			}
+
+			IGraph graph = await _jobService.GetGraphAsync(job);
+			if (groupIdx < 0 || groupIdx >= graph.Groups.Count)
+			{
+				return NotFound(jobId, groupIdx);
+			}
+
+			return new GetGroupResponse(graph.Groups[groupIdx], graph.Groups).ApplyFilter(filter);
+		}
+
+		/// <summary>
+		/// Gets the nodes for a particular group
+		/// </summary>
+		/// <param name="jobId">Unique id for the job</param>
+		/// <param name="groupIdx">Index of the group containing the node to update</param>
+		/// <param name="filter">Filter for the properties to return</param>
+		[HttpGet]
+		[Route("/api/v1/jobs/{jobId}/groups/{groupIdx}/nodes")]
+		[ProducesResponseType(typeof(List<GetNodeResponse>), 200)]
+		public async Task<ActionResult<List<object>>> GetNodesAsync(JobId jobId, int groupIdx, [FromQuery] PropertyFilter? filter = null)
+		{
+			IJob? job = await _jobService.GetJobAsync(jobId);
+			if (job == null)
+			{
+				return NotFound(jobId);
+			}
+
+			StreamConfig? streamConfig;
+			if (!_globalConfig.Value.TryGetStream(job.StreamId, out streamConfig))
+			{
+				return NotFound(job.StreamId);
+			}
+			if (!streamConfig.Authorize(AclAction.ViewJob, User))
+			{
+				return Forbid(AclAction.ViewJob, jobId);
+			}
+
+			IGraph graph = await _jobService.GetGraphAsync(job);
+			if (groupIdx < 0 || groupIdx >= graph.Groups.Count)
+			{
+				return NotFound(jobId, groupIdx);
+			}
+
+			return graph.Groups[groupIdx].Nodes.ConvertAll(x => new GetNodeResponse(x, graph.Groups).ApplyFilter(filter));
+		}
+
+		/// <summary>
+		/// Gets a particular node definition
+		/// </summary>
+		/// <param name="jobId">Unique id for the job</param>
+		/// <param name="groupIdx">Index of the group containing the node to update</param>
+		/// <param name="nodeIdx">Index of the node to update</param>
+		/// <param name="filter">Filter for the properties to return</param>
+		[HttpGet]
+		[Route("/api/v1/jobs/{jobId}/groups/{groupIdx}/nodes/{nodeIdx}")]
+		[ProducesResponseType(typeof(GetNodeResponse), 200)]
+		public async Task<ActionResult<object>> GetNodeAsync(JobId jobId, int groupIdx, int nodeIdx, [FromQuery] PropertyFilter? filter = null)
+		{
+			IJob? job = await _jobService.GetJobAsync(jobId);
+			if (job == null)
+			{
+				return NotFound(jobId);
+			}
+
+			StreamConfig? streamConfig;
+			if (!_globalConfig.Value.TryGetStream(job.StreamId, out streamConfig))
+			{
+				return NotFound(job.StreamId);
+			}
+			if (!streamConfig.Authorize(AclAction.ViewJob, User))
+			{
+				return Forbid(AclAction.ViewJob, jobId);
+			}
+
+			IGraph graph = await _jobService.GetGraphAsync(job);
+			if (groupIdx < 0 || groupIdx >= graph.Groups.Count)
+			{
+				return NotFound(jobId, groupIdx);
+			}
+			if(nodeIdx < 0 || nodeIdx >= graph.Groups[groupIdx].Nodes.Count)
+			{
+				return NotFound(jobId, groupIdx, nodeIdx);
+			}
+
+			return new GetNodeResponse(graph.Groups[groupIdx].Nodes[nodeIdx], graph.Groups).ApplyFilter(filter);
+		}
+
+		/// <summary>
+		/// Gets the steps currently scheduled to be executed for a job
+		/// </summary>
+		/// <param name="jobId">Unique id for the job</param>
+		/// <param name="filter">Filter for the properties to return</param>
+		/// <returns>List of nodes to be executed</returns>
+		[HttpGet]
+		[Route("/api/v1/jobs/{jobId}/batches")]
+		[ProducesResponseType(typeof(List<GetBatchResponse>), 200)]
+		public async Task<ActionResult<List<object>>> GetBatchesAsync(JobId jobId, [FromQuery] PropertyFilter? filter = null)
+		{
+			IJob? job = await _jobService.GetJobAsync(jobId);
+			if (job == null)
+			{
+				return NotFound(jobId);
+			}
+
+			StreamConfig? streamConfig;
+			if (!_globalConfig.Value.TryGetStream(job.StreamId, out streamConfig))
+			{
+				return NotFound(job.StreamId);
+			}
+			if (!streamConfig.Authorize(AclAction.ViewJob, User))
+			{
+				return Forbid(AclAction.ViewJob, job.StreamId);
+			}
+
+			bool includeCosts = streamConfig.Authorize(AclAction.ViewCosts, User);
+
+			List<object> responses = new List<object>();
+			foreach (IJobStepBatch batch in job.Batches)
+			{
+				GetBatchResponse response = await CreateBatchResponseAsync(batch, includeCosts);
+				responses.Add(response.ApplyFilter(filter));
+			}
+			return responses;
+		}
+
+		/// <summary>
+		/// Updates the state of a jobstep
+		/// </summary>
+		/// <param name="jobId">Unique id for the job</param>
+		/// <param name="batchId">Unique id for the step</param>
+		/// <param name="request">Updates to apply to the node</param>
+		[HttpPut]
+		[Route("/api/v1/jobs/{jobId}/batches/{batchId}")]
+		public async Task<ActionResult> UpdateBatchAsync(JobId jobId, SubResourceId batchId, [FromBody] UpdateBatchRequest request)
+		{
+			IJob? job = await _jobService.GetJobAsync(jobId);
+			if (job == null)
+			{
+				return NotFound(jobId);
+			}
+			if (!_globalConfig.Value.TryGetStream(job.StreamId, out StreamConfig? streamConfig))
+			{
+				return NotFound(job.StreamId);
+			}
+
+			IJobStepBatch? batch = job.Batches.FirstOrDefault(x => x.Id == batchId);
+			if (batch == null)
+			{
+				return NotFound(jobId, batchId);
+			}
+			if (batch.SessionId == null || !User.HasSessionClaim(batch.SessionId.Value))
+			{
+				return Forbid("Missing session claim for job {JobId} batch {BatchId}", jobId, batchId);
+			}
+
+			IJob? newJob = await _jobService.UpdateBatchAsync(job, batchId, streamConfig, (request.LogId == null)? null : LogId.Parse(request.LogId), request.State);
+			if (newJob == null)
+			{
+				return NotFound(jobId);
+			}
+			return Ok();
+		}
+
+		/// <summary>
+		/// Gets a particular step currently scheduled to be executed for a job
+		/// </summary>
+		/// <param name="jobId">Unique id for the job</param>
+		/// <param name="batchId">Unique id for the step</param>
+		/// <param name="filter">Filter for the properties to return</param>
+		/// <returns>List of nodes to be executed</returns>
+		[HttpGet]
+		[Route("/api/v1/jobs/{jobId}/batches/{batchId}")]
+		[ProducesResponseType(typeof(GetBatchResponse), 200)]
+		public async Task<ActionResult<object>> GetBatchAsync(JobId jobId, SubResourceId batchId, [FromQuery] PropertyFilter? filter = null)
+		{
+			IJob? job = await _jobService.GetJobAsync(jobId);
+			if (job == null)
+			{
+				return NotFound(jobId);
+			}
+
+			StreamConfig? streamConfig;
+			if (!_globalConfig.Value.TryGetStream(job.StreamId, out streamConfig))
+			{
+				return NotFound(job.StreamId);
+			}
+			if (!streamConfig.Authorize(AclAction.ViewJob, User))
+			{
+				return Forbid(AclAction.ViewJob, job.StreamId);
+			}
+
+			bool includeCosts = streamConfig.Authorize(AclAction.ViewCosts, User);
+			foreach (IJobStepBatch batch in job.Batches)
+			{
+				if (batch.Id == batchId)
+				{
+					GetBatchResponse response = await CreateBatchResponseAsync(batch, includeCosts);
+					return response.ApplyFilter(filter);
+				}
+			}
+
+			return NotFound(jobId, batchId);
+		}
+
+		/// <summary>
+		/// Gets the steps currently scheduled to be executed for a job
+		/// </summary>
+		/// <param name="jobId">Unique id for the job</param>
+		/// <param name="batchId">Unique id for the batch</param>
+		/// <param name="filter">Filter for the properties to return</param>
+		/// <returns>List of nodes to be executed</returns>
+		[HttpGet]
+		[Route("/api/v1/jobs/{jobId}/batches/{batchId}/steps")]
+		[ProducesResponseType(typeof(List<GetStepResponse>), 200)]
+		public async Task<ActionResult<List<object>>> GetStepsAsync(JobId jobId, SubResourceId batchId, [FromQuery] PropertyFilter? filter = null)
+		{
+			IJob? job = await _jobService.GetJobAsync(jobId);
+			if (job == null)
+			{
+				return NotFound(jobId);
+			}
+
+			StreamConfig? streamConfig;
+			if (!_globalConfig.Value.TryGetStream(job.StreamId, out streamConfig))
+			{
+				return NotFound(job.StreamId);
+			}
+			if (!streamConfig.Authorize(AclAction.ViewJob, User))
+			{
+				return Forbid(AclAction.ViewJob, jobId);
+			}
+
+			foreach (IJobStepBatch batch in job.Batches)
+			{
+				if (batch.Id == batchId)
+				{
+					List<object> responses = new List<object>();
+					foreach (IJobStep step in batch.Steps)
+					{
+						GetStepResponse response = await CreateStepResponseAsync(step);
+						responses.Add(response.ApplyFilter(filter));
+					}
+					return responses;
+				}
+			}
+
+			return NotFound(jobId, batchId);
+		}
+
+		/// <summary>
+		/// Updates the state of a jobstep
+		/// </summary>
+		/// <param name="jobId">Unique id for the job</param>
+		/// <param name="batchId">Unique id for the batch</param>
+		/// <param name="stepId">Unique id for the step</param>
+		/// <param name="request">Updates to apply to the node</param>
+		[HttpPut]
+		[Route("/api/v1/jobs/{jobId}/batches/{batchId}/steps/{stepId}")]
+		public async Task<ActionResult<UpdateStepResponse>> UpdateStepAsync(JobId jobId, SubResourceId batchId, SubResourceId stepId, [FromBody] UpdateStepRequest request)
+		{
+			IJob? job = await _jobService.GetJobAsync(jobId);
+			if (job == null)
+			{
+				return NotFound(jobId);
+			}
+
+			StreamConfig? streamConfig;
+			if (!_globalConfig.Value.TryGetStream(job.StreamId, out streamConfig))
+			{
+				return NotFound(job.StreamId);
+			}
+
+			// Check permissions for updating this step. Only the agent executing the step can modify the state of it.
+			if (request.State != JobStepState.Unspecified || request.Outcome != JobStepOutcome.Unspecified)
+			{
+				IJobStepBatch? batch = job.Batches.FirstOrDefault(x => x.Id == batchId);
+				if (batch == null)
+				{
+					return NotFound(jobId, batchId);
+				}
+				if (!batch.SessionId.HasValue || !User.HasSessionClaim(batch.SessionId.Value))
+				{
+					return Forbid();
+				}
+			}
+
+			if (request.Retry != null || request.Priority != null)
+			{
+				if (!streamConfig.Authorize(AclAction.RetryJobStep, User))
+				{
+					return Forbid(AclAction.RetryJobStep, jobId);
+				}
+			}
+			if (request.Properties != null)
+			{
+				if (!streamConfig.Authorize(AclAction.UpdateJob, User))
+				{
+					return Forbid(AclAction.UpdateJob, jobId);
+				}
+			}
+
+			UserId? retryByUser = (request.Retry.HasValue && request.Retry.Value) ? User.GetUserId() : null;
+			UserId? abortByUser = (request.AbortRequested.HasValue && request.AbortRequested.Value) ? User.GetUserId() : null;
+
+			try
+			{
+				IJob? newJob = await _jobService.UpdateStepAsync(job, batchId, stepId, streamConfig, request.State, request.Outcome, null, request.AbortRequested, abortByUser, (request.LogId == null)? null : LogId.Parse(request.LogId), null, retryByUser, request.Priority, null, request.Properties);
+				if (newJob == null)
+				{
+					return NotFound(jobId);
+				}
+
+				UpdateStepResponse response = new UpdateStepResponse();
+				if (request.Retry ?? false)
+				{
+					JobStepRefId? retriedStepId = FindRetriedStep(job, batchId, stepId);
+					if (retriedStepId != null)
+					{
+						response.BatchId = retriedStepId.Value.BatchId.ToString();
+						response.StepId = retriedStepId.Value.StepId.ToString();
+					}
+				}
+				return response;
+			}
+			catch (RetryNotAllowedException ex)
+			{
+				return BadRequest(ex.Message);
+			}
+		}
+
+		/// <summary>
+		/// Find the first retried step after the given step
+		/// </summary>
+		/// <param name="job">The job being run</param>
+		/// <param name="batchId">Batch id of the last step instance</param>
+		/// <param name="stepId">Step id of the last instance</param>
+		/// <returns>The retried step information</returns>
+		static JobStepRefId? FindRetriedStep(IJob job, SubResourceId batchId, SubResourceId stepId)
+		{
+			NodeRef? lastNodeRef = null;
+			foreach (IJobStepBatch batch in job.Batches)
+			{
+				if ((lastNodeRef == null && batch.Id == batchId) || (lastNodeRef != null && batch.GroupIdx == lastNodeRef.GroupIdx))
+				{
+					foreach (IJobStep step in batch.Steps)
+					{
+						if (lastNodeRef == null && step.Id == stepId)
+						{
+							lastNodeRef = new NodeRef(batch.GroupIdx, step.NodeIdx);
+						}
+						else if (lastNodeRef != null && step.NodeIdx == lastNodeRef.NodeIdx)
+						{
+							return new JobStepRefId(job.Id, batch.Id, step.Id);
+						}
+					}
+				}
+			}
+			return null;
+		}
+
+		/// <summary>
+		/// Gets a particular step currently scheduled to be executed for a job
+		/// </summary>
+		/// <param name="jobId">Unique id for the job</param>
+		/// <param name="batchId">Unique id for the batch</param>
+		/// <param name="stepId">Unique id for the step</param>
+		/// <param name="filter">Filter for the properties to return</param>
+		/// <returns>List of nodes to be executed</returns>
+		[HttpGet]
+		[Route("/api/v1/jobs/{jobId}/batches/{batchId}/steps/{stepId}")]
+		[ProducesResponseType(typeof(GetStepResponse), 200)]
+		public async Task<ActionResult<object>> GetStepAsync(JobId jobId, SubResourceId batchId, SubResourceId stepId, [FromQuery] PropertyFilter? filter = null)
+		{
+			IJob? job = await _jobService.GetJobAsync(jobId);
+			if (job == null)
+			{
+				return NotFound(jobId);
+			}
+
+			StreamConfig? streamConfig;
+			if (!_globalConfig.Value.TryGetStream(job.StreamId, out streamConfig))
+			{
+				return NotFound(job.StreamId);
+			}
+			if (!streamConfig.Authorize(AclAction.ViewJob, User))
+			{
+				return Forbid(AclAction.ViewJob, jobId);
+			}
+
+			foreach (IJobStepBatch batch in job.Batches)
+			{
+				if (batch.Id == batchId)
+				{
+					foreach (IJobStep step in batch.Steps)
+					{
+						if (step.Id == stepId)
+						{
+							GetStepResponse response = await CreateStepResponseAsync(step);
+							return response.ApplyFilter(filter);
+						}
+					}
+					return NotFound(jobId, batchId, stepId);
+				}
+			}
+
+			return NotFound(jobId, batchId);
+		}
+
+		/// <summary>
+		/// Updates notifications for a specific job.
+		/// </summary>
+		/// <param name="jobId">Unique id for the job</param>
+		/// <param name="batchId">Unique id for the batch</param>
+		/// <param name="stepId">Unique id for the step</param>
+		/// <param name="request">The notification request</param>
+		/// <returns>Information about the requested job</returns>
+		[HttpPut]
+		[Route("/api/v1/jobs/{jobId}/batches/{batchId}/steps/{stepId}/notifications")]
+		public async Task<ActionResult> UpdateStepNotificationsAsync(JobId jobId, SubResourceId batchId, SubResourceId stepId, [FromBody] UpdateNotificationsRequest request)
+		{
+			IJob? job = await _jobService.GetJobAsync(jobId);
+			if (job == null)
+			{
+				return NotFound(jobId);
+			}
+
+			StreamConfig? streamConfig;
+			if (!_globalConfig.Value.TryGetStream(job.StreamId, out streamConfig))
+			{
+				return NotFound(job.StreamId);
+			}
+			if (!streamConfig.Authorize(AclAction.CreateSubscription, User))
+			{
+				return Forbid(AclAction.CreateSubscription, jobId);
+			}
+
+			if (!job.TryGetBatch(batchId, out IJobStepBatch? batch))
+			{
+				return NotFound(jobId, batchId);
+			}
+
+			if (!batch.TryGetStep(stepId, out IJobStep? step))
+			{
+				return NotFound(jobId, batchId, stepId);
+			}
+
+			ObjectId? triggerId = step.NotificationTriggerId;
+			if (triggerId == null)
+			{
+				triggerId = ObjectId.GenerateNewId();
+				if (await _jobService.UpdateStepAsync(job, batchId, stepId, streamConfig, JobStepState.Unspecified, JobStepOutcome.Unspecified, newNotificationTriggerId: triggerId) == null)
+				{
+					return NotFound(jobId, batchId, stepId);
+				}
+			}
+
+			await _notificationService.UpdateSubscriptionsAsync(triggerId.Value, User, request.Email, request.Slack);
+			return Ok();
+		}
+
+		/// <summary>
+		/// Gets information about a specific job.
+		/// </summary>
+		/// <param name="jobId">Id of the job to find</param>
+		/// <param name="batchId">Unique id for the batch</param>
+		/// <param name="stepId">Unique id for the step</param>
+		/// <returns>Information about the requested job</returns>
+		[HttpGet]
+		[Route("/api/v1/jobs/{jobId}/batches/{batchId}/steps/{stepId}/notifications")]
+		public async Task<ActionResult<GetNotificationResponse>> GetStepNotificationsAsync(JobId jobId, SubResourceId batchId, SubResourceId stepId)
+		{
+			IJob? job = await _jobService.GetJobAsync(jobId);
+			if (job == null)
+			{
+				return NotFound(jobId);
+			}
+
+			IJobStep? step;
+			if (!job.TryGetBatch(batchId, out IJobStepBatch? batch))
+			{
+				return NotFound(jobId, batchId);
+			}
+			if (!batch.TryGetStep(stepId, out step))
+			{
+				return NotFound(jobId, batchId, stepId);
+			}
+
+			StreamConfig? streamConfig;
+			if (!_globalConfig.Value.TryGetStream(job.StreamId, out streamConfig))
+			{
+				return NotFound(job.StreamId);
+			}
+			if (!streamConfig.Authorize(AclAction.CreateSubscription, User))
+			{
+				return Forbid(AclAction.CreateSubscription, jobId);
+			}
+
+			INotificationSubscription? subscription;
+			if (step.NotificationTriggerId == null)
+			{
+				subscription = null;
+			}
+			else
+			{
+				subscription = await _notificationService.GetSubscriptionsAsync(step.NotificationTriggerId.Value, User);
+			}
+			return new GetNotificationResponse(subscription);
+		}
+
+		/// <summary>
+		/// Gets a particular step currently scheduled to be executed for a job
+		/// </summary>
+		/// <param name="jobId">Unique id for the job</param>
+		/// <param name="batchId">Unique id for the batch</param>
+		/// <param name="stepId">Unique id for the step</param>
+		/// <param name="name"></param>
+		/// <returns>List of nodes to be executed</returns>
+		[HttpGet]
+		[Route("/api/v1/jobs/{jobId}/batches/{batchId}/steps/{stepId}/artifacts/{*name}")]
+		public async Task<ActionResult> GetArtifactAsync(JobId jobId, SubResourceId batchId, SubResourceId stepId, string name)
+		{
+			IJob? job = await _jobService.GetJobAsync(jobId);
+			if (job == null)
+			{
+				return NotFound(jobId);
+			}
+
+			StreamConfig? streamConfig;
+			if (!_globalConfig.Value.TryGetStream(job.StreamId, out streamConfig))
+			{
+				return NotFound(job.StreamId);
+			}
+			if (!streamConfig.Authorize(AclAction.ViewJob, User))
+			{
+				return Forbid(AclAction.ViewJob, jobId);
+			}
+
+			if (!job.TryGetBatch(batchId, out IJobStepBatch? batch))
+			{
+				return NotFound(jobId, batchId);
+			}
+			if (!batch.TryGetStep(stepId, out _))
+			{
+				return NotFound(jobId, batchId, stepId);
+			}
+
+			List<IArtifactV1> artifacts = await _artifactCollection.GetArtifactsAsync(jobId, stepId, name);
+			if (artifacts.Count == 0)
+			{
+				return NotFound();
+			}
+
+			IArtifactV1 artifact = artifacts[0];
+			return new FileStreamResult(await _artifactCollection.OpenArtifactReadStreamAsync(artifact), artifact.MimeType);
+		}
+
+		/// <summary>
+		/// Gets a particular step currently scheduled to be executed for a job
+		/// </summary>
+		/// <param name="jobId">Unique id for the job</param>
+		/// <param name="batchId">Unique id for the batch</param>
+		/// <param name="stepId">Unique id for the step</param>
+		/// <returns>List of nodes to be executed</returns>
+		[HttpGet]
+		[Route("/api/v1/jobs/{jobId}/batches/{batchId}/steps/{stepId}/trace")]
+		public async Task<ActionResult> GetStepTraceAsync(JobId jobId, SubResourceId batchId, SubResourceId stepId)
+		{
+			IJob? job = await _jobService.GetJobAsync(jobId);
+			if (job == null)
+			{
+				return NotFound(jobId);
+			}
+
+			StreamConfig? streamConfig;
+			if (!_globalConfig.Value.TryGetStream(job.StreamId, out streamConfig))
+			{
+				return NotFound(job.StreamId);
+			}
+			if (!streamConfig.Authorize(AclAction.ViewJob, User))
+			{
+				return Forbid(AclAction.ViewJob, jobId);
+			}
+
+			if (!job.TryGetBatch(batchId, out IJobStepBatch? batch))
+			{
+				return NotFound(jobId, batchId);
+			}
+			if (!batch.TryGetStep(stepId, out _))
+			{
+				return NotFound(jobId, batchId, stepId);
+			}
+
+			List<IArtifactV1> artifacts = await _artifactCollection.GetArtifactsAsync(jobId, stepId, null);
+			foreach (IArtifactV1 artifact in artifacts)
+			{
+				if (artifact.Name.Equals("trace.json", StringComparison.OrdinalIgnoreCase))
+				{
+					return new FileStreamResult(await _artifactCollection.OpenArtifactReadStreamAsync(artifact), "text/json");
+				}
+			}
+			return NotFound();
+		}
+
+		/// <summary>
+		/// Updates notifications for a specific label.
+		/// </summary>
+		/// <param name="jobId">Unique id for the job</param>
+		/// <param name="labelIndex">Index for the label</param>
+		/// <param name="request">The notification request</param>
+		[HttpPut]
+		[Route("/api/v1/jobs/{jobId}/labels/{labelIndex}/notifications")]
+		public async Task<ActionResult> UpdateLabelNotificationsAsync(JobId jobId, int labelIndex, [FromBody] UpdateNotificationsRequest request)
+		{
+			ObjectId triggerId;
+			for (; ; )
+			{
+				IJob? job = await _jobService.GetJobAsync(jobId);
+				if (job == null)
+				{
+					return NotFound(jobId);
+				}
+
+				StreamConfig? streamConfig;
+				if (!_globalConfig.Value.TryGetStream(job.StreamId, out streamConfig))
+				{
+					return NotFound(job.StreamId);
+				}
+				if (!streamConfig.Authorize(AclAction.CreateSubscription, User))
+				{
+					return Forbid(AclAction.CreateSubscription, jobId);
+				}
+
+				ObjectId newTriggerId;
+				if (job.LabelIdxToTriggerId.TryGetValue(labelIndex, out newTriggerId))
+				{
+					triggerId = newTriggerId;
+					break;
+				}
+
+				newTriggerId = ObjectId.GenerateNewId();
+
+				IJob? newJob = await _jobService.UpdateJobAsync(job, labelIdxToTriggerId: new KeyValuePair<int, ObjectId>(labelIndex, newTriggerId));
+				if (newJob != null)
+				{
+					triggerId = newTriggerId;
+					break;
+				}
+			}
+
+			await _notificationService.UpdateSubscriptionsAsync(triggerId, User, request.Email, request.Slack);
+			return Ok();
+		}
+
+		/// <summary>
+		/// Gets notification info about a specific label in a job.
+		/// </summary>
+		/// <param name="jobId">Id of the job to find</param>
+		/// <param name="labelIndex">Index for the label</param>
+		/// <returns>Notification info for the requested label in the job</returns>
+		[HttpGet]
+		[Route("/api/v1/jobs/{jobId}/labels/{labelIndex}/notifications")]
+		public async Task<ActionResult<GetNotificationResponse>> GetLabelNotificationsAsync(JobId jobId, int labelIndex)
+		{
+			IJob? job = await _jobService.GetJobAsync(jobId);
+			if (job == null)
+			{
+				return NotFound(jobId);
+			}
+
+			INotificationSubscription? subscription;
+			if (!job.LabelIdxToTriggerId.ContainsKey(labelIndex))
+			{
+				subscription = null;
+			}
+			else
+			{
+				subscription = await _notificationService.GetSubscriptionsAsync(job.LabelIdxToTriggerId[labelIndex], User);
+			}
+			return new GetNotificationResponse(subscription);
+		}
+	}
+}
