Horde: Rename the Horde.Build project to Horde.Server.

diff --git a/Engine/Source/Programs/Horde/Horde.Server/Jobs/Schedules/ScheduleService.cs b/Engine/Source/Programs/Horde/Horde.Server/Jobs/Schedules/ScheduleService.cs
new file mode 100644
--- /dev/null
+++ b/Engine/Source/Programs/Horde/Horde.Server/Jobs/Schedules/ScheduleService.cs
@@ -0,0 +1,539 @@
+// Copyright Epic Games, Inc. All Rights Reserved.
+
+using System;
+using System.Collections.Generic;
+using System.Diagnostics;
+using System.Linq;
+using System.Text.RegularExpressions;
+using System.Threading;
+using System.Threading.Tasks;
+using EpicGames.Core;
+using EpicGames.Redis;
+using EpicGames.Redis.Utility;
+using EpicGames.Serialization;
+using Horde.Server.Jobs.Graphs;
+using Horde.Server.Jobs.Templates;
+using Horde.Server.Perforce;
+using Horde.Server.Server;
+using Horde.Server.Streams;
+using Horde.Server.Utilities;
+using HordeCommon;
+using Microsoft.Extensions.Hosting;
+using Microsoft.Extensions.Logging;
+using Microsoft.Extensions.Options;
+using MongoDB.Driver;
+using OpenTracing;
+using OpenTracing.Util;
+using StackExchange.Redis;
+
+namespace Horde.Server.Jobs.Schedules
+{
+	/// <summary>
+	/// Manipulates schedule instances
+	/// </summary>
+	public sealed class ScheduleService : IHostedService, IDisposable
+	{
+		[RedisConverter(typeof(RedisCbConverter<QueueItem>))]
+		class QueueItem
+		{
+			[CbField("sid")]
+			public StreamId StreamId { get; set; }
+
+			[CbField("tid")]
+			public TemplateId TemplateId { get; set; }
+
+			public QueueItem()
+			{
+			}
+
+			public QueueItem(StreamId streamId, TemplateId templateId)
+			{
+				StreamId = streamId;
+				TemplateId = templateId;
+			}
+
+			public static DateTime GetTimeFromScore(double score) => DateTime.UnixEpoch + TimeSpan.FromSeconds(score);
+			public static double GetScoreFromTime(DateTime time) => (time.ToUniversalTime() - DateTime.UnixEpoch).TotalSeconds;
+		}
+
+		readonly IGraphCollection _graphs;
+		readonly ICommitService _commitService;
+		readonly IJobCollection _jobCollection;
+		readonly JobService _jobService;
+		readonly IStreamCollection _streamCollection;
+		readonly ITemplateCollection _templateCollection;
+		readonly IClock _clock;
+		readonly ILogger _logger;
+		readonly RedisService _redis;
+		static readonly RedisKey s_baseLockKey = "scheduler/locks";
+		static readonly RedisKey s_tickLockKey = s_baseLockKey.Append("/tick"); // Lock to tick the queue
+		static readonly RedisSortedSetKey<QueueItem> s_queueKey = "scheduler/queue"; // Items to tick, ordered by time
+		readonly IOptionsMonitor<GlobalConfig> _globalConfig;
+		readonly ITicker _ticker;
+
+		/// <summary>
+		/// Constructor
+		/// </summary>
+		public ScheduleService(RedisService redis, IGraphCollection graphs, ICommitService commitService, IJobCollection jobCollection, JobService jobService, IStreamCollection streamCollection, ITemplateCollection templateCollection, MongoService mongoService, IClock clock, IOptionsMonitor<GlobalConfig> globalConfig, ILogger<ScheduleService> logger)
+		{
+			_graphs = graphs;
+			_commitService = commitService;
+			_jobCollection = jobCollection;
+			_jobService = jobService;
+			_streamCollection = streamCollection;
+			_templateCollection = templateCollection;
+			_clock = clock;
+			_globalConfig = globalConfig;
+			_logger = logger;
+
+			_redis = redis;
+			if (mongoService.ReadOnlyMode)
+			{
+				_ticker = new NullTicker();
+			}
+			else
+			{
+				_ticker = clock.AddTicker<ScheduleService>(TimeSpan.FromMinutes(1.0), TickAsync, logger);
+			}
+		}
+
+		/// <inheritdoc/>
+		public async Task StartAsync(CancellationToken cancellationToken)
+		{
+			await _ticker.StartAsync();
+		}
+
+		/// <inheritdoc/>
+		public async Task StopAsync(CancellationToken cancellationToken)
+		{
+			await _ticker.StopAsync();
+		}
+
+		/// <inheritdoc/>
+		public void Dispose()
+		{
+			_ticker.Dispose();
+		}
+
+		async ValueTask TickAsync(CancellationToken cancellationToken)
+		{
+			DateTime utcNow = _clock.UtcNow;
+
+			// Update the current queue
+			await using (RedisLock sharedLock = new (_redis.GetDatabase(), s_tickLockKey))
+			{
+				if (await sharedLock.AcquireAsync(TimeSpan.FromMinutes(1.0), false))
+				{
+					await UpdateQueueAsync(utcNow);
+				}
+			}
+
+			// Keep updating schedules
+			while (!cancellationToken.IsCancellationRequested)
+			{
+				// Get the item with the lowest score (ie. the one that hasn't been updated in the longest time)
+				QueueItem? item = await PopQueueItemAsync();
+				if (item == null)
+				{
+					break;
+				}
+
+				// Acquire the lock for this schedule and update it
+				await using (RedisLock sharedLock = new RedisLock<QueueItem>(_redis.GetDatabase(), s_baseLockKey, item))
+				{
+					if (await sharedLock.AcquireAsync(TimeSpan.FromMinutes(1.0)))
+					{
+						try
+						{
+							await TriggerAsync(item.StreamId, item.TemplateId, utcNow, cancellationToken);
+						}
+						catch (OperationCanceledException)
+						{
+							throw;
+						}
+						catch (Exception ex)
+						{
+							_logger.LogError(ex, "Error while updating schedule for {StreamId}/{TemplateId}", item.StreamId, item.TemplateId);
+						}
+					}
+				}
+			}
+		}
+
+		async Task<QueueItem?> PopQueueItemAsync()
+		{
+			IDatabaseAsync target = _redis.GetDatabase();
+			for (; ; )
+			{
+				QueueItem[] items = await target.SortedSetRangeByRankAsync(s_queueKey, 0, 0);
+				if (items.Length == 0)
+				{
+					return null;
+				}
+				if (await target.SortedSetRemoveAsync(s_queueKey, items[0]))
+				{
+					return items[0];
+				}
+			}
+		}
+
+		internal async Task ResetAsync()
+		{
+			IDatabase redis = _redis.GetDatabase();
+			await redis.KeyDeleteAsync(s_queueKey);
+			await redis.KeyDeleteAsync(s_tickLockKey);
+		}
+
+		internal async Task TickForTestingAsync()
+		{
+			await UpdateQueueAsync(_clock.UtcNow);
+			await TickAsync(CancellationToken.None);
+		}
+
+		/// <summary>
+		/// Get the current set of streams and ensure there's an entry for each item
+		/// </summary>
+		public async Task UpdateQueueAsync(DateTime utcNow)
+		{
+			List<SortedSetEntry<QueueItem>> queueItems = new List<SortedSetEntry<QueueItem>>();
+
+			GlobalConfig globalConfig = _globalConfig.CurrentValue;
+			foreach(StreamConfig streamConfig in globalConfig.Streams)
+			{
+				IStream stream = await _streamCollection.GetAsync(streamConfig);
+				foreach ((TemplateId templateId, ITemplateRef templateRef) in stream.Templates)
+				{
+					if (templateRef.Schedule != null)
+					{
+						DateTime? nextTriggerTimeUtc = templateRef.Schedule.GetNextTriggerTimeUtc(_clock.TimeZone);
+						if (nextTriggerTimeUtc != null)
+						{
+							if (utcNow > nextTriggerTimeUtc.Value)
+							{
+								double score = QueueItem.GetScoreFromTime(nextTriggerTimeUtc.Value);
+								queueItems.Add(new SortedSetEntry<QueueItem>(new QueueItem(stream.Id, templateId), score));
+
+								await _streamCollection.UpdateScheduleTriggerAsync(stream, templateId, utcNow);
+							}
+						}
+					}
+				}
+			}
+
+			await _redis.GetDatabase().SortedSetAddAsync(s_queueKey, queueItems.ToArray());
+		}
+
+		/// <summary>
+		/// Trigger a schedule to run
+		/// </summary>
+		/// <param name="streamId">Stream for the schedule</param>
+		/// <param name="templateId"></param>
+		/// <param name="utcNow"></param>
+		/// <param name="cancellationToken"></param>
+		/// <returns>Async task</returns>
+		internal async Task<bool> TriggerAsync(StreamId streamId, TemplateId templateId, DateTime utcNow, CancellationToken cancellationToken)
+		{
+			GlobalConfig globalConfig = _globalConfig.CurrentValue;
+			if (!globalConfig.TryGetStream(streamId, out StreamConfig? streamConfig))
+			{
+				return false;
+			}
+
+			IStream? stream = await _streamCollection.GetAsync(streamConfig);
+			if (stream == null || !stream.Templates.TryGetValue(templateId, out ITemplateRef? templateRef))
+			{
+				return false;
+			}
+
+			ITemplateSchedule? schedule = templateRef.Schedule;
+			if (schedule == null)
+			{
+				return false;
+			}
+
+			using IScope scope = GlobalTracer.Instance.BuildSpan("ScheduleService.TriggerAsync").StartActive();
+			scope.Span.SetTag("StreamId", stream.Id);
+			scope.Span.SetTag("TemplateId", templateId);
+
+			Stopwatch stopwatch = Stopwatch.StartNew();
+			_logger.LogInformation("Updating schedule for {StreamId} template {TemplateId}", stream.Id, templateId);
+
+			// Get a list of jobs that we need to remove
+			List<JobId> removeJobIds = new List<JobId>();
+			foreach (JobId activeJobId in schedule.ActiveJobs)
+			{
+				IJob? job = await _jobService.GetJobAsync(activeJobId);
+				if (job == null || job.Batches.All(x => x.State == JobStepBatchState.Complete))
+				{
+					_logger.LogInformation("Removing active job {JobId}", activeJobId);
+					removeJobIds.Add(activeJobId);
+				}
+			}
+			await _streamCollection.UpdateScheduleTriggerAsync(stream, templateId, removeJobs: removeJobIds);
+
+			// If the stream is paused, bail out
+			if (stream.IsPaused(utcNow))
+			{
+				_logger.LogDebug("Skipping schedule update for stream {StreamId}. It has been paused until {PausedUntil} with comment '{PauseComment}'.", stream.Id, stream.PausedUntil, stream.PauseComment);
+				return false;
+			}
+
+			// Trigger this schedule
+			try
+			{
+				await TriggerAsync(stream, templateId, templateRef, schedule, schedule.ActiveJobs.Count - removeJobIds.Count, utcNow, cancellationToken);
+			}
+			catch (Exception ex)
+			{
+				_logger.LogError(ex, "Failed to start schedule {StreamId}/{TemplateId}", stream.Id, templateId);
+			}
+
+			// Print some timing info
+			stopwatch.Stop();
+			_logger.LogInformation("Updated schedule for {StreamId} template {TemplateId} in {TimeSeconds}ms", stream.Id, templateId, (long)stopwatch.Elapsed.TotalMilliseconds);
+			return true;
+		}
+
+		/// <summary>
+		/// Trigger a schedule to run
+		/// </summary>
+		/// <param name="stream">Stream for the schedule</param>
+		/// <param name="templateId"></param>
+		/// <param name="templateRef"></param>
+		/// <param name="schedule"></param>
+		/// <param name="numActiveJobs"></param>
+		/// <param name="utcNow">The current time</param>
+		/// <param name="cancellationToken"></param>
+		/// <returns>Async task</returns>
+		private async Task TriggerAsync(IStream stream, TemplateId templateId, ITemplateRef templateRef, ITemplateSchedule schedule, int numActiveJobs, DateTime utcNow, CancellationToken cancellationToken)
+		{
+			// Check we're not already at the maximum number of allowed jobs
+			if (schedule.Config.MaxActive != 0 && numActiveJobs >= schedule.Config.MaxActive)
+			{
+				_logger.LogInformation("Skipping trigger of {StreamId} template {TemplateId} - already have maximum number of jobs running ({NumJobs})", stream.Id, templateId, schedule.Config.MaxActive);
+				foreach (JobId jobId in schedule.ActiveJobs)
+				{
+					_logger.LogInformation("Active job for {StreamId} template {TemplateId}: {JobId}", stream.Id, templateId, jobId);
+				}
+				return;
+			}
+
+			// Minimum changelist number, inclusive
+			int minChangeNumber = schedule.LastTriggerChange;
+			if (minChangeNumber > 0 && !schedule.Config.RequireSubmittedChange)
+			{
+				minChangeNumber--;
+			}
+
+			// Maximum changelist number, exclusive
+			int? maxChangeNumber = null;
+
+			// Get the maximum number of changes to trigger
+			int maxNewChanges = 1;
+			if (schedule.Config.MaxChanges != 0)
+			{
+				maxNewChanges = schedule.Config.MaxChanges;
+			}
+			if (schedule.Config.MaxActive != 0)
+			{
+				maxNewChanges = Math.Min(maxNewChanges, schedule.Config.MaxActive - numActiveJobs);
+			}
+
+			// Create a timer to limit the amount we look back through P4 history
+			Stopwatch timer = Stopwatch.StartNew();
+
+			// Create a file filter
+			FileFilter? fileFilter = null;
+			if (schedule.Config.Files != null)
+			{
+				fileFilter = new FileFilter(schedule.Config.Files);
+			}
+
+			// Cache the Perforce history as we're iterating through changes to improve query performance
+			ICommitCollection commits = _commitService.GetCollection(stream.Config);
+			IAsyncEnumerable<ICommit> commitEnumerable = commits.FindAsync(minChangeNumber, null, null, schedule.Config.GetAllCommitTags(), cancellationToken);
+			await using IAsyncEnumerator<ICommit> commitEnumerator = commitEnumerable.GetAsyncEnumerator(cancellationToken);
+
+			// Start as many jobs as possible
+			List<(int Change, int CodeChange)> triggerChanges = new List<(int, int)>();
+			while (triggerChanges.Count < maxNewChanges)
+			{
+				cancellationToken.ThrowIfCancellationRequested();
+
+				// Get the next valid change
+				int change;
+				ICommit? commit;
+
+				if (schedule.Config.Gate != null)
+				{
+					change = await GetNextChangeForGateAsync(stream.Id, templateId, schedule.Config.Gate, minChangeNumber, maxChangeNumber, cancellationToken);
+					commit = await commits.FindAsync(change, change, 1, null, cancellationToken).FirstOrDefaultAsync(cancellationToken); // May be a change in a different stream
+				}
+				else if (await commitEnumerator.MoveNextAsync(cancellationToken))
+				{
+					commit = commitEnumerator.Current;
+					change = commit.Number;
+				}
+				else
+				{
+					commit = null;
+					change = 0;
+				}
+
+				// Quit if we didn't find anything
+				if (change <= 0)
+				{
+					break;
+				}
+				if (change < minChangeNumber)
+				{
+					break;
+				}
+				if (change == minChangeNumber && (schedule.Config.RequireSubmittedChange || triggerChanges.Count > 0))
+				{
+					break;
+				}
+
+				// Adjust the changelist for the desired filter
+				if (commit == null || await ShouldBuildChangeAsync(commit, schedule.Config.GetAllCommitTags(), fileFilter, cancellationToken))
+				{
+					int codeChange = change;
+
+					ICommit? lastCodeCommit = await commits.GetLastCodeChangeAsync(change, cancellationToken);
+					if (lastCodeCommit != null)
+					{
+						codeChange = lastCodeCommit.Number;
+					}
+					else
+					{
+						_logger.LogWarning("Unable to find code change for CL {Change}", change);
+					}
+
+					triggerChanges.Add((change, codeChange));
+				}
+
+				// Check we haven't exceeded the time limit
+				if (timer.Elapsed > TimeSpan.FromMinutes(2.0))
+				{
+					_logger.LogError("Querying for changes to trigger for {StreamId} template {TemplateId} has taken {Time}. Aborting.", stream.Id, templateId, timer.Elapsed);
+					break;
+				}
+				
+				// Update the remaining range of changes to check for
+				maxChangeNumber = change - 1;
+				if(maxChangeNumber < minChangeNumber)
+				{
+					break;
+				}
+			}
+
+			// Early out if there's nothing to do
+			if (triggerChanges.Count == 0)
+			{
+				_logger.LogInformation("Skipping trigger of {StreamName} template {TemplateId} - no candidate changes after CL {LastTriggerChange}", stream.Id, templateId, schedule.LastTriggerChange);
+				return;
+			}
+
+			// Get the matching template
+			ITemplate template = await _templateCollection.GetOrAddAsync(templateRef.Config);
+
+			// Register the graph for it
+			IGraph graph = await _graphs.AddAsync(template, stream.Config.InitialAgentType);
+
+			// We may need to submit a new change for any new jobs. This only makes sense if there's one change.
+			if (template.SubmitNewChange != null)
+			{
+				int newChange = await commits.CreateNewAsync(template, cancellationToken);
+				ICommit? newCodeChange = await commits.GetLastCodeChangeAsync(newChange, cancellationToken);
+				triggerChanges = new List<(int, int)> { (newChange, newCodeChange?.Number ?? newChange) };
+			}
+
+			// Try to start all the new jobs
+			_logger.LogInformation("Starting {NumJobs} new jobs for {StreamId} template {TemplateId} (active: {NumActive}, max new: {MaxNewJobs})", triggerChanges.Count, stream.Id, templateId, numActiveJobs, maxNewChanges);
+			foreach ((int change, int codeChange) in triggerChanges.OrderBy(x => x.Change))
+			{
+				cancellationToken.ThrowIfCancellationRequested();
+				List<string> defaultArguments = template.GetDefaultArguments();
+
+				CreateJobOptions options = new CreateJobOptions(templateRef.Config);
+				options.Priority = template.Priority;
+				options.Arguments.AddRange(template.GetDefaultArguments());
+
+				IJob newJob = await _jobService.CreateJobAsync(null, stream.Config, templateId, template.Hash, graph, template.Name, change, codeChange, options);
+				_logger.LogInformation("Started new job for {StreamId} template {TemplateId} at CL {Change} (Code CL {CodeChange}): {JobId}", stream.Id, templateId, change, codeChange, newJob.Id);
+				await _streamCollection.UpdateScheduleTriggerAsync(stream, templateId, utcNow, change, new List<JobId> { newJob.Id }, new List<JobId>());
+			}
+		}
+
+		/// <summary>
+		/// Tests whether a schedule should build a particular change, based on its requested change filters
+		/// </summary>
+		/// <param name="commit">The commit details</param>
+		/// <param name="filterTags">Filter for the tags to trigger a build</param>
+		/// <param name="fileFilter">Filter for the files to trigger a build</param>
+		/// <param name="cancellationToken">Cancellation token for the operation</param>
+		/// <returns></returns>
+		private async ValueTask<bool> ShouldBuildChangeAsync(ICommit commit, List<CommitTag>? filterTags, FileFilter? fileFilter, CancellationToken cancellationToken)
+		{
+			if (Regex.IsMatch(commit.Description, @"^\s*#\s*skipci", RegexOptions.Multiline))
+			{
+				return false;
+			}
+			if (filterTags != null && filterTags.Count > 0)
+			{
+				IReadOnlyList<CommitTag> commitTags = await commit.GetTagsAsync(cancellationToken);
+				if (!commitTags.Any(x => filterTags.Contains(x)))
+				{
+					_logger.LogDebug("Not building change {Change} ({ChangeTags}) due to filter tags ({FilterTags})", commit.Number, String.Join(", ", commitTags.Select(x => x.ToString())), String.Join(", ", filterTags.Select(x => x.ToString())));
+					return false;
+				}
+			}
+			if (fileFilter != null)
+			{
+				if (!await commit.MatchesFilterAsync(fileFilter, cancellationToken))
+				{
+					_logger.LogDebug("Not building change {Change} due to file filter", commit.Number);
+					return false;
+				}
+			}
+			return true;
+		}
+
+		/// <summary>
+		/// Gets the next change to build for a schedule on a gate
+		/// </summary>
+		/// <returns></returns>
+		private async Task<int> GetNextChangeForGateAsync(StreamId streamId, TemplateId templateRefId, ScheduleGateConfig gate, int? minChange, int? maxChange, CancellationToken cancellationToken)
+		{
+			for (; ; )
+			{
+				cancellationToken.ThrowIfCancellationRequested();
+
+				List<IJob> jobs = await _jobCollection.FindAsync(streamId: streamId, templates: new[] { gate.TemplateId }, minChange: minChange, maxChange: maxChange, count: 1);
+				if (jobs.Count == 0)
+				{
+					return 0;
+				}
+
+				IJob job = jobs[0];
+
+				IGraph? graph = await _graphs.GetAsync(job.GraphHash);
+				if (graph != null)
+				{
+					(JobStepState, JobStepOutcome)? state = job.GetTargetState(graph, gate.Target);
+					if (state != null && state.Value.Item1 == JobStepState.Completed)
+					{
+						JobStepOutcome outcome = state.Value.Item2;
+						if (outcome == JobStepOutcome.Success || outcome == JobStepOutcome.Warnings)
+						{
+							return job.Change;
+						}
+						_logger.LogInformation("Skipping trigger of {StreamName} template {TemplateId} - last {OtherTemplateRefId} job ({JobId}) ended with errors", streamId, templateRefId, gate.TemplateId, job.Id);
+					}
+				}
+
+				maxChange = job.Change - 1;
+			}
+		}
+	}
+}
diff --git a/Engine/Source/Programs/Horde/Horde.Server/Jobs/Schedules/Schedules.cs b/Engine/Source/Programs/Horde/Horde.Server/Jobs/Schedules/Schedules.cs
new file mode 100644
--- /dev/null
+++ b/Engine/Source/Programs/Horde/Horde.Server/Jobs/Schedules/Schedules.cs
@@ -0,0 +1,140 @@
+// Copyright Epic Games, Inc. All Rights Reserved.
+
+using System;
+using System.Collections.Generic;
+using Horde.Server.Perforce;
+using Horde.Server.Streams;
+using Horde.Server.Utilities;
+
+namespace Horde.Server.Jobs.Schedules
+{
+	/// <summary>
+	/// Response describing a schedule
+	/// </summary>
+	public class GetScheduleResponse
+	{
+		/// <summary>
+		/// Whether the schedule is currently enabled
+		/// </summary>
+		public bool Enabled { get; set; }
+
+		/// <summary>
+		/// Maximum number of scheduled jobs at once
+		/// </summary>
+		public int MaxActive { get; set; }
+
+		/// <summary>
+		/// Maximum number of changes the schedule can fall behind head revision. If greater than zero, builds will be triggered for every submitted changelist until the backlog is this size.
+		/// </summary>
+		public int MaxChanges { get; set; }
+
+		/// <summary>
+		/// Whether the build requires a change to be submitted
+		/// </summary>
+		public bool RequireSubmittedChange { get; set; }
+
+		/// <summary>
+		/// Gate for this schedule to trigger
+		/// </summary>
+		public ScheduleGateConfig? Gate { get; set; }
+
+		/// <summary>
+		/// Which commits to run this job for
+		/// </summary>
+		public List<CommitTag>? Commits { get; set; }
+
+		/// <summary>
+		/// The types of changes to run for
+		/// </summary>
+		[Obsolete("Use Commits instead")]
+		public List<ChangeContentFlags>? Filter { get; set; }
+
+		/// <summary>
+		/// Parameters for the template
+		/// </summary>
+		public Dictionary<string, string> TemplateParameters { get; set; }
+
+		/// <summary>
+		/// New patterns for the schedule
+		/// </summary>
+		public List<SchedulePatternConfig> Patterns { get; set; }
+
+		/// <summary>
+		/// Last changelist number that this was triggered for
+		/// </summary>
+		public int LastTriggerChange { get; set; }
+
+		/// <summary>
+		/// Last time that the schedule was triggered
+		/// </summary>
+		public DateTimeOffset LastTriggerTime { get; set; }
+
+		/// <summary>
+		/// Next trigger times for schedule
+		/// </summary>
+		public List<DateTime> NextTriggerTimesUTC { get; set; }
+
+		/// <summary>
+		/// List of active jobs
+		/// </summary>
+		public List<JobId> ActiveJobs { get; set; }
+
+		/// <summary>
+		/// Constructor
+		/// </summary>
+		/// <param name="schedule">Schedule to construct from</param>
+		/// <param name="schedulerTimeZone">The scheduler time zone</param>
+		public GetScheduleResponse(ITemplateSchedule schedule, TimeZoneInfo schedulerTimeZone)
+		{
+			Enabled = schedule.Config.Enabled;
+			MaxActive = schedule.Config.MaxActive;
+			MaxChanges = schedule.Config.MaxChanges;
+			RequireSubmittedChange = schedule.Config.RequireSubmittedChange;
+			Gate = schedule.Config.Gate;
+			Commits = schedule.Config.Commits;
+#pragma warning disable CS0618 // Type or member is obsolete
+			Filter = schedule.Config.Filter;
+#pragma warning restore CS0618 // Type or member is obsolete
+			TemplateParameters = schedule.Config.TemplateParameters;
+			Patterns = schedule.Config.Patterns;
+			LastTriggerChange = schedule.LastTriggerChange;
+			LastTriggerTime = schedule.LastTriggerTimeUtc;
+			ActiveJobs = new List<JobId>(schedule.ActiveJobs);
+
+			DateTime curTime = schedule.LastTriggerTimeUtc;
+			NextTriggerTimesUTC = new List<DateTime>();
+			for (int i = 0; i < 16; i++)
+			{
+				DateTime? nextTime = schedule.Config.GetNextTriggerTimeUtc(curTime, schedulerTimeZone);
+				if (nextTime == null)
+				{
+					break;
+				}
+				
+				curTime = nextTime.Value;
+				NextTriggerTimesUTC.Add(curTime);
+			}
+
+		}
+	}
+
+	/// <summary>
+	/// Response describing when a schedule is expected to trigger
+	/// </summary>
+	public class GetScheduleForecastResponse
+	{
+		/// <summary>
+		/// Next trigger times
+		/// </summary>
+		public List<DateTime> Times { get; set; }
+
+		/// <summary>
+		/// Constructor
+		/// </summary>
+		/// <param name="times">List of trigger times</param>
+		public GetScheduleForecastResponse(List<DateTime> times)
+		{
+			Times = times;
+		}
+	}
+}
diff --git a/Engine/Source/Programs/Horde/Horde.Server/Jobs/Templates/ITemplate.cs b/Engine/Source/Programs/Horde/Horde.Server/Jobs/Templates/ITemplate.cs
new file mode 100644
--- /dev/null
+++ b/Engine/Source/Programs/Horde/Horde.Server/Jobs/Templates/ITemplate.cs
@@ -0,0 +1,505 @@
+// Copyright Epic Games, Inc. All Rights Reserved.
+
+using System;
+using System.Collections.Generic;
+using EpicGames.Core;
+using HordeCommon;
+using MongoDB.Bson.Serialization.Attributes;
+
+namespace Horde.Server.Jobs.Templates
+{
+	/// <summary>
+	/// Base class for parameters used to configure templates via the new build dialog
+	/// </summary>
+	[BsonDiscriminator(RootClass = true)]
+	[BsonKnownTypes(typeof(GroupParameter), typeof(TextParameter), typeof(ListParameter), typeof(BoolParameter))]
+	public abstract class Parameter
+	{
+		/// <summary>
+		/// Gets the default arguments for this parameter and its children
+		/// </summary>
+		/// <param name="defaultArguments">List of default arguments</param>
+		public abstract void GetDefaultArguments(List<string> defaultArguments);
+
+		/// <summary>
+		/// Convert this parameter to data for serialization
+		/// </summary>
+		/// <returns>Serializable parameter data</returns>
+		public abstract ParameterData ToData();
+	}
+
+	/// <summary>
+	/// Used to group a number of other parameters
+	/// </summary>
+	public class GroupParameter : Parameter
+	{
+		/// <summary>
+		/// Label to display next to this parameter
+		/// </summary>
+		public string Label { get; set; }
+
+		/// <summary>
+		/// How to display this group
+		/// </summary>
+		public GroupParameterStyle Style { get; set; }
+
+		/// <summary>
+		/// List of child parameters
+		/// </summary>
+		public List<Parameter> Children { get; set; }
+
+		/// <summary>
+		/// Private constructor for serialization
+		/// </summary>
+		public GroupParameter()
+		{
+			Label = null!;
+			Children = null!;
+		}
+
+		/// <summary>
+		/// Constructor
+		/// </summary>
+		/// <param name="label">Name of the group</param>
+		/// <param name="style">How to display this group</param>
+		/// <param name="children">List of child parameters</param>
+		public GroupParameter(string label, GroupParameterStyle style, List<Parameter> children)
+		{
+			Label = label;
+			Style = style;
+			Children = children;
+		}
+
+		/// <summary>
+		/// Gets the default arguments for this parameter and its children
+		/// </summary>
+		/// <param name="defaultArguments">List of default arguments</param>
+		public override void GetDefaultArguments(List<string> defaultArguments)
+		{
+			foreach (Parameter child in Children)
+			{
+				child.GetDefaultArguments(defaultArguments);
+			}
+		}
+
+		/// <summary>
+		/// Convert this parameter to data for serialization
+		/// </summary>
+		/// <returns>Serializable parameter data</returns>
+		public override ParameterData ToData()
+		{
+			return new GroupParameterData(Label, Style, Children.ConvertAll(x => x.ToData()));
+		}
+	}
+
+	/// <summary>
+	/// Free-form text entry parameter
+	/// </summary>
+	public sealed class TextParameter : Parameter
+	{
+		/// <summary>
+		/// Label to display next to this parameter. Should default to the parameter name.
+		/// </summary>
+		public string Label { get; set; }
+
+		/// <summary>
+		/// Argument to add (will have the value of this field appended)
+		/// </summary>
+		public string Argument { get; set; }
+
+		/// <summary>
+		/// Default value for this argument
+		/// </summary>
+		public string Default { get; set; }
+
+		/// <summary>
+		/// Hint text to display when the field is empty
+		/// </summary>
+		public string? Hint { get; set; }
+
+		/// <summary>
+		/// Regex used to validate values entered into this text field.
+		/// </summary>
+		[BsonIgnoreIfNull]
+		public string? Validation { get; set; }
+
+		/// <summary>
+		/// Message displayed to explain valid values if validation fails.
+		/// </summary>
+		[BsonIgnoreIfNull]
+		public string? Description { get; set; }
+
+		/// <summary>
+		/// Tool tip text to display
+		/// </summary>
+		public string? ToolTip { get; set; }
+
+		/// <summary>
+		/// Private constructor for serialization
+		/// </summary>
+		public TextParameter()
+		{
+			Label = null!;
+			Argument = null!;
+			Default = String.Empty;
+		}
+
+		/// <summary>
+		/// Constructor
+		/// </summary>
+		/// <param name="label">Parameter to pass this value to the BuildGraph script with</param>
+		/// <param name="argument">Label to show next to the parameter</param>
+		/// <param name="defaultValue">Default value for this argument</param>
+		/// <param name="hint">Hint text to display</param>
+		/// <param name="validation">Regex used to validate entries</param>
+		/// <param name="description">Message displayed for invalid values</param>
+		/// <param name="toolTip">Tool tip text to display</param>
+		public TextParameter(string label, string argument, string defaultValue, string? hint, string? validation, string? description, string? toolTip)
+		{
+			Label = label;
+			Argument = argument;
+			Default = defaultValue;
+			Hint = hint;
+			Validation = validation;
+			Description = description;
+			ToolTip = toolTip;
+		}
+
+		/// <summary>
+		/// Gets the default arguments for this parameter and its children
+		/// </summary>
+		/// <param name="defaultArguments">List of default arguments</param>
+		public override void GetDefaultArguments(List<string> defaultArguments)
+		{
+			defaultArguments.Add(Argument + Default);
+		}
+
+		/// <summary>
+		/// Convert this parameter to data for serialization
+		/// </summary>
+		/// <returns>Serializable parameter data</returns>
+		public override ParameterData ToData()
+		{
+			return new TextParameterData(Label, Argument, Default, Hint, Validation, Description, ToolTip);
+		}
+	}
+
+	/// <summary>
+	/// Possible option for a list parameter
+	/// </summary>
+	public class ListParameterItem
+	{
+		/// <summary>
+		/// Group to display this entry in
+		/// </summary>
+		[BsonIgnoreIfNull]
+		public string? Group { get; set; }
+
+		/// <summary>
+		/// Text to display for this option.
+		/// </summary>
+		public string Text { get; set; }
+
+		/// <summary>
+		/// Argument to add with this parameter.
+		/// </summary>
+		[BsonIgnoreIfNull]
+		public string? ArgumentIfEnabled { get; set; }
+
+		/// <summary>
+		/// Argument to add with this parameter.
+		/// </summary>
+		[BsonIgnoreIfNull]
+		public string? ArgumentIfDisabled { get; set; }
+
+		/// <summary>
+		/// Whether this item is selected by default
+		/// </summary>
+		[BsonIgnoreIfDefault, BsonDefaultValue(false)]
+		public bool Default { get; set; }
+
+		/// <summary>
+		/// Default constructor for JSON serializer
+		/// </summary>
+		public ListParameterItem()
+		{
+			Text = String.Empty;
+		}
+
+		/// <summary>
+		/// Constructor
+		/// </summary>
+		/// <param name="group">The group to put this parameter in</param>
+		/// <param name="text">Text to display for this option</param>
+		/// <param name="argumentIfEnabled">Argument to add if this option is enabled</param>
+		/// <param name="argumentIfDisabled">Argument to add if this option is disabled</param>
+		/// <param name="defaultValue">Whether this item is selected by default</param>
+		public ListParameterItem(string? group, string text, string? argumentIfEnabled, string? argumentIfDisabled, bool defaultValue)
+		{
+			Group = group;
+			Text = text;
+			ArgumentIfEnabled = argumentIfEnabled;
+			ArgumentIfDisabled = argumentIfDisabled;
+			Default = defaultValue;
+		}
+
+		/// <summary>
+		/// Convert this parameter to data for serialization
+		/// </summary>
+		/// <returns>Serializable parameter data</returns>
+		public ListParameterItemData ToData()
+		{
+			return new ListParameterItemData(Group, Text, ArgumentIfEnabled, ArgumentIfDisabled, Default);
+		}
+	}
+
+	/// <summary>
+	/// Allows the user to select a value from a constrained list of choices
+	/// </summary>
+	public class ListParameter : Parameter
+	{
+		/// <summary>
+		/// Label to display next to this parameter.
+		/// </summary>
+		public string Label { get; set; }
+
+		/// <summary>
+		/// Style of picker parameter to use
+		/// </summary>
+		public ListParameterStyle Style { get; set; }
+
+		/// <summary>
+		/// List of values to display in the list
+		/// </summary>
+		public List<ListParameterItem> Items { get; set; }
+
+		/// <summary>
+		/// Tool tip text to display
+		/// </summary>
+		public string? ToolTip { get; set; }
+
+		/// <summary>
+		/// Private constructor for serialization
+		/// </summary>
+		public ListParameter()
+		{
+			Label = null!;
+			Items = null!;
+		}
+
+		/// <summary>
+		/// List of possible values
+		/// </summary>
+		/// <param name="label">Label to show next to this parameter</param>
+		/// <param name="style">Style of list parameter to use</param>
+		/// <param name="entries">Entries for this list</param>
+		/// <param name="toolTip">Tool tip text to display</param>
+		public ListParameter(string label, ListParameterStyle style, List<ListParameterItem> entries, string? toolTip)
+		{
+			Label = label;
+			Style = style;
+			Items = entries;
+			ToolTip = toolTip;
+		}
+
+		/// <summary>
+		/// Gets the default arguments for this parameter and its children
+		/// </summary>
+		/// <param name="defaultArguments">List of default arguments</param>
+		public override void GetDefaultArguments(List<string> defaultArguments)
+		{
+			foreach(ListParameterItem item in Items)
+			{
+				if (item.Default)
+				{
+					if (item.ArgumentIfEnabled != null)
+					{
+						defaultArguments.Add(item.ArgumentIfEnabled);
+					}
+				}
+				else
+				{
+					if (item.ArgumentIfDisabled != null)
+					{
+						defaultArguments.Add(item.ArgumentIfDisabled);
+					}
+				}
+			}
+		}
+
+		/// <summary>
+		/// Convert this parameter to data for serialization
+		/// </summary>
+		/// <returns>Serializable parameter data</returns>
+		public override ParameterData ToData()
+		{
+			return new ListParameterData(Label, Style, Items.ConvertAll(x => x.ToData()), ToolTip);
+		}
+	}
+
+	/// <summary>
+	/// Allows the user to toggle an option on or off
+	/// </summary>
+	public class BoolParameter : Parameter
+	{
+		/// <summary>
+		/// Label to display next to this parameter.
+		/// </summary>
+		public string Label { get; set; }
+
+		/// <summary>
+		/// Value if enabled
+		/// </summary>
+		[BsonIgnoreIfNull]
+		public string? ArgumentIfEnabled { get; set; }
+
+		/// <summary>
+		/// Value if disabled
+		/// </summary>
+		[BsonIgnoreIfNull]
+		public string? ArgumentIfDisabled { get; set; }
+
+		/// <summary>
+		/// Whether this option should be enabled by default
+		/// </summary>
+		[BsonIgnoreIfDefault, BsonDefaultValue(false)]
+		public bool Default { get; set; }
+
+		/// <summary>
+		/// Tool tip text to display
+		/// </summary>
+		public string? ToolTip { get; set; }
+
+		/// <summary>
+		/// Private constructor for serialization
+		/// </summary>
+		public BoolParameter()
+		{
+			Label = String.Empty;
+		}
+
+		/// <summary>
+		/// Constructor
+		/// </summary>
+		/// <param name="label">Label to display for this parameter</param>
+		/// <param name="argumentIfEnabled">Value if enabled</param>
+		/// <param name="argumentIfDisabled">Value if disabled</param>
+		/// <param name="defaultValue">Default value for this argument</param>
+		/// <param name="toolTip">Tool tip text to display</param>
+		public BoolParameter(string label, string? argumentIfEnabled, string? argumentIfDisabled, bool defaultValue, string? toolTip)
+		{
+			Label = label;
+			ArgumentIfEnabled = argumentIfEnabled;
+			ArgumentIfDisabled = argumentIfDisabled;
+			Default = defaultValue;
+			ToolTip = toolTip;
+		}
+
+		/// <summary>
+		/// Gets the default arguments for this parameter and its children
+		/// </summary>
+		/// <param name="defaultArguments">List of default arguments</param>
+		public override void GetDefaultArguments(List<string> defaultArguments)
+		{
+			string? defaultArgument = Default ? ArgumentIfEnabled : ArgumentIfDisabled;
+			if (!String.IsNullOrEmpty(defaultArgument))
+			{
+				defaultArguments.Add(defaultArgument);
+			}
+		}
+		/// <summary>
+		/// Convert this parameter to data for serialization
+		/// </summary>
+		/// <returns><see cref="BoolParameterData"/> instance</returns>
+		public override ParameterData ToData()
+		{
+			return new BoolParameterData(Label, ArgumentIfEnabled, ArgumentIfDisabled, Default, ToolTip);
+		}
+	}
+
+	/// <summary>
+	/// Document describing a job template. These objects are considered immutable once created and uniquely referenced by hash, in order to de-duplicate across all job runs.
+	/// </summary>
+	public interface ITemplate
+	{
+		/// <summary>
+		/// Hash of this template
+		/// </summary>
+		public ContentHash Hash { get; }
+
+		/// <summary>
+		/// Name of the template.
+		/// </summary>
+		public string Name { get; }
+
+		/// <summary>
+		/// Description for the template
+		/// </summary>
+		public string? Description { get; }
+
+		/// <summary>
+		/// Priority of this job
+		/// </summary>
+		public Priority? Priority { get; }
+
+		/// <summary>
+		/// Whether to allow preflights for this job type
+		/// </summary>
+		public bool AllowPreflights { get; }
+
+		/// <summary>
+		/// Whether to always issues for jobs using this template
+		/// </summary>
+		public bool UpdateIssues { get; }
+
+		/// <summary>
+		/// Whether to promote issues by default for jobs using this template
+		/// </summary>
+		public bool PromoteIssuesByDefault { get; }
+
+		/// <summary>
+		/// Agent type to use for parsing the job state
+		/// </summary>
+		public string? InitialAgentType { get; }
+
+		/// <summary>
+		/// Path to a file within the stream to submit to generate a new changelist for jobs
+		/// </summary>
+		public string? SubmitNewChange { get; }
+
+		/// <summary>
+		/// Description for new changelists
+		/// </summary>
+		public string? SubmitDescription { get; }
+
+		/// <summary>
+		/// Optional predefined user-defined properties for this job
+		/// </summary>
+		public IReadOnlyList<string> Arguments { get; }
+
+		/// <summary>
+		/// Parameters for this template
+		/// </summary>
+		public IReadOnlyList<Parameter> Parameters { get; }
+	}
+
+	/// <summary>
+	/// Extension methods for templates
+	/// </summary>
+	public static class TemplateExtensions
+	{
+		/// <summary>
+		/// Gets the arguments for default options in this template. Does not include the standard template arguments.
+		/// </summary>
+		/// <returns>List of default arguments</returns>
+		public static List<string> GetDefaultArguments(this ITemplate template)
+		{
+			List<string> defaultArguments = new List<string>(template.Arguments);
+			foreach (Parameter parameter in template.Parameters)
+			{
+				parameter.GetDefaultArguments(defaultArguments);
+			}
+			return defaultArguments;
+		}
+	}
+}
diff --git a/Engine/Source/Programs/Horde/Horde.Server/Jobs/Templates/ITemplateCollection.cs b/Engine/Source/Programs/Horde/Horde.Server/Jobs/Templates/ITemplateCollection.cs
new file mode 100644
--- /dev/null
+++ b/Engine/Source/Programs/Horde/Horde.Server/Jobs/Templates/ITemplateCollection.cs
@@ -0,0 +1,28 @@
+// Copyright Epic Games, Inc. All Rights Reserved.
+
+using System.Collections.Generic;
+using System.Threading.Tasks;
+using EpicGames.Core;
+
+namespace Horde.Server.Jobs.Templates
+{
+	/// <summary>
+	/// Interface for a collection of template documents
+	/// </summary>
+	public interface ITemplateCollection
+	{
+		/// <summary>
+		/// Gets a template by ID
+		/// </summary>
+		/// <param name="templateId">Unique id of the template</param>
+		/// <returns>The template document</returns>
+		Task<ITemplate?> GetAsync(ContentHash templateId);
+
+		/// <summary>
+		/// Gets a template instance from its configuration
+		/// </summary>
+		/// <param name="templateConfig">The template configuration</param>
+		/// <returns>Template instance</returns>
+		Task<ITemplate> GetOrAddAsync(TemplateConfig templateConfig);
+	}
+}
diff --git a/Engine/Source/Programs/Horde/Horde.Server/Jobs/TestData/ITestDataCollection.cs b/Engine/Source/Programs/Horde/Horde.Server/Jobs/TestData/ITestDataCollection.cs
new file mode 100644
--- /dev/null
+++ b/Engine/Source/Programs/Horde/Horde.Server/Jobs/TestData/ITestDataCollection.cs
@@ -0,0 +1,128 @@
+// Copyright Epic Games, Inc. All Rights Reserved.
+
+using System;
+using System.Collections.Generic;
+using System.Threading.Tasks;
+using Horde.Server.Streams;
+using Horde.Server.Utilities;
+using MongoDB.Bson;
+
+namespace Horde.Server.Jobs.TestData
+{	
+	/// <summary>
+	/// Collection of test data documents
+	/// </summary>
+	public interface ITestDataCollection
+	{
+		/// <summary>
+		/// Gets the tests running in provided streams
+		/// </summary>
+		/// <param name="streamIds"></param>
+		/// <returns></returns>
+		Task<List<ITestStream>> FindTestStreams(StreamId[] streamIds);
+
+		/// <summary>
+		/// 
+		/// </summary>
+		/// <param name="streamIds"></param>
+		/// <param name="metaIds"></param>
+		/// <param name="testIds"></param>
+		/// <param name="suiteIds"></param>
+		/// <param name="minCreateTime"></param>
+		/// <param name="maxCreateTime"></param>
+		/// <param name="minChange"></param>
+		/// <param name="maxChange"></param>
+		/// <returns></returns>
+		Task<List<ITestDataRef>> FindTestRefs(StreamId[] streamIds, TestMetaId[]? metaIds = null, TestId[]? testIds = null, TestSuiteId[]? suiteIds = null, DateTime? minCreateTime = null, DateTime? maxCreateTime = null, int? minChange = null, int? maxChange = null);
+
+
+		/// <summary>
+		/// Find test details
+		/// </summary>
+		/// <param name="ids"></param>
+		/// <returns></returns>
+		Task<List<ITestDataDetails>> FindTestDetails(TestRefId[] ids);
+
+		/// <summary>
+		/// Find tests
+		/// </summary>
+		/// <param name="testIds"></param>
+		/// <returns></returns>
+		Task<List<ITest>> FindTests(TestId[] testIds);
+
+		/// <summary>
+		/// Find test suites
+		/// </summary>
+		/// <param name="suiteIds"></param>
+		/// <returns></returns>
+		Task<List<ITestSuite>> FindTestSuites(TestSuiteId[] suiteIds);
+
+		/// <summary>
+		/// Creates a new test data document
+		/// </summary>
+		/// <param name="job">The job containing the step</param>
+		/// <param name="step">The step producing the data</param>
+		/// <param name="data">The test data generated by the job step</param>
+		/// <returns>The new stream document</returns>
+		Task<List<ITestData>> AddAsync(IJob job, IJobStep step, (string key, BsonDocument value)[] data);
+
+		/// <summary>
+		/// Gets a stream by ID
+		/// </summary>
+		/// <param name="id">Unique id of the stream</param>
+		/// <returns>The stream document</returns>
+		Task<ITestData?> GetAsync(ObjectId id);
+
+		/// <summary>
+		/// Searches for test data that matches a set of criteria
+		/// </summary>
+		/// <param name="streamId">The stream id</param>
+		/// <param name="minChange">The minimum changelist number to return (inclusive)</param>
+		/// <param name="maxChange">The maximum changelist number to return (inclusive)</param>
+		/// <param name="jobId">The job id</param>
+		/// <param name="stepId">The unique step id</param>
+		/// <param name="key">Key identifying the result to return</param>
+		/// <param name="index">Offset within the results to return</param>
+		/// <param name="count">Number of results to return</param>
+		/// <returns>The stream document</returns>
+		Task<List<ITestData>> FindAsync(StreamId? streamId, int? minChange, int? maxChange, JobId? jobId, SubResourceId? stepId, string? key = null, int index = 0, int count = 10);
+
+
+		/// <summary>
+		/// Find test meta data
+		/// </summary>
+		/// <param name="projectNames"></param>
+		/// <param name="platforms"></param>
+		/// <param name="configurations"></param>
+		/// <param name="buildTargets"></param>
+		/// <param name="rhi"></param>
+		/// <param name="variation"></param>
+		/// <param name="metaIds"></param>
+		/// <returns></returns>
+		Task<List<ITestMeta>> FindTestMeta(string[]? projectNames = null, string[]? platforms = null, string[]? configurations = null, string[]? buildTargets = null, string? rhi = null, string? variation = null, TestMetaId[]? metaIds = null);
+
+		/// <summary>
+		/// Delete the test data
+		/// </summary>
+		/// <param name="id">Unique id of the test data</param>
+		/// <returns>Async task</returns>
+		Task DeleteAsync(ObjectId id);
+
+		/// <summary>
+		/// Prune test data from collection
+		/// </summary>
+		/// <returns></returns>
+		public Task PruneDataAsync();
+
+		/// <summary>
+		/// Run a tick on the collection
+		/// </summary>
+		/// <returns></returns>
+		public Task<bool> UpdateAsync(int retainMonths);
+
+		/// <summary>
+		/// Upgrades the collection
+		/// </summary>
+		public Task UpgradeAsync();
+	}
+}
diff --git a/Engine/Source/Programs/Horde/Horde.Server/Jobs/TestData/TestDataController.cs b/Engine/Source/Programs/Horde/Horde.Server/Jobs/TestData/TestDataController.cs
new file mode 100644
--- /dev/null
+++ b/Engine/Source/Programs/Horde/Horde.Server/Jobs/TestData/TestDataController.cs
@@ -0,0 +1,380 @@
+// Copyright Epic Games, Inc. All Rights Reserved.
+
+using System;
+using System.Collections.Generic;
+using System.Linq;
+using System.Threading.Tasks;
+using Horde.Server.Acls;
+using Horde.Server.Streams;
+using Horde.Server.Utilities;
+using Microsoft.AspNetCore.Authorization;
+using Microsoft.AspNetCore.Mvc;
+using MongoDB.Bson;
+using EpicGames.Core;
+using System.Text;
+using System.Text.Json.Serialization;
+using System.Text.Json;
+using System.Threading;
+using Microsoft.Extensions.Options;
+using Horde.Server.Server;
+
+namespace Horde.Server.Jobs.TestData
+{
+	/// <summary>
+	/// Controller for the /api/v1/testdata endpoint
+	/// </summary>
+	[ApiController]
+	[Authorize]
+	[Route("[controller]")]
+	public class TestDataController : ControllerBase
+	{
+		/// <summary>
+		/// Collection of job documents
+		/// </summary>
+		private readonly JobService _jobService;
+
+		/// <summary>
+		/// Collection of test data documents
+		/// </summary>
+		private readonly ITestDataCollection _testDataCollection;
+
+		readonly TestDataService _testDataService;
+
+		readonly IStreamCollection _streamCollection;
+	
+		readonly IOptionsSnapshot<GlobalConfig> _globalConfig;
+
+		/// <summary>
+		/// Constructor
+		/// </summary>
+		public TestDataController(TestDataService testDataService, IStreamCollection streamCollection, JobService jobService, ITestDataCollection testDataCollection, IOptionsSnapshot<GlobalConfig> globalConfig)
+		{
+			_jobService = jobService;
+			_testDataCollection = testDataCollection;
+			_testDataService = testDataService;
+			_streamCollection = streamCollection;
+			_globalConfig = globalConfig;
+		}
+
+		/// <summary>
+		/// Get metadata 
+		/// </summary>
+		/// <param name="projects"></param>
+		/// <param name="platforms"></param>
+		/// <param name="targets"></param>
+		/// <param name="configurations"></param>
+		/// <returns></returns>
+		[HttpGet]
+		[Route("/api/v2/testdata/metadata")]
+		[ProducesResponseType(typeof(List<GetTestMetaResponse>), 200)]
+		public async Task<ActionResult<List<GetTestMetaResponse>>> GetTestMetaAsync(
+			[FromQuery(Name = "project")] string[]? projects = null,
+			[FromQuery(Name = "platform")] string[]? platforms = null,
+			[FromQuery(Name = "target")] string[]? targets = null,
+			[FromQuery(Name = "configuration")] string[]? configurations = null)
+		{
+			List<GetTestMetaResponse> responses = new List<GetTestMetaResponse>();
+			List<ITestMeta> metaData = await _testDataService.FindTestMeta(projects, platforms, configurations, targets);
+			metaData.ForEach(m => responses.Add(new GetTestMetaResponse(m)));
+			return responses;
+		}
+
+		/// <summary>
+		/// Get test details from provided refs
+		/// </summary>
+		/// <param name="ids"></param>
+		/// <returns></returns>
+		[HttpGet]
+		[Route("/api/v2/testdata/details")]
+		[ProducesResponseType(typeof(List<GetTestDataDetailsResponse>), 200)]
+		public async Task<ActionResult<List<GetTestDataDetailsResponse>>> GetTestDetailsAsync([FromQuery(Name = "id")] string[] ids)
+		{
+			TestRefId[] idValues = Array.ConvertAll(ids, x => TestRefId.Parse(x));
+			List<ITestDataDetails> details = await _testDataService.FindTestDetails(idValues);
+			return details.Select(d => new GetTestDataDetailsResponse(d)).ToList();
+		}
+
+		/// <summary>
+		/// Get test details from provided refs
+		/// </summary>
+		/// <param name="request"></param>
+		/// <returns></returns>
+		[HttpPost]
+		[Route("/api/v2/testdata/tests")]
+		[ProducesResponseType(typeof(List<GetTestResponse>), 200)]
+		public async Task<ActionResult<List<GetTestResponse>>> GetTestsAsync([FromBody] GetTestsRequest request)
+		{
+			HashSet<string> testIds = new HashSet<string>(request.testIds);
+
+			List<ITest> testValues = await _testDataService.FindTests(testIds.Select(x => TestId.Parse(x)).ToArray());			
+
+			return testValues.Select(x => new GetTestResponse(x)).ToList();
+		}
+
+		/// <summary>
+		/// Get stream test data for the provided ids
+		/// </summary>
+		/// <param name="streamIds"></param>
+		/// <returns></returns>
+		[HttpGet]
+		[Route("/api/v2/testdata/streams")]
+		[ProducesResponseType(typeof(List<GetTestStreamResponse>), 200)]
+		public async Task<ActionResult<List<GetTestStreamResponse>>> GetTestStreamsAsync([FromQuery(Name = "Id")] string[] streamIds)
+		{
+			StreamId[] streamIdValues = Array.ConvertAll(streamIds, x => new StreamId(x));
+
+			List<StreamId> queryStreams = new List<StreamId>();
+			List<GetTestStreamResponse> responses = new List<GetTestStreamResponse>();
+
+			// authorize streams
+			foreach (StreamId streamId in streamIdValues)
+			{
+				if (_globalConfig.Value.TryGetStream(streamId, out StreamConfig? streamConfig) && streamConfig.Authorize(AclAction.ViewJob, User))
+				{
+					queryStreams.Add(streamId);
+				}
+			}
+
+			if (queryStreams.Count == 0)
+			{
+				return responses;
+			}
+
+			HashSet<TestId> testIds = new HashSet<TestId>();
+			HashSet<TestMetaId> metaIds = new HashSet<TestMetaId>();
+
+			List<ITestStream> streams = await _testDataService.FindTestStreams(queryStreams.ToArray());
+
+			// flatten requested streams to single service queries		
+			HashSet<TestSuiteId> suiteIds = new HashSet<TestSuiteId>();
+			for (int i = 0; i < streams.Count; i++)
+			{
+				foreach(TestId testId in streams[i].Tests)
+				{
+					testIds.Add(testId);
+				}
+
+				foreach(TestSuiteId suiteId in streams[i].TestSuites)
+				{
+					suiteIds.Add(suiteId);
+				}				
+			}
+
+			List<ITestSuite> suites = new List<ITestSuite>();
+			if (suiteIds.Count > 0)
+			{
+				suites = await _testDataService.FindTestSuites(suiteIds.ToArray());
+			}
+
+			List<ITest> tests = new List<ITest>();
+			if (testIds.Count > 0)
+			{
+				tests = await _testDataService.FindTests(testIds.ToArray());
+			}
+
+			// gather all meta data
+			List<ITestMeta> metaData = new List<ITestMeta>();
+			foreach (ITest test in tests)
+			{
+				foreach (TestMetaId metaId in test.Metadata)
+				{
+					metaIds.Add(metaId);
+				}
+			}
+
+			foreach (ITestSuite suite in suites)
+			{
+				foreach (TestMetaId metaId in suite.Metadata)
+				{
+					metaIds.Add(metaId);
+				}
+			}
+
+			if (metaIds.Count > 0)
+			{
+				metaData = await _testDataService.FindTestMeta(metaIds: metaIds.ToArray());
+			}
+
+			// generate individual stream responses
+			streams.ForEach(s => {
+				
+				List<ITest> streamTests = tests.Where(x => s.Tests.Contains(x.Id)).ToList();
+
+				List<ITestSuite> streamSuites = new List<ITestSuite>();
+				foreach (TestSuiteId suiteId in s.TestSuites)
+				{
+					ITestSuite? suite = suites.FirstOrDefault(x => x.Id == suiteId);
+					if (suite != null)
+					{
+						streamSuites.Add(suite);
+					}
+				}
+
+				HashSet<TestMetaId> streamMetaIds = new HashSet<TestMetaId>();
+				foreach (ITest test in streamTests)
+				{
+					foreach (TestMetaId id in test.Metadata)
+					{
+						streamMetaIds.Add(id);
+					}						
+				}
+
+				foreach (ITestSuite suite in streamSuites)
+				{
+					foreach (TestMetaId id in suite.Metadata)
+					{
+						streamMetaIds.Add(id);
+					}
+				}
+
+				List<ITestMeta> streamMetaData = metaData.Where(x => streamMetaIds.Contains(x.Id)).ToList();
+
+				responses.Add(new GetTestStreamResponse(s.StreamId, streamTests, streamSuites, streamMetaData));				
+
+			});
+
+			return responses;
+		}
+
+		/// <summary>
+		/// Gets test data refs 
+		/// </summary>
+		/// <param name="streamIds"></param>
+		/// <param name="testIds"></param>
+		/// <param name="suiteIds"></param>
+		/// <param name="metaIds"></param>
+		/// <param name="minCreateTime"></param>
+		/// <param name="maxCreateTime"></param>
+		/// <param name="minChange"></param>
+		/// <param name="maxChange"></param>
+		/// <returns></returns>
+		[HttpGet]
+		[Route("/api/v2/testdata/refs")]
+		[ProducesResponseType(typeof(List<GetTestDataRefResponse>), 200)]
+		public async Task<ActionResult<List<GetTestDataRefResponse>>> GetTestDataRefAsync(
+			[FromQuery(Name = "Id")] string[] streamIds,
+			[FromQuery(Name = "Mid")] string[] metaIds,
+			[FromQuery(Name = "Tid")] string[]? testIds = null,
+			[FromQuery(Name = "Sid")] string[]? suiteIds = null,			
+			[FromQuery] DateTimeOffset? minCreateTime = null,
+			[FromQuery] DateTimeOffset? maxCreateTime = null,
+			[FromQuery] int? minChange = null, 
+			[FromQuery] int? maxChange = null)
+		{
+			StreamId[] streamIdValues = Array.ConvertAll(streamIds, x => new StreamId(x));
+
+			List<StreamId> queryStreams = new List<StreamId>();
+			List<GetTestDataRefResponse> responses = new List<GetTestDataRefResponse>();
+
+			// authorize streams
+			foreach (StreamId streamId in streamIdValues)
+			{
+				if (_globalConfig.Value.TryGetStream(streamId, out StreamConfig? streamConfig) && streamConfig.Authorize(AclAction.ViewJob, User))
+				{
+					queryStreams.Add(streamId);
+				}
+			}
+
+			if (queryStreams.Count == 0)
+			{
+				return responses;
+			}
+
+			List<ITestDataRef> dataRefs = await _testDataService.FindTestRefs(queryStreams.ToArray(), metaIds.ConvertAll(x => TestMetaId.Parse(x)).ToArray(), testIds, suiteIds, minCreateTime?.UtcDateTime, maxCreateTime?.UtcDateTime, minChange, maxChange);
+
+			dataRefs.ForEach(d => responses.Add(new GetTestDataRefResponse(d)));
+
+			return responses;
+		}
+
+		/// <summary>
+		/// Creates a new TestData document
+		/// </summary>
+		/// <returns>The stream document</returns>
+		[HttpPost]
+		[Route("/api/v1/testdata")]
+		public async Task<ActionResult<CreateTestDataResponse>> CreateAsync(CreateTestDataRequest request)
+		{
+			IJob? job = await _jobService.GetJobAsync(JobId.Parse(request.JobId));
+			if (job == null)
+			{
+				return NotFound();
+			}
+			if (!_globalConfig.Value.Authorize(job, AclAction.UpdateJob, User))
+			{
+				return Forbid();
+			}
+
+			IJobStep? jobStep;
+			if (!job.TryGetStep(request.StepId.ToSubResourceId(), out jobStep))
+			{
+				return NotFound();
+			}
+			
+			List<ITestData> testData = await _testDataCollection.AddAsync(job, jobStep, new (string key, BsonDocument value)[] { (request.Key, new BsonDocument(request.Data))});
+			return new CreateTestDataResponse(testData[0].Id.ToString());
+		}
+
+		/// <summary>
+		/// Searches for test data that matches a set of criteria
+		/// </summary>
+		/// <param name="streamId">The stream id</param>
+		/// <param name="minChange">The minimum changelist number to return (inclusive)</param>
+		/// <param name="maxChange">The maximum changelist number to return (inclusive)</param>
+		/// <param name="jobId">The job id</param>
+		/// <param name="jobStepId">The unique step id</param>
+		/// <param name="key">Key identifying the result to return</param>
+		/// <param name="index">Offset within the results to return</param>
+		/// <param name="count">Number of results to return</param>
+		/// <param name="filter">Filter for properties to return</param>
+		/// <returns>The stream document</returns>
+		[HttpGet]
+		[Route("/api/v1/testdata")]
+		[ProducesResponseType(typeof(List<GetTestDataResponse>), 200)]
+		public async Task<ActionResult<List<object>>> FindTestDataAsync([FromQuery] string? streamId = null, [FromQuery] int? minChange = null, [FromQuery] int? maxChange = null, string? jobId = null, string? jobStepId = null, string? key = null, int index = 0, int count = 10, PropertyFilter? filter = null)
+		{
+			StreamId? streamIdValue = null;
+			if(streamId != null)
+			{
+				streamIdValue = new StreamId(streamId);
+			}
+
+			List<object> results = new List<object>();
+
+			List<ITestData> documents = await _testDataCollection.FindAsync(streamIdValue, minChange, maxChange, (jobId == null)? null : JobId.Parse(jobId), jobStepId?.ToSubResourceId(), key, index, count);
+			foreach (ITestData document in documents)
+			{
+				if (await _jobService.AuthorizeAsync(document.JobId, AclAction.ViewJob, User, _globalConfig.Value))
+				{
+					results.Add(PropertyFilter.Apply(new GetTestDataResponse(document), filter));
+				}
+			}
+
+			return results;
+		}
+
+		/// <summary>
+		/// Retrieve information about a specific issue
+		/// </summary>
+		/// <param name="testDataId">Id of the document to get information about</param>
+		/// <param name="filter">Filter for the properties to return</param>
+		/// <returns>List of matching agents</returns>
+		[HttpGet]
+		[Route("/api/v1/testdata/{testDataId}")]
+		[ProducesResponseType(typeof(GetTestDataResponse), 200)]
+		public async Task<ActionResult<object>> GetTestDataAsync(string testDataId, [FromQuery] PropertyFilter? filter = null)
+		{
+			ITestData? testData = await _testDataCollection.GetAsync(testDataId.ToObjectId());
+			if (testData == null)
+			{
+				return NotFound();
+			}
+			if (!await _jobService.AuthorizeAsync(testData.JobId, AclAction.ViewJob, User, _globalConfig.Value))
+			{
+				return Forbid();
+			}
+
+			return PropertyFilter.Apply(new GetTestDataResponse(testData), filter);
+		}
+	}
+}
diff --git a/Engine/Source/Programs/Horde/Horde.Server/Logs/Builder/LocalLogBuilder.cs b/Engine/Source/Programs/Horde/Horde.Server/Logs/Builder/LocalLogBuilder.cs
new file mode 100644
--- /dev/null
+++ b/Engine/Source/Programs/Horde/Horde.Server/Logs/Builder/LocalLogBuilder.cs
@@ -0,0 +1,235 @@
+// Copyright Epic Games, Inc. All Rights Reserved.
+
+using System;
+using System.Collections.Concurrent;
+using System.Collections.Generic;
+using System.IO;
+using System.Threading.Tasks;
+using Horde.Server.Logs.Data;
+using Horde.Server.Utilities;
+
+namespace Horde.Server.Logs.Builder
+{
+	/// <summary>
+	/// In-memory implementation of a log write buffer
+	/// </summary>
+	class LocalLogBuilder : ILogBuilder
+	{
+		/// <summary>
+		/// Stores information about a sub-chunk which is still being written to
+		/// </summary>
+		class PendingSubChunk
+		{
+			/// <summary>
+			/// The type of data stored in this sub chunk
+			/// </summary>
+			public LogType Type { get; }
+
+			/// <summary>
+			/// Data for the sub-chunk
+			/// </summary>
+			public MemoryStream Stream { get; } = new MemoryStream(4096);
+
+			/// <summary>
+			/// Constructor
+			/// </summary>
+			/// <param name="type">Type of data stored in this subchunk</param>
+			public PendingSubChunk(LogType type)
+			{
+				Type = type;
+			}
+
+			/// <summary>
+			/// Converts this pending sub-chunk to a concrete LogSubChunkData object
+			/// </summary>
+			/// <param name="offset">Offset of this sub-chunk within the file</param>
+			/// <param name="lineIndex">The line index</param>
+			/// <returns>New sub-chunk data object</returns>
+			public LogSubChunkData ToSubChunkData(long offset, int lineIndex)
+			{
+				return new LogSubChunkData(Type, offset, lineIndex, new ReadOnlyLogText(Stream.ToArray()));
+			}
+		}
+
+		/// <summary>
+		/// Stores information about a pending log chunk
+		/// </summary>
+		class PendingChunk
+		{
+			/// <summary>
+			/// Time that this chunk was created
+			/// </summary>
+			public DateTime CreateTimeUtc { get; set; } = DateTime.UtcNow;
+
+			/// <summary>
+			/// Array of sub chunks for this chunk
+			/// </summary>
+			public List<PendingSubChunk> _subChunks = new List<PendingSubChunk>();
+
+			/// <summary>
+			/// The next sub-chunk
+			/// </summary>
+			public PendingSubChunk _nextSubChunk;
+
+			/// <summary>
+			/// Total length of this chunk
+			/// </summary>
+			public long _length;
+
+			/// <summary>
+			/// Whether the chunk is complete
+			/// </summary>
+			public bool _complete;
+
+			/// <summary>
+			/// Constructor
+			/// </summary>
+			/// <param name="type">Type of data to store in this file</param>
+			public PendingChunk(LogType type)
+			{
+				_nextSubChunk = new PendingSubChunk(type);
+			}
+
+			/// <summary>
+			/// Complete the current sub-chunk
+			/// </summary>
+			public void CompleteSubChunk()
+			{
+				if (_nextSubChunk.Stream.Length > 0)
+				{
+					_subChunks.Add(_nextSubChunk);
+					_nextSubChunk = new PendingSubChunk(_nextSubChunk.Type);
+				}
+			}
+		}
+
+		/// <summary>
+		/// Current log chunk state
+		/// </summary>
+		readonly ConcurrentDictionary<(LogId, long), PendingChunk> _pendingChunks = new ConcurrentDictionary<(LogId, long), PendingChunk>();
+
+		/// <inheritdoc/>
+		public bool FlushOnShutdown => true;
+
+		/// <inheritdoc/>
+		public Task<bool> AppendAsync(LogId logId, long chunkOffset, long writeOffset, int writeLineIndex, int writeLineCount, ReadOnlyMemory<byte> data, LogType type)
+		{
+			PendingChunk? pendingChunk;
+			while (!_pendingChunks.TryGetValue((logId, chunkOffset), out pendingChunk))
+			{
+				if(writeOffset != chunkOffset)
+				{
+					return Task.FromResult(false);
+				}
+
+				pendingChunk = new PendingChunk(type);
+
+				if (_pendingChunks.TryAdd((logId, chunkOffset), pendingChunk))
+				{
+					break;
+				}
+			}
+
+			lock (pendingChunk)
+			{
+				if (!pendingChunk._complete && chunkOffset + pendingChunk._length == writeOffset)
+				{
+					pendingChunk._nextSubChunk.Stream.Write(data.Span);
+					pendingChunk._length += data.Span.Length;
+					return Task.FromResult(true);
+				}
+			}
+			return Task.FromResult(false);
+		}
+
+		/// <inheritdoc/>
+		public Task CompleteSubChunkAsync(LogId logId, long offset)
+		{
+			PendingChunk? pendingChunk;
+			if (_pendingChunks.TryGetValue((logId, offset), out pendingChunk))
+			{
+				lock (pendingChunk)
+				{
+					pendingChunk.CompleteSubChunk();
+				}
+			}
+			return Task.CompletedTask;
+		}
+
+		/// <inheritdoc/>
+		public Task CompleteChunkAsync(LogId logId, long offset)
+		{
+			PendingChunk? pendingChunk;
+			if (_pendingChunks.TryGetValue((logId, offset), out pendingChunk))
+			{
+				lock (pendingChunk)
+				{
+					if (!pendingChunk._complete)
+					{
+						pendingChunk._complete = true;
+						pendingChunk.CompleteSubChunk();
+					}
+				}
+			}
+			return Task.CompletedTask;
+		}
+
+		/// <inheritdoc/>
+		public Task RemoveChunkAsync(LogId logId, long offset)
+		{
+			_pendingChunks.TryRemove((logId, offset), out _);
+			return Task.CompletedTask;
+		}
+
+		/// <inheritdoc/>
+		public Task<LogChunkData?> GetChunkAsync(LogId logId, long offset, int lineIndex)
+		{
+			PendingChunk? pendingChunk;
+			if (_pendingChunks.TryGetValue((logId, offset), out pendingChunk))
+			{
+				long subChunkOffset = offset;
+				int subChunkLineIndex = lineIndex;
+
+				List<LogSubChunkData> subChunks = new List<LogSubChunkData>();
+				foreach(PendingSubChunk pendingSubChunk in pendingChunk._subChunks)
+				{
+					LogSubChunkData subChunk = pendingSubChunk.ToSubChunkData(subChunkOffset, subChunkLineIndex);
+					subChunkOffset += subChunk.Length;
+					subChunkLineIndex += subChunk.LineCount;
+					subChunks.Add(subChunk);
+				}
+
+				PendingSubChunk nextSubChunk = pendingChunk._nextSubChunk;
+				if (nextSubChunk.Stream.Length > 0)
+				{
+					LogSubChunkData subChunkData = nextSubChunk.ToSubChunkData(subChunkOffset, subChunkLineIndex);
+					subChunks.Add(subChunkData);
+				}
+
+				return Task.FromResult<LogChunkData?>(new LogChunkData(offset, lineIndex, subChunks));
+			}
+
+			return Task.FromResult<LogChunkData?>(null);
+		}
+
+		/// <inheritdoc/>
+		public Task<List<(LogId, long)>> TouchChunksAsync(TimeSpan minAge)
+		{
+			DateTime utcNow = DateTime.UtcNow;
+
+			List<(LogId, long)> chunks = new List<(LogId, long)>();
+			foreach (KeyValuePair<(LogId, long), PendingChunk> pendingChunk in _pendingChunks.ToArray())
+			{
+				lock (pendingChunk.Value)
+				{
+					if (pendingChunk.Value.CreateTimeUtc < utcNow - minAge)
+					{
+						chunks.Add(pendingChunk.Key);
+						pendingChunk.Value.CreateTimeUtc = utcNow;
+					}
+				}
+			}
+			return Task.FromResult(chunks);
+		}
+	}
+}
diff --git a/Engine/Source/Programs/Horde/Horde.Server/Logs/Builder/RedisLogBuilder.cs b/Engine/Source/Programs/Horde/Horde.Server/Logs/Builder/RedisLogBuilder.cs
new file mode 100644
--- /dev/null
+++ b/Engine/Source/Programs/Horde/Horde.Server/Logs/Builder/RedisLogBuilder.cs
@@ -0,0 +1,332 @@
+// Copyright Epic Games, Inc. All Rights Reserved.
+
+using System;
+using System.Collections.Generic;
+using System.Globalization;
+using System.Text.RegularExpressions;
+using System.Threading.Tasks;
+using EpicGames.Core;
+using EpicGames.Redis;
+using Horde.Server.Logs.Data;
+using Horde.Server.Utilities;
+using Microsoft.Extensions.Logging;
+using OpenTracing;
+using OpenTracing.Util;
+using StackExchange.Redis;
+
+namespace Horde.Server.Logs.Builder
+{
+	using Condition = StackExchange.Redis.Condition;
+
+	/// <summary>
+	/// Redis-based cache for log file chunks
+	/// </summary>
+	class RedisLogBuilder : ILogBuilder
+	{
+		const string ItemsKey = "log-items";
+
+		struct ChunkKeys
+		{
+			public string Prefix { get; }
+			public string Type => $"{Prefix}-type";
+			public string LineIndex => $"{Prefix}-lineindex";
+			public string Length => $"{Prefix}-length";
+			public string ChunkData => $"{Prefix}-chunk";
+			public string SubChunkData => $"{Prefix}-subchunk";
+			public string Complete => $"{Prefix}-complete";
+
+			public ChunkKeys(LogId logId, long offset)
+			{
+				Prefix = $"log-{logId}-chunk-{offset}-builder";
+			}
+
+			public static bool TryParse(string prefix, out LogId logId, out long offset)
+			{
+				Match match = Regex.Match(prefix, "log-([^-]+)-chunk-([^-]+)-builder");
+				if (!match.Success || !LogId.TryParse(match.Groups[1].Value, out logId) || !Int64.TryParse(match.Groups[2].Value, out offset))
+				{
+					logId = LogId.Empty;
+					offset = 0;
+					return false;
+				}
+				return true;
+			}
+		}
+
+		/// <inheritdoc/>
+		public bool FlushOnShutdown => false;
+
+		/// <summary>
+		/// The Redis database connection pool
+		/// </summary>
+		readonly RedisConnectionPool _redisConnectionPool;
+
+		/// <summary>
+		/// Logger for debug output
+		/// </summary>
+		readonly ILogger _logger;
+
+		/// <summary>
+		/// Constructor
+		/// </summary>
+		/// <param name="redisConnectionPool">The redis database singleton</param>
+		/// <param name="logger">Logger for debug output</param>
+		public RedisLogBuilder(RedisConnectionPool redisConnectionPool, ILogger logger)
+		{
+			_redisConnectionPool = redisConnectionPool;
+			_logger = logger;
+		}
+
+		/// <inheritdoc/>
+		public async Task<bool> AppendAsync(LogId logId, long chunkOffset, long writeOffset, int writeLineIndex, int writeLineCount, ReadOnlyMemory<byte> data, LogType type)
+		{
+			IDatabase redisDb = _redisConnectionPool.GetDatabase();
+			ChunkKeys keys = new ChunkKeys(logId, chunkOffset);
+
+			if(chunkOffset == writeOffset)
+			{
+				using IScope scope = GlobalTracer.Instance.BuildSpan("Redis.CreateChunk").StartActive();
+				scope.Span.SetTag("LogId", logId.ToString());
+				scope.Span.SetTag("Offset", chunkOffset.ToString(CultureInfo.InvariantCulture));
+				scope.Span.SetTag("WriteOffset", writeOffset.ToString(CultureInfo.InvariantCulture));
+
+				ITransaction createTransaction = redisDb.CreateTransaction();
+				createTransaction.AddCondition(Condition.SortedSetNotContains(ItemsKey, keys.Prefix));
+				_ = createTransaction.SortedSetAddAsync(ItemsKey, keys.Prefix, DateTime.UtcNow.Ticks, flags: CommandFlags.FireAndForget);
+				_ = createTransaction.StringSetAsync(keys.Type, (int)type, flags: CommandFlags.FireAndForget);
+				_ = createTransaction.StringSetAsync(keys.Length, data.Length, flags: CommandFlags.FireAndForget);
+				_ = createTransaction.StringSetAsync(keys.LineIndex, writeLineIndex + writeLineCount, flags: CommandFlags.FireAndForget);
+				_ = createTransaction.StringSetAsync(keys.SubChunkData, data, flags: CommandFlags.FireAndForget);
+				if (await createTransaction.ExecuteAsync())
+				{
+					_logger.LogTrace("Created {Size} bytes in {Key}", data.Length, keys.SubChunkData);
+					return true;
+				}
+			}
+
+			{
+				using IScope scope = GlobalTracer.Instance.BuildSpan("Redis.AppendChunk").StartActive();
+				scope.Span.SetTag("LogId", logId.ToString());
+				scope.Span.SetTag("Offset", chunkOffset.ToString(CultureInfo.InvariantCulture));
+				scope.Span.SetTag("WriteOffset", writeOffset.ToString(CultureInfo.InvariantCulture));
+
+				ITransaction appendTransaction = redisDb.CreateTransaction();
+				appendTransaction.AddCondition(Condition.SortedSetContains(ItemsKey, keys.Prefix));
+				appendTransaction.AddCondition(Condition.KeyNotExists(keys.Complete));
+				appendTransaction.AddCondition(Condition.StringEqual(keys.Type, (int)type));
+				appendTransaction.AddCondition(Condition.StringEqual(keys.Length, (int)(writeOffset - chunkOffset)));
+				_ = appendTransaction.StringAppendAsync(keys.SubChunkData, data, flags: CommandFlags.FireAndForget);
+				_ = appendTransaction.StringSetAsync(keys.Length, (int)(writeOffset - chunkOffset) + data.Length, flags: CommandFlags.FireAndForget);
+				_ = appendTransaction.StringSetAsync(keys.LineIndex, writeLineIndex + writeLineCount, flags: CommandFlags.FireAndForget);
+				if (await appendTransaction.ExecuteAsync())
+				{
+					_logger.LogTrace("Appended {Size} to {Key}", data.Length, keys.SubChunkData);
+					return true;
+				}
+			}
+
+			return false;
+		}
+
+		/// <inheritdoc/>
+		public async Task CompleteSubChunkAsync(LogId logId, long offset)
+		{
+			IDatabase redisDb = _redisConnectionPool.GetDatabase();
+			ChunkKeys keys = new ChunkKeys(logId, offset);
+			
+			const int MaxRetries = 10;
+			int numTries;
+			for (numTries = 0; numTries < MaxRetries; numTries++)
+			{
+				using IScope scope = GlobalTracer.Instance.BuildSpan("Redis.CompleteSubChunk").StartActive();
+				scope.Span.SetTag("LogId", logId.ToString());
+				scope.Span.SetTag("Offset", offset.ToString(CultureInfo.InvariantCulture));
+
+				LogType type = (LogType)(int)await redisDb.StringGetAsync(keys.Type);
+				int length = (int)await redisDb.StringGetAsync(keys.Length);
+				int lineIndex = (int)await redisDb.StringGetAsync(keys.LineIndex);
+				long chunkDataLength = await redisDb.StringLengthAsync(keys.ChunkData);
+
+				RedisValue subChunkTextValue = await redisDb.StringGetAsync(keys.SubChunkData);
+				if (subChunkTextValue.IsNullOrEmpty)
+				{
+					return;
+				}
+
+				ReadOnlyLogText subChunkText = new ReadOnlyLogText(subChunkTextValue);
+				LogSubChunkData subChunkData = new LogSubChunkData(type, offset + length - subChunkText.Data.Length, lineIndex - subChunkText.LineCount, subChunkText);
+
+				byte[] subChunkDataBytes;
+				try
+				{
+					subChunkDataBytes = subChunkData.ToByteArray(_logger);
+				}
+				catch (Exception ex)
+				{
+					throw new LogTextParseException($"Error while parsing sub-chunk for {logId} offset {offset}", ex);
+				}
+
+				ITransaction writeTransaction = redisDb.CreateTransaction();
+
+				writeTransaction.AddCondition(Condition.StringLengthEqual(keys.ChunkData, chunkDataLength));
+				writeTransaction.AddCondition(Condition.StringLengthEqual(keys.SubChunkData, subChunkText.Length));
+				writeTransaction.AddCondition(Condition.StringEqual(keys.Length, length));
+				writeTransaction.AddCondition(Condition.StringEqual(keys.LineIndex, lineIndex));
+				Task<long> newLength = writeTransaction.StringAppendAsync(keys.ChunkData, subChunkDataBytes);
+				_ = writeTransaction.KeyDeleteAsync(keys.SubChunkData, flags: CommandFlags.FireAndForget);
+
+				if (await writeTransaction.ExecuteAsync())
+				{
+					_logger.LogDebug("Completed sub-chunk for log {LogId} chunk offset {Offset} -> sub-chunk size {SubChunkSize}, chunk size {ChunkSize}", logId, offset, subChunkDataBytes.Length, await newLength);
+					return;
+				}
+				await writeTransaction.WaitAndIgnoreCancellations(newLength);
+
+				// Cool down before retrying
+				await Task.Delay(100);
+			}
+
+			_logger.LogError("Unable to complete sub-chunk for {LogId} at {Offset} after {NumTries} tries", logId, offset, numTries);
+			throw new Exception($"Unable to complete sub-chunk for {logId} at {offset} after {numTries} tries");
+		}
+
+		/// <inheritdoc/>
+		public async Task CompleteChunkAsync(LogId logId, long offset)
+		{
+			IDatabase redisDb = _redisConnectionPool.GetDatabase();
+			using IScope scope = GlobalTracer.Instance.BuildSpan("Redis.CompleteChunk").StartActive();
+			scope.Span.SetTag("LogId", logId.ToString());
+			scope.Span.SetTag("Offset", offset.ToString(CultureInfo.InvariantCulture));
+
+			ChunkKeys keys = new ChunkKeys(logId, offset);
+
+			ITransaction transaction = redisDb.CreateTransaction();
+			transaction.AddCondition(Condition.SortedSetContains(ItemsKey, keys.Prefix));
+			_ = transaction.StringSetAsync(keys.Complete, true, flags: CommandFlags.FireAndForget);
+			if(!await transaction.ExecuteAsync())
+			{
+				_logger.LogDebug("Log {LogId} chunk offset {Offset} is not in Redis builder", logId, offset);
+				return;
+			}
+
+			await CompleteSubChunkAsync(logId, offset);
+		}
+
+		/// <inheritdoc/>
+		public async Task RemoveChunkAsync(LogId logId, long offset)
+		{
+			IDatabase redisDb = _redisConnectionPool.GetDatabase();
+			using IScope scope = GlobalTracer.Instance.BuildSpan("Redis.RemoveChunk").StartActive();
+			scope.Span.SetTag("LogId", logId.ToString());
+			scope.Span.SetTag("Offset", offset.ToString(CultureInfo.InvariantCulture));
+
+			ChunkKeys keys = new ChunkKeys(logId, offset);
+
+			ITransaction transaction = redisDb.CreateTransaction();
+			_ = transaction.KeyDeleteAsync(keys.Type, flags: CommandFlags.FireAndForget);
+			_ = transaction.KeyDeleteAsync(keys.LineIndex, flags: CommandFlags.FireAndForget);
+			_ = transaction.KeyDeleteAsync(keys.Length, flags: CommandFlags.FireAndForget);
+			_ = transaction.KeyDeleteAsync(keys.ChunkData, flags: CommandFlags.FireAndForget);
+			_ = transaction.KeyDeleteAsync(keys.SubChunkData, flags: CommandFlags.FireAndForget);
+			_ = transaction.KeyDeleteAsync(keys.Complete, flags: CommandFlags.FireAndForget);
+
+			_ = transaction.SortedSetRemoveAsync(ItemsKey, keys.Prefix, flags: CommandFlags.FireAndForget);
+			await transaction.ExecuteAsync();
+		}
+
+		/// <inheritdoc/>
+		public async Task<LogChunkData?> GetChunkAsync(LogId logId, long offset, int lineIndex)
+		{
+			IDatabase redisDb = _redisConnectionPool.GetDatabase();
+			ChunkKeys keys = new ChunkKeys(logId, offset);
+			const int MaxRetries = 10;
+			int numTries;
+			for (numTries = 0; numTries < MaxRetries; numTries++)
+			{
+				using IScope scope = GlobalTracer.Instance.BuildSpan("Redis.GetChunk").StartActive();
+				scope.Span.SetTag("LogId", logId.ToString());
+				scope.Span.SetTag("Offset", offset.ToString(CultureInfo.InvariantCulture));
+
+				RedisValue chunkDataValue = await redisDb.StringGetAsync(keys.ChunkData);
+
+				ReadOnlyMemory<byte> chunkData = chunkDataValue;
+
+				ITransaction transaction = redisDb.CreateTransaction();
+				transaction.AddCondition(Condition.StringLengthEqual(keys.ChunkData, chunkData.Length));
+				Task<RedisValue> typeTask = transaction.StringGetAsync(keys.Type);
+				Task<RedisValue> lastSubChunkDataTask = transaction.StringGetAsync(keys.SubChunkData);
+
+				if (await transaction.ExecuteAsync())
+				{
+					MemoryReader reader = new MemoryReader(chunkData);
+
+					long subChunkOffset = offset;
+					int subChunkLineIndex = lineIndex;
+
+					List<LogSubChunkData> subChunks = new List<LogSubChunkData>();
+					while (reader.RemainingMemory.Length > 0)
+					{
+						LogSubChunkData subChunkData = reader.ReadLogSubChunkData(subChunkOffset, subChunkLineIndex);
+						subChunkOffset += subChunkData.Length;
+						subChunkLineIndex += subChunkData.LineCount;
+						subChunks.Add(subChunkData);
+					}
+
+					RedisValue type = await typeTask;
+
+					RedisValue lastSubChunkData = await lastSubChunkDataTask;
+					if (lastSubChunkData.Length() > 0)
+					{
+						ReadOnlyLogText subChunkDataText = new ReadOnlyLogText(lastSubChunkData);
+						subChunks.Add(new LogSubChunkData((LogType)(int)type, subChunkOffset, subChunkLineIndex, subChunkDataText));
+					}
+
+					if (subChunks.Count == 0)
+					{
+						return null;
+					}
+
+					return new LogChunkData(offset, lineIndex, subChunks);
+				}
+
+				await transaction.WaitAndIgnoreCancellations(typeTask, lastSubChunkDataTask);
+
+				// Cool down before retrying
+				await Task.Delay(100);
+			}
+
+			_logger.LogWarning("Unable to get chunk for {LogId} at {Offset}:{LineIndex} after {NumTries} tries", logId, offset, lineIndex, numTries);
+			return null;
+		}
+
+		/// <inheritdoc/>
+		public async Task<List<(LogId, long)>> TouchChunksAsync(TimeSpan minAge)
+		{
+			IDatabase redisDb = _redisConnectionPool.GetDatabase();
+			
+			// Find all the chunks that are suitable for expiry
+			DateTime utcNow = DateTime.UtcNow;
+			SortedSetEntry[] entries = await redisDb.SortedSetRangeByScoreWithScoresAsync(ItemsKey, stop: (utcNow - minAge).Ticks);
+
+			// Update the score for each element in a transaction. If it succeeds, we can write the chunk. Otherwise another pod has beat us to it.
+			List<(LogId, long)> results = new List<(LogId, long)>();
+			foreach (SortedSetEntry entry in entries)
+			{
+				if (ChunkKeys.TryParse(entry.Element.ToString(), out LogId logId, out long offset))
+				{
+					ITransaction transaction = redisDb.CreateTransaction();
+					transaction.AddCondition(Condition.SortedSetEqual(ItemsKey, entry.Element, entry.Score));
+					_ = transaction.SortedSetAddAsync(ItemsKey, entry.Element, utcNow.Ticks, flags: CommandFlags.FireAndForget);
+
+					if (!await transaction.ExecuteAsync())
+					{
+						break;
+					}
+
+					results.Add((logId, offset));
+				}
+			}
+			return results;
+		}
+	}
+}
diff --git a/Engine/Source/Programs/Horde/Horde.Server/Logs/Data/LogChunkData.cs b/Engine/Source/Programs/Horde/Horde.Server/Logs/Data/LogChunkData.cs
new file mode 100644
--- /dev/null
+++ b/Engine/Source/Programs/Horde/Horde.Server/Logs/Data/LogChunkData.cs
@@ -0,0 +1,277 @@
+// Copyright Epic Games, Inc. All Rights Reserved.
+
+using System;
+using System.Collections.Generic;
+using System.Linq;
+using EpicGames.Core;
+using Microsoft.Extensions.Logging;
+
+namespace Horde.Server.Logs.Data
+{
+	/// <summary>
+	/// Data for a log chunk
+	/// </summary>
+	public class LogChunkData
+	{
+		/// <summary>
+		/// Offset of this chunk
+		/// </summary>
+		public long Offset { get; }
+
+		/// <summary>
+		/// Length of this chunk
+		/// </summary>
+		public int Length { get; }
+
+		/// <summary>
+		/// Line index of this chunk
+		/// </summary>
+		public int LineIndex { get; }
+
+		/// <summary>
+		/// Number of lines in this chunk
+		/// </summary>
+		public int LineCount { get; }
+			
+		/// <summary>
+		/// List of sub-chunks
+		/// </summary>
+		public IReadOnlyList<LogSubChunkData> SubChunks { get; }
+
+		/// <summary>
+		/// Offset of each sub-chunk within this chunk
+		/// </summary>
+		public int[] SubChunkOffset { get; }
+
+		/// <summary>
+		/// Total line count of the chunk after each sub-chunk
+		/// </summary>
+		public int[] SubChunkLineIndex { get; }
+
+		/// <summary>
+		/// Constructor
+		/// </summary>
+		/// <param name="offset">Offset of this chunk</param>
+		/// <param name="lineIndex">First line index of this chunk</param>
+		/// <param name="subChunks">Sub-chunks for this chunk</param>
+		public LogChunkData(long offset, int lineIndex, IReadOnlyList<LogSubChunkData> subChunks)
+		{
+			Offset = offset;
+			LineIndex = lineIndex;
+			SubChunks = subChunks;
+			SubChunkOffset = CreateSumLookup(subChunks, x => x.Length);
+			SubChunkLineIndex = CreateSumLookup(subChunks, x => x.LineCount);
+
+			if(subChunks.Count > 0)
+			{
+				int lastSubChunkIdx = subChunks.Count - 1;
+				Length = SubChunkOffset[lastSubChunkIdx] + subChunks[lastSubChunkIdx].Length;
+				LineCount = SubChunkLineIndex[lastSubChunkIdx] + subChunks[lastSubChunkIdx].LineCount;
+			}
+		}
+
+		/// <summary>
+		/// Creates a lookup table by summing a field across a list of subchunks
+		/// </summary>
+		/// <param name="subChunks"></param>
+		/// <param name="field"></param>
+		/// <returns></returns>
+		static int[] CreateSumLookup(IReadOnlyList<LogSubChunkData> subChunks, Func<LogSubChunkData, int> field)
+		{
+			int[] total = new int[subChunks.Count];
+
+			int value = 0;
+			for (int idx = 0; idx + 1 < subChunks.Count; idx++)
+			{
+				value += field(subChunks[idx]);
+				total[idx + 1] = value;
+			}
+
+			return total;
+		}
+
+		/// <summary>
+		/// Gets the offset of a line within the file
+		/// </summary>
+		/// <param name="lineIdx">The line index</param>
+		/// <returns>Offset of the line within the file</returns>
+		public long GetLineOffsetWithinChunk(int lineIdx)
+		{
+			int subChunkIdx = SubChunkLineIndex.BinarySearch(lineIdx);
+			if (subChunkIdx < 0)
+			{
+				subChunkIdx = ~subChunkIdx - 1;
+			}
+			return SubChunkOffset[subChunkIdx] + SubChunks[subChunkIdx].InflateText().LineOffsets[lineIdx - SubChunkLineIndex[subChunkIdx]];
+		}
+
+		/// <summary>
+		/// Gets the sub chunk index for the given line
+		/// </summary>
+		/// <param name="chunkLineIdx">Line index within the chunk</param>
+		/// <returns>Subchunk line index</returns>
+		public int GetSubChunkForLine(int chunkLineIdx)
+		{
+			int subChunkIdx = SubChunkLineIndex.BinarySearch(chunkLineIdx);
+			if (subChunkIdx < 0)
+			{
+				subChunkIdx = ~subChunkIdx - 1;
+			}
+			return subChunkIdx;
+		}
+
+		/// <summary>
+		/// Gets the index of the sub-chunk containing the given offset
+		/// </summary>
+		/// <param name="offset">Offset to search for</param>
+		/// <returns>Index of the sub-chunk</returns>
+		public int GetSubChunkForOffsetWithinChunk(int offset)
+		{
+			int subChunkIdx = SubChunkOffset.BinarySearch(offset);
+			if (subChunkIdx < 0)
+			{
+				subChunkIdx = ~subChunkIdx - 1;
+			}
+			return subChunkIdx;
+		}
+
+		/// <summary>
+		/// Signature for output data
+		/// </summary>
+		const int CurrentSignature = 'L' | ('C' << 8) | ('D' << 16);
+
+		/// <summary>
+		/// Read a log chunk from the given stream
+		/// </summary>
+		/// <param name="reader">The reader to read from</param>
+		/// <param name="offset">Offset of this chunk within the file</param>
+		/// <param name="lineIndex">Line index of this chunk</param>
+		/// <returns>New log chunk data</returns>
+		public static LogChunkData Read(MemoryReader reader, long offset, int lineIndex)
+		{
+			int signature = reader.ReadInt32();
+			if ((signature & 0xffffff) != CurrentSignature)
+			{
+				List<LogSubChunkData> subChunks = new List<LogSubChunkData>();
+				subChunks.Add(new LogSubChunkData(LogType.Json, offset, lineIndex, new ReadOnlyLogText(reader.RemainingMemory)));
+				reader.Advance(reader.RemainingMemory.Length);
+				return new LogChunkData(offset, lineIndex, subChunks);
+			}
+
+			int version = signature >> 24;
+			if (version == 0)
+			{
+				List<LogSubChunkData> subChunks = ReadSubChunkList(reader, offset, lineIndex);
+				reader.ReadVariableLengthBytesWithInt32Length();
+				return new LogChunkData(offset, lineIndex, subChunks);
+			}
+			else
+			{
+				List<LogSubChunkData> subChunks = ReadSubChunkList(reader, offset, lineIndex);
+				return new LogChunkData(offset, lineIndex, subChunks);
+			}
+		}
+
+		/// <summary>
+		/// Read a list of sub-chunks from the given stream
+		/// </summary>
+		/// <param name="reader">The reader to read from</param>
+		/// <param name="subChunkOffset">Offset of this chunk within the file</param>
+		/// <param name="subChunkLineIndex">Line index of this chunk</param>
+		/// <returns>List of sub-chunks</returns>
+		static List<LogSubChunkData> ReadSubChunkList(MemoryReader reader, long subChunkOffset, int subChunkLineIndex)
+		{
+			int numSubChunks = reader.ReadInt32();
+
+			List<LogSubChunkData> subChunks = new List<LogSubChunkData>();
+			for (int idx = 0; idx < numSubChunks; idx++)
+			{
+				LogSubChunkData subChunkData = reader.ReadLogSubChunkData(subChunkOffset, subChunkLineIndex);
+				subChunkOffset += subChunkData.Length;
+				subChunkLineIndex += subChunkData.LineCount;
+				subChunks.Add(subChunkData);
+			}
+
+			return subChunks;
+		}
+
+		/// <summary>
+		/// Construct an object from flat memory buffer
+		/// </summary>
+		/// <param name="memory">Memory buffer</param>
+		/// <param name="offset">Offset of this chunk within the file</param>
+		/// <param name="lineIndex">Line index of this chunk</param>
+		/// <returns>Log chunk data</returns>
+		public static LogChunkData FromMemory(ReadOnlyMemory<byte> memory, long offset, int lineIndex)
+		{
+			MemoryReader reader = new MemoryReader(memory);
+			LogChunkData chunkData = Read(reader, offset, lineIndex);
+			reader.CheckEmpty();
+			return chunkData;
+		}
+
+		/// <summary>
+		/// Write the chunk data to a stream
+		/// </summary>
+		/// <returns>Serialized data</returns>
+		public void Write(MemoryWriter writer, ILogger logger)
+		{
+			writer.WriteInt32(CurrentSignature | (1 << 24));
+
+			writer.WriteInt32(SubChunks.Count);
+			foreach (LogSubChunkData subChunk in SubChunks)
+			{
+				writer.WriteLogSubChunkData(subChunk, logger);
+			}
+		}
+
+		/// <summary>
+		/// Construct an object from flat memory buffer
+		/// </summary>
+		/// <param name="logger">Logger for output</param>
+		/// <returns>Log chunk data</returns>
+		public byte[] ToByteArray(ILogger logger)
+		{
+			byte[] data = new byte[GetSerializedSize(logger)];
+			MemoryWriter writer = new MemoryWriter(data);
+			Write(writer, logger);
+			writer.CheckEmpty();
+			return data;
+		}
+
+		/// <summary>
+		/// Determines the size of the serialized buffer
+		/// </summary>
+		public int GetSerializedSize(ILogger logger)
+		{
+			return sizeof(int) + (sizeof(int) + SubChunks.Sum(x => x.GetSerializedSize(logger)));
+		}
+	}
+
+	/// <summary>
+	/// Extensions for the log chunk data
+	/// </summary>
+	static class LogChunkDataExtensions
+	{
+		/// <summary>
+		/// Read a log chunk from the given stream
+		/// </summary>
+		/// <param name="reader">The reader to read from</param>
+		/// <param name="offset">Offset of this chunk within the file</param>
+		/// <param name="lineIndex">Line index of this chunk</param>
+		/// <returns>New log chunk data</returns>
+		public static LogChunkData ReadLogChunkData(this MemoryReader reader, long offset, int lineIndex)
+		{
+			return LogChunkData.Read(reader, offset, lineIndex);
+		}
+
+		/// <summary>
+		/// Write the chunk data to a stream
+		/// </summary>
+		/// <returns>Serialized data</returns>
+		public static void WriteLogChunkData(this MemoryWriter writer, LogChunkData chunkData, ILogger logger)
+		{
+			chunkData.Write(writer, logger);
+		}
+	}
+}
