Horde: Rename the Horde.Build project to Horde.Server.

diff --git a/Engine/Source/Programs/Horde/Horde.Server/Logs/Data/LogIndexData.cs b/Engine/Source/Programs/Horde/Horde.Server/Logs/Data/LogIndexData.cs
new file mode 100644
--- /dev/null
+++ b/Engine/Source/Programs/Horde/Horde.Server/Logs/Data/LogIndexData.cs
@@ -0,0 +1,604 @@
+// Copyright Epic Games, Inc. All Rights Reserved.
+
+using System;
+using System.Collections.Generic;
+using System.IO;
+using System.Linq;
+using EpicGames.Core;
+using EpicGames.Horde.Logs;
+using Horde.Server.Utilities;
+using OpenTracing;
+using OpenTracing.Util;
+
+namespace Horde.Server.Logs.Data
+{
+	/// <summary>
+	/// Contains a source block of text to be indexed
+	/// </summary>
+	public class LogIndexBlock
+	{
+		/// <summary>
+		/// First line within the file
+		/// </summary>
+		public int LineIndex { get; }
+
+		/// <summary>
+		/// Number of lines in this block
+		/// </summary>
+		public int LineCount { get; }
+
+		/// <summary>
+		/// The plain text for this block
+		/// </summary>
+		public ILogText? CachedPlainText { get; private set; }
+
+		/// <summary>
+		/// The compressed plain text data
+		/// </summary>
+		public ReadOnlyMemory<byte> CompressedPlainText { get; }
+
+		/// <summary>
+		/// Constructor
+		/// </summary>
+		/// <param name="lineIndex">Index of the first line within this block</param>
+		/// <param name="lineCount">Number of lines in the block</param>
+		/// <param name="cachedPlainText">The decompressed plain text</param>
+		/// <param name="compressedPlainText">The compressed text data</param>
+		public LogIndexBlock(int lineIndex, int lineCount, ILogText? cachedPlainText, ReadOnlyMemory<byte> compressedPlainText)
+		{
+			LineIndex = lineIndex;
+			LineCount = lineCount;
+			CachedPlainText = cachedPlainText;
+			CompressedPlainText = compressedPlainText;
+		}
+
+		/// <summary>
+		/// Accessor for the decompressed plaintext
+		/// </summary>
+		public ILogText InflatePlainText()
+		{
+			CachedPlainText ??= new ReadOnlyLogText(CompressedPlainText.DecompressBzip2());
+			return CachedPlainText;
+		}
+	}
+
+	/// <summary>
+	/// Extension methods for <see cref="LogIndexBlock"/>
+	/// </summary>
+	public static class LogIndexBlockExtensions
+	{
+		/// <summary>
+		/// Deserialize an index block
+		/// </summary>
+		/// <param name="reader">Reader to deserialize from</param>
+		/// <returns>The new index block</returns>
+		public static LogIndexBlock ReadLogIndexBlock(this MemoryReader reader)
+		{
+			int lineIndex = reader.ReadInt32();
+			int lineCount = reader.ReadInt32();
+			ReadOnlyMemory<byte> compressedPlainText = reader.ReadVariableLengthBytesWithInt32Length();
+			return new LogIndexBlock(lineIndex, lineCount, null, compressedPlainText);
+		}
+
+		/// <summary>
+		/// Serialize an index block
+		/// </summary>
+		/// <param name="writer">Writer to serialize to</param>
+		/// <param name="block">The block to serialize</param>
+		public static void WriteLogIndexBlock(this MemoryWriter writer, LogIndexBlock block)
+		{
+			writer.WriteInt32(block.LineIndex);
+			writer.WriteInt32(block.LineCount);
+			writer.WriteVariableLengthBytesWithInt32Length(block.CompressedPlainText.Span);
+		}
+
+		/// <summary>
+		/// Gets the serialized size of a block
+		/// </summary>
+		/// <param name="block">The block to serialize</param>
+		public static int GetSerializedSize(this LogIndexBlock block)
+		{
+			return sizeof(int) + sizeof(int) + (sizeof(int) + block.CompressedPlainText.Length);
+		}
+	}
+
+	/// <summary>
+	/// Data for a log chunk
+	/// </summary>
+	public class LogIndexData
+	{
+		/// <summary>
+		/// Index for tokens into the block list
+		/// </summary>
+		ReadOnlyTrie? _cachedTrie;
+
+		/// <summary>
+		/// Number of bits in the index devoted to the block index
+		/// </summary>
+		readonly int _numBlockBits;
+
+		/// <summary>
+		/// List of text blocks
+		/// </summary>
+		readonly LogIndexBlock[] _blocks;
+
+		/// <summary>
+		/// Empty index data
+		/// </summary>
+		public static LogIndexData Empty { get; } = new LogIndexData(ReadOnlyTrie.Empty, 0, Array.Empty<LogIndexBlock>());
+
+		/// <summary>
+		/// Number of lines covered by the index
+		/// </summary>
+		public int LineCount => (_blocks.Length > 0) ? (_blocks[^1].LineIndex + _blocks[^1].LineCount) : 0;
+
+		/// <summary>
+		/// Constructor
+		/// </summary>
+		/// <param name="cachedTrie">Index into the text blocks</param>
+		/// <param name="numBlockBits">Number of bits devoted to the block index</param>
+		/// <param name="blocks">Bloom filters for this log file</param>
+		public LogIndexData(ReadOnlyTrie? cachedTrie, int numBlockBits, LogIndexBlock[] blocks)
+		{
+			_cachedTrie = cachedTrie;
+			_numBlockBits = numBlockBits;
+			_blocks = blocks;
+		}
+
+		/// <summary>
+		/// Updates the blocks in this index relative to a give line index
+		/// </summary>
+		/// <param name="lineIndex"></param>
+		public void SetBaseLineIndex(int lineIndex)
+		{
+			for(int blockIdx = 0; blockIdx < _blocks.Length; blockIdx++)
+			{
+				LogIndexBlock block = _blocks[blockIdx];
+				_blocks[blockIdx] = new LogIndexBlock(lineIndex, block.LineCount, block.CachedPlainText, block.CompressedPlainText);
+				lineIndex += block.LineCount;
+			}
+		}
+
+		/// <summary>
+		/// Builds a trie for this index
+		/// </summary>
+		/// <returns>The trie for this data</returns>
+		ReadOnlyTrie BuildTrie()
+		{
+			if(_cachedTrie == null)
+			{
+				ReadOnlyTrieBuilder builder = new ReadOnlyTrieBuilder();
+				for (int blockIdx = 0; blockIdx < _blocks.Length; blockIdx++)
+				{
+					LogIndexBlock block = _blocks[blockIdx];
+					LogToken.GetTokens(block.InflatePlainText().Data.Span, token => builder.Add((token << _numBlockBits) | (uint)blockIdx));
+				}
+				_cachedTrie = builder.Build();
+			}
+			return _cachedTrie;
+		}
+
+		/// <summary>
+		/// Create an index from an array of blocks
+		/// </summary>
+		/// <param name="indexes">List of indexes to merge</param>
+		/// <returns>Index data</returns>
+		public static LogIndexData Merge(IEnumerable<LogIndexData> indexes)
+		{
+			using IScope scope = GlobalTracer.Instance.BuildSpan("LogIndexData.Merge").StartActive();
+
+			// Create the combined block list
+			LogIndexBlock[] newBlocks = indexes.SelectMany(x => x._blocks).ToArray();
+
+			// Figure out how many bits to devote to the block size
+			int newNumBlockBits = 0;
+			while(newBlocks.Length > (1 << newNumBlockBits))
+			{
+				newNumBlockBits += 4;
+			}
+
+			// Add all the blocks into a combined index
+			ReadOnlyTrieBuilder newTrieBuilder = new ReadOnlyTrieBuilder();
+
+			int blockCount = 0;
+			foreach(LogIndexData index in indexes)
+			{
+				ulong blockMask = ((1UL << index._numBlockBits) - 1);
+				foreach (ulong value in index.BuildTrie())
+				{
+					ulong token = value >> index._numBlockBits;
+					int blockIdx = (int)(value & blockMask);
+
+					ulong newToken = (token << newNumBlockBits) | (ulong)(long)(blockCount + blockIdx);
+					newTrieBuilder.Add(newToken);
+				}
+				blockCount += index._blocks.Length;
+			}
+
+			// Construct the new index data
+			ReadOnlyTrie newTrie = newTrieBuilder.Build();
+			return new LogIndexData(newTrie, newNumBlockBits, newBlocks);
+		}
+
+		/// <summary>
+		/// Search for the given text in the index
+		/// </summary>
+		/// <param name="firstLineIndex">First line index to search from</param>
+		/// <param name="text">Text to search for</param>
+		/// <param name="stats">Receives stats for the search</param>
+		/// <returns>List of line numbers for the text</returns>
+		public IEnumerable<int> Search(int firstLineIndex, SearchText text, SearchStats stats)
+		{
+			int lastBlockCount = 0;
+			foreach (int blockIdx in EnumeratePossibleBlocks(text.Bytes, firstLineIndex))
+			{
+				LogIndexBlock block = _blocks[blockIdx];
+
+				stats.NumScannedBlocks++;
+				stats.NumDecompressedBlocks += (block.CachedPlainText == null)? 1 : 0;
+
+				stats.NumSkippedBlocks += blockIdx - lastBlockCount;
+				lastBlockCount = blockIdx + 1;
+
+				// Decompress the text
+				ILogText blockText = block.InflatePlainText();
+
+				// Find the initial offset within this block
+				int offset = 0;
+				if(firstLineIndex > block.LineIndex)
+				{
+					int lineIndexWithinBlock = firstLineIndex - block.LineIndex;
+					offset = blockText.LineOffsets[lineIndexWithinBlock];
+				}
+
+				// Search within this block
+				for(; ;)
+				{
+					// Find the next offset
+					int nextOffset = blockText.Data.Span.FindNextOcurrence(offset, text);
+					if(nextOffset == -1)
+					{
+						stats.NumScannedBytes += blockText.Length - offset;
+						break;
+					}
+
+					// Update the stats
+					stats.NumScannedBytes += nextOffset - offset;
+					offset = nextOffset;
+
+					// Check it's not another match within the same line
+					int lineIndexWithinBlock = blockText.GetLineIndexForOffset(offset);
+					yield return block.LineIndex + lineIndexWithinBlock;
+
+					// Move to the next line
+					offset = blockText.LineOffsets[lineIndexWithinBlock + 1];
+				}
+
+				// If the last scanned bytes is zero, we didn't have any matches from this chunk
+				if(offset == 0 && _cachedTrie != null)
+				{
+					stats.NumFalsePositiveBlocks++;
+				}
+			}
+			stats.NumSkippedBlocks += _blocks.Length - lastBlockCount;
+		}
+
+		/// <summary>
+		/// Search for the given text in the index
+		/// </summary>
+		/// <param name="text">Text to search for</param>
+		/// <param name="lineIndex">The first </param>
+		/// <returns>List of line numbers for the text</returns>
+		IEnumerable<int> EnumeratePossibleBlocks(ReadOnlyMemory<byte> text, int lineIndex)
+		{
+			// Find the starting block index
+			int blockIdx = _blocks.BinarySearch(x => x.LineIndex, lineIndex);
+			if(blockIdx < 0)
+			{
+				blockIdx = Math.Max(~blockIdx - 1, 0);
+			}
+
+			// Make sure we're not starting out of range
+			if (blockIdx < _blocks.Length)
+			{
+				// In order to be considered a possible block for a positive match, the block must contain an exact list of tokens parsed from the
+				// search text, along with a set of partial matches (for substrings that are not guaranteed to start on token boundaries). The former
+				// is relatively cheap to compute, and done by stepping through all enumerators for each token in a wave, only returning blocks
+				// which all contain the token. At the very least, this must contain blocks in the index from BlockIdx onwards.
+				List<IEnumerator<int>> enumerators = new List<IEnumerator<int>>();
+				enumerators.Add(Enumerable.Range(blockIdx, _blocks.Length - blockIdx).GetEnumerator());
+
+				// The second aspect of the search requires a more expensive search through the trie. This is done by running a set of arbitrary 
+				// delegates to filter the matches returned from the enumerators.
+				List<Predicate<int>> predicates = new List<Predicate<int>>();
+
+				// If the index has a trie, tokenize the input and generate a list of enumerators and predicates.
+				if (_cachedTrie != null)
+				{
+					// Find a list of filters for matching blocks
+					HashSet<ulong> tokens = new HashSet<ulong>();
+					for (int tokenPos = 0; tokenPos < text.Length;)
+					{
+						ReadOnlySpan<byte> token = LogToken.GetTokenText(text.Span, tokenPos);
+						if (tokenPos == 0)
+						{
+							GetUnalignedTokenPredicate(token, token.Length == text.Length, predicates);
+						}
+						else
+						{
+							GetAlignedTokenPredicate(token, tokenPos + token.Length == text.Length, tokens, predicates);
+						}
+						tokenPos += token.Length;
+					}
+
+					// Create an enumerator for each token
+					foreach (ulong token in tokens)
+					{
+						ulong minValue = (token << _numBlockBits) | (ulong)(long)blockIdx;
+						ulong maxValue = minValue + (1UL << _numBlockBits) - 1;
+
+						IEnumerator<int> enumerator = _cachedTrie.EnumerateRange(minValue, maxValue).Select(x => (int)(x - minValue)).GetEnumerator();
+						if (!enumerator.MoveNext())
+						{
+							yield break;
+						}
+
+						enumerators.Add(enumerator);
+					}
+				}
+
+				// Enumerate the matches
+				for (; ; )
+				{
+					// Advance all the enumerators that are behind the current block index. If they are all equal, we have a match.
+					bool match = true;
+					foreach (IEnumerator<int> enumerator in enumerators)
+					{
+						while (enumerator.Current < blockIdx)
+						{
+							if (!enumerator.MoveNext())
+							{
+								yield break;
+							}
+						}
+
+						if (enumerator.Current > blockIdx)
+						{
+							blockIdx = enumerator.Current;
+							match = false;
+						}
+					}
+
+					// Return the match and move to the next block
+					if (match)
+					{
+						if (predicates.All(predicate => predicate(blockIdx)))
+						{
+							yield return blockIdx;
+						}
+						blockIdx++;
+					}
+				}
+			}
+		}
+
+		/// <summary>
+		/// Gets predicates for matching a token that starts 
+		/// </summary>
+		/// <param name="text">The token text</param>
+		/// <param name="allowPartialMatch">Whether to allow a partial match of the token</param>
+		/// <param name="tokens">Set of aligned tokens that are required</param>
+		/// <param name="predicates">List of predicates for the search</param>
+		void GetAlignedTokenPredicate(ReadOnlySpan<byte> text, bool allowPartialMatch, HashSet<ulong> tokens, List<Predicate<int>> predicates)
+		{
+			for (int offset = 0; offset < text.Length; offset += LogToken.MaxTokenBytes)
+			{
+				ulong token = LogToken.GetWindowedTokenValue(text, offset);
+				if (offset + LogToken.MaxTokenBytes > text.Length && allowPartialMatch)
+				{
+					ulong tokenMask = LogToken.GetWindowedTokenMask(text, offset, true);
+					predicates.Add(blockIdx => BlockContainsToken(blockIdx, token, tokenMask));
+					break;
+				}
+				tokens.Add(token);
+			}
+		}
+
+		/// <summary>
+		/// Generates a predicate for matching a token which may or may not start on a regular token boundary
+		/// </summary>
+		/// <param name="text">The token text</param>
+		/// <param name="allowPartialMatch">Whether to allow a partial match of the token</param>
+		/// <param name="predicates">List of predicates for the search</param>
+		void GetUnalignedTokenPredicate(ReadOnlySpan<byte> text, bool allowPartialMatch, List<Predicate<int>> predicates)
+		{
+			byte[] textCopy = text.ToArray();
+
+			Lazy<HashSet<int>> blocks = new Lazy<HashSet<int>>(() =>
+			{
+				HashSet<int> union = new HashSet<int>();
+				for (int shift = 0; shift < LogToken.MaxTokenBytes; shift++)
+				{
+					HashSet<int> blocks = new HashSet<int>(BlocksContainingToken(textCopy.AsSpan(), -shift, allowPartialMatch));
+					for (int offset = -shift + LogToken.MaxTokenBytes; offset < textCopy.Length && blocks.Count > 0; offset += LogToken.MaxTokenBytes)
+					{
+						blocks.IntersectWith(BlocksContainingToken(textCopy.AsSpan(), offset, allowPartialMatch));
+					}
+					if (blocks.Count > 0)
+					{
+						union.UnionWith(blocks);
+					}
+				}
+				return union;
+			});
+
+			predicates.Add(blockIdx => blocks.Value.Contains(blockIdx));
+		}
+
+		/// <summary>
+		/// Tests whether a block contains a particular token
+		/// </summary>
+		/// <param name="blockIdx">Index of the block to search</param>
+		/// <param name="token">The token to test</param>
+		/// <param name="tokenMask">Mask of which bits in the token are valid</param>
+		/// <returns>True if the given block contains a token</returns>
+		bool BlockContainsToken(int blockIdx, ulong token, ulong tokenMask)
+		{
+			token = (token << _numBlockBits) | (uint)blockIdx;
+			tokenMask = (tokenMask << _numBlockBits) | ((1UL << _numBlockBits) - 1);
+			return _cachedTrie!.EnumerateValues((value, valueMask) => (value & tokenMask) == (token & valueMask)).Any();
+		}
+
+		/// <summary>
+		/// Tests whether a block contains a particular token
+		/// </summary>
+		/// <param name="text">The token to test</param>
+		/// <param name="offset">Offset of the window into the token to test</param>
+		/// <param name="allowPartialMatch">Whether to allow a partial match of the token</param>
+		/// <returns>True if the given block contains a token</returns>
+		IEnumerable<int> BlocksContainingToken(ReadOnlySpan<byte> text, int offset, bool allowPartialMatch)
+		{
+			ulong token = LogToken.GetWindowedTokenValue(text, offset) << _numBlockBits;
+			ulong tokenMask = LogToken.GetWindowedTokenMask(text, offset, allowPartialMatch) << _numBlockBits;
+			ulong blockMask = (1UL << _numBlockBits) - 1;
+			return _cachedTrie!.EnumerateValues((value, valueMask) => (value & tokenMask) == (token & valueMask)).Select(x => (int)(x & blockMask)).Distinct();
+		}
+
+		/// <summary>
+		/// Deserialize the index from memory
+		/// </summary>
+		/// <param name="reader">Reader to deserialize from</param>
+		/// <returns>Index data</returns>
+		public static LogIndexData Read(MemoryReader reader)
+		{
+			int version = reader.ReadInt32();
+			if (version == 0)
+			{
+				return LogIndexData.Empty;
+			}
+			else if (version == 1)
+			{
+				ReadOnlyTrie index = reader.ReadTrie();
+				int numBlockBits = reader.ReadInt32();
+
+				LogIndexBlock[] blocks = new LogIndexBlock[reader.ReadInt32()];
+				for (int idx = 0; idx < blocks.Length; idx++)
+				{
+					int lineIndex = (idx > 0) ? (blocks[idx - 1].LineIndex + blocks[idx - 1].LineCount) : 0;
+					int lineCount = reader.ReadInt32();
+					ReadOnlyMemory<byte> compressedPlainText = reader.ReadVariableLengthBytesWithInt32Length();
+					reader.ReadTrie();
+					blocks[idx] = new LogIndexBlock(lineIndex, lineCount, null, compressedPlainText);
+				}
+
+				return new LogIndexData(index, numBlockBits, blocks);
+			}
+			else if (version == 2)
+			{
+				ReadOnlyTrie index = reader.ReadTrie();
+				int numBlockBits = reader.ReadInt32();
+
+				LogIndexBlock[] blocks = new LogIndexBlock[reader.ReadInt32()];
+				for (int idx = 0; idx < blocks.Length; idx++)
+				{
+					int lineIndex = reader.ReadInt32();
+					int lineCount = reader.ReadInt32();
+					ReadOnlyMemory<byte> compressedPlainText = reader.ReadVariableLengthBytesWithInt32Length();
+					reader.ReadTrie();
+					blocks[idx] = new LogIndexBlock(lineIndex, lineCount, null, compressedPlainText);
+				}
+
+				return new LogIndexData(index, numBlockBits, blocks);
+			}
+			else if (version == 3)
+			{
+				ReadOnlyTrie index = reader.ReadTrie();
+				int numBlockBits = reader.ReadInt32();
+				LogIndexBlock[] blocks = reader.ReadVariableLengthArrayWithInt32Length(() => reader.ReadLogIndexBlock());
+				for(int idx = 1; idx < blocks.Length; idx++)
+				{
+					if(blocks[idx].LineIndex == 0)
+					{
+						int lineIndex = blocks[idx - 1].LineIndex + blocks[idx - 1].LineCount;
+						blocks[idx] = new LogIndexBlock(lineIndex, blocks[idx].LineCount, blocks[idx].CachedPlainText, blocks[idx].CompressedPlainText);
+					}
+				}
+				return new LogIndexData(index, numBlockBits, blocks);
+			}
+			else
+			{
+				throw new InvalidDataException($"Invalid index version number {version}");
+			}
+		}
+
+		/// <summary>
+		/// Serialize an index into memory
+		/// </summary>
+		/// <param name="writer">Writer to serialize to</param>
+		public void Write(MemoryWriter writer)
+		{
+			writer.WriteInt32(3);
+			writer.WriteTrie(BuildTrie());
+			writer.WriteInt32(_numBlockBits);
+			writer.WriteVariableLengthArrayWithInt32Length(_blocks, x => writer.WriteLogIndexBlock(x));
+		}
+
+		/// <summary>
+		/// Deserialize the index from memory
+		/// </summary>
+		/// <param name="memory">Memory to deserialize from</param>
+		/// <returns>Index data</returns>
+		public static LogIndexData FromMemory(ReadOnlyMemory<byte> memory)
+		{
+			MemoryReader reader = new MemoryReader(memory);
+			return Read(reader);
+		}
+
+		/// <summary>
+		/// Serializes the index
+		/// </summary>
+		/// <returns>Index data</returns>
+		public byte[] ToByteArray()
+		{
+			byte[] buffer = new byte[GetSerializedSize()];
+
+			MemoryWriter writer = new MemoryWriter(buffer);
+			Write(writer);
+			writer.CheckEmpty();
+
+			return buffer;
+		}
+
+		/// <summary>
+		/// Gets the serialized size of this index
+		/// </summary>
+		/// <returns>The serialized size</returns>
+		public int GetSerializedSize()
+		{
+			return sizeof(int) + BuildTrie().GetSerializedSize() + sizeof(int) + (sizeof(int) + _blocks.Sum(x => x.GetSerializedSize()));
+		}
+	}
+
+	static class LogIndexDataExtensions
+	{
+		/// <summary>
+		/// Deserialize the index from memory
+		/// </summary>
+		/// <param name="reader">Reader to deserialize from</param>
+		/// <returns>Index data</returns>
+		public static LogIndexData ReadLogIndexData(this MemoryReader reader)
+		{
+			return LogIndexData.Read(reader);
+		}
+
+		/// <summary>
+		/// Serialize an index into memory
+		/// </summary>
+		/// <param name="writer">Writer to serialize to</param>
+		/// <param name="index">The index to write</param>
+		public static void WriteLogIndexData(this MemoryWriter writer, LogIndexData index)
+		{
+			index.Write(writer);
+		}
+	}
+}
diff --git a/Engine/Source/Programs/Horde/Horde.Server/Logs/Data/LogSubChunkData.cs b/Engine/Source/Programs/Horde/Horde.Server/Logs/Data/LogSubChunkData.cs
new file mode 100644
--- /dev/null
+++ b/Engine/Source/Programs/Horde/Horde.Server/Logs/Data/LogSubChunkData.cs
@@ -0,0 +1,246 @@
+// Copyright Epic Games, Inc. All Rights Reserved.
+
+using System;
+using EpicGames.Core;
+using Horde.Server.Utilities;
+using Microsoft.Extensions.Logging;
+
+namespace Horde.Server.Logs.Data
+{
+	/// <summary>
+	/// Pending data for a sub-chunk
+	/// </summary>
+	public class LogSubChunkData
+	{
+		/// <summary>
+		/// Type of data stored in this subchunk
+		/// </summary>
+		public LogType Type { get; }
+
+		/// <summary>
+		/// Offset within the file of this sub-chunk
+		/// </summary>
+		public long Offset { get; }
+
+		/// <summary>
+		/// Length of this sub-chunk
+		/// </summary>
+		public int Length { get; }
+
+		/// <summary>
+		/// Index of the first line in this sub-chunk
+		/// </summary>
+		public int LineIndex { get; }
+
+		/// <summary>
+		/// Gets the number of lines in this sub-chunk
+		/// </summary>
+		public int LineCount { get; }
+
+		/// <summary>
+		/// Text data
+		/// </summary>
+		ILogText? _textInternal;
+
+		/// <summary>
+		/// Compressed text data
+		/// </summary>
+		ReadOnlyMemory<byte> _compressedTextInternal;
+
+		/// <summary>
+		/// The index for this sub-chunk
+		/// </summary>
+		LogIndexData? _indexInternal;
+
+		/// <summary>
+		/// The log text
+		/// </summary>
+		public ILogText InflateText()
+		{
+			_textInternal ??= new ReadOnlyLogText(_compressedTextInternal.DecompressBzip2());
+			return _textInternal;
+		}
+
+		/// <summary>
+		/// The compressed log text
+		/// </summary>
+		public ReadOnlyMemory<byte> DeflateText()
+		{
+			if(_compressedTextInternal.IsEmpty)
+			{
+				_compressedTextInternal = _textInternal!.Data.CompressBzip2();
+			}
+			return _compressedTextInternal;
+		}
+
+		/// <summary>
+		/// Index for tokens in this chunk
+		/// </summary>
+		public LogIndexData BuildIndex(ILogger logger)
+		{
+			if(_indexInternal == null)
+			{
+				ILogText plainText = InflateText();
+				if (Type != LogType.Text)
+				{
+					plainText = plainText.ToPlainText(logger);
+				}
+
+				LogIndexBlock[] blocks = new LogIndexBlock[1];
+				blocks[0] = new LogIndexBlock(LineIndex, LineCount, plainText, plainText.Data.CompressBzip2());
+				_indexInternal = new LogIndexData(null, 0, blocks);
+			}
+			return _indexInternal;
+		}
+
+		/// <summary>
+		/// Constructor
+		/// </summary>
+		/// <param name="type">Type of data stored in this subchunk</param>
+		/// <param name="offset">Offset within the file of this sub-chunk</param>
+		/// <param name="lineIndex">Index of the first line in this sub-chunk</param>
+		/// <param name="text">Text to add</param>
+		public LogSubChunkData(LogType type, long offset, int lineIndex, ILogText text)
+		{
+			Type = type;
+			Offset = offset;
+			Length = text.Data.Length;
+			LineIndex = lineIndex;
+			LineCount = text.LineCount;
+			_textInternal = text;
+		}
+
+		/// <summary>
+		/// Constructor for raw data
+		/// </summary>
+		/// <param name="type">Type of data stored in this subchunk</param>
+		/// <param name="offset">Offset within the file of this sub-chunk</param>
+		/// <param name="length">Length of the uncompressed data</param>
+		/// <param name="lineIndex">Index of the first line</param>
+		/// <param name="lineCount">Number of lines in the uncompressed text</param>
+		/// <param name="compressedText">Compressed text data</param>
+		/// <param name="index">Index data</param>
+		public LogSubChunkData(LogType type, long offset, int length, int lineIndex, int lineCount, ReadOnlyMemory<byte> compressedText, LogIndexData? index)
+		{
+			Type = type;
+			Offset = offset;
+			Length = length;
+			LineIndex = lineIndex;
+			LineCount = lineCount;
+			_compressedTextInternal = compressedText;
+			_indexInternal = index;
+		}
+
+		/// <summary>
+		/// Constructs a sub-chunk from a block of memory. Uses slices of the given memory buffer rather than copying the data.
+		/// </summary>
+		/// <param name="reader">The reader to read from</param>
+		/// <param name="offset">Offset of the sub-chunk within this file</param>
+		/// <param name="lineIndex">Index of the first line in this subchunk</param>
+		/// <returns>New subchunk data</returns>
+		public static LogSubChunkData Read(MemoryReader reader, long offset, int lineIndex)
+		{
+			int version = reader.ReadInt32(); // Version placeholder
+			if (version == 0)
+			{
+				LogType type = (LogType)reader.ReadInt32();
+				int length = reader.ReadInt32();
+				int lineCount = reader.ReadInt32();
+				ReadOnlyMemory<byte> compressedText = reader.ReadVariableLengthBytesWithInt32Length();
+				reader.ReadVariableLengthBytesWithInt32Length();
+				return new LogSubChunkData(type, offset, length, lineIndex, lineCount, compressedText, null);
+			}
+			else if (version == 1)
+			{
+				LogType type = (LogType)reader.ReadInt32();
+				int length = reader.ReadInt32();
+				int lineCount = reader.ReadInt32();
+				ReadOnlyMemory<byte> compressedText = reader.ReadVariableLengthBytesWithInt32Length();
+				ReadOnlyMemory<byte> compressedPlainText = reader.ReadVariableLengthBytesWithInt32Length();
+				ReadOnlyTrie trie = reader.ReadTrie();
+				LogIndexBlock indexBlock = new LogIndexBlock(lineIndex, lineCount, null, compressedPlainText);
+				LogIndexData index = new LogIndexData(trie, 0, new[] { indexBlock });
+				return new LogSubChunkData(type, offset, length, lineIndex, lineCount, compressedText, index);
+			}
+			else
+			{
+				LogType type = (LogType)reader.ReadInt32();
+				int length = reader.ReadInt32();
+				int lineCount = reader.ReadInt32();
+				ReadOnlyMemory<byte> compressedText = reader.ReadVariableLengthBytesWithInt32Length();
+				LogIndexData index = reader.ReadLogIndexData();
+				index.SetBaseLineIndex(lineIndex); // Fix for incorrectly saved data
+				return new LogSubChunkData(type, offset, length, lineIndex, lineCount, compressedText, index);
+			}
+		}
+
+		/// <summary>
+		/// Serializes the sub-chunk to a stream
+		/// </summary>
+		/// <param name="writer">Writer to output to</param>
+		/// <param name="logger">Logger for output</param>
+		public void Write(MemoryWriter writer, ILogger logger)
+		{
+			writer.WriteInt32(2); // Version placeholder
+
+			writer.WriteInt32((int)Type);
+			writer.WriteInt32(Length);
+			writer.WriteInt32(LineCount);
+
+			writer.WriteVariableLengthBytesWithInt32Length(DeflateText().Span);
+			writer.WriteLogIndexData(BuildIndex(logger));
+		}
+
+		/// <summary>
+		/// Serializes this object to a byte array
+		/// </summary>
+		/// <returns>Byte array</returns>
+		public byte[] ToByteArray(ILogger logger)
+		{
+			byte[] data = new byte[GetSerializedSize(logger)];
+
+			MemoryWriter writer = new MemoryWriter(data);
+			Write(writer, logger);
+			writer.CheckEmpty();
+
+			return data;
+		}
+
+		/// <summary>
+		/// Determines the size of serialized data
+		/// </summary>
+		public int GetSerializedSize(ILogger logger)
+		{
+			return sizeof(int) + sizeof(int) + sizeof(int) + sizeof(int) + (sizeof(int) + DeflateText().Length) + BuildIndex(logger).GetSerializedSize();
+		}
+	}
+
+	/// <summary>
+	/// Pending data for a sub-chunk
+	/// </summary>
+	public static class LogSubChunkDataExtensions
+	{
+		/// <summary>
+		/// Constructs a sub-chunk from a block of memory. Uses slices of the given memory buffer rather than copying the data.
+		/// </summary>
+		/// <param name="reader">The reader to read from</param>
+		/// <param name="offset">Offset of this sub-chunk within the file</param>
+		/// <param name="lineIndex">Index of the first line within this chunk</param>
+		/// <returns>New subchunk data</returns>
+		public static LogSubChunkData ReadLogSubChunkData(this MemoryReader reader, long offset, int lineIndex)
+		{
+			return LogSubChunkData.Read(reader, offset, lineIndex);
+		}
+
+		/// <summary>
+		/// Serializes the sub-chunk to a stream
+		/// </summary>
+		/// <param name="writer">Writer to output to</param>
+		/// <param name="subChunkData">The sub-chunk data to write</param>
+		/// <param name="logger">Logger for output</param>
+		public static void WriteLogSubChunkData(this MemoryWriter writer, LogSubChunkData subChunkData, ILogger logger)
+		{
+			subChunkData.Write(writer, logger);
+		}
+	}
+}
diff --git a/Engine/Source/Programs/Horde/Horde.Server/Logs/Data/LogText.cs b/Engine/Source/Programs/Horde/Horde.Server/Logs/Data/LogText.cs
new file mode 100644
--- /dev/null
+++ b/Engine/Source/Programs/Horde/Horde.Server/Logs/Data/LogText.cs
@@ -0,0 +1,502 @@
+// Copyright Epic Games, Inc. All Rights Reserved.
+
+using System;
+using System.Collections.Generic;
+using System.Diagnostics;
+using System.IO;
+using System.Text;
+using System.Text.Json;
+using EpicGames.Core;
+using Microsoft.Extensions.Logging;
+
+namespace Horde.Server.Logs.Data
+{
+	/// <summary>
+	/// Exception thrown trying to parse log output
+	/// </summary>
+	public class LogTextParseException : Exception
+	{
+		/// <summary>
+		/// Constructor
+		/// </summary>
+		public LogTextParseException(string message, Exception? ex)
+			: base(message, ex)
+		{
+		}
+	}
+
+	/// <summary>
+	/// Interface for log data
+	/// </summary>
+	public interface ILogText
+	{
+		/// <summary>
+		/// The raw text data. Contains a complete set of lines followed by newline characters.
+		/// </summary>
+		public ReadOnlyMemory<byte> Data { get; }
+
+		/// <summary>
+		/// Length of the data in this block
+		/// </summary>
+		public int Length { get; }
+
+		/// <summary>
+		/// Offsets of lines within the data object, including a sentinel for the end of the data (LineCount + 1 entries).
+		/// </summary>
+		public IReadOnlyList<int> LineOffsets { get; }
+
+		/// <summary>
+		/// Number of lines in the 
+		/// </summary>
+		public int LineCount { get; }
+
+		/// <summary>
+		/// Computes the number of bytes used by this object
+		/// </summary>
+		public int AllocatedSize { get; }
+	}
+
+	/// <summary>
+	/// Helper methods for parsing log text
+	/// </summary>
+	public static class LogTextHelpers
+	{
+		/// <summary>
+		/// Finds the number of newline characters in the given span
+		/// </summary>
+		/// <param name="span">Span of utf-8 characters to scan</param>
+		/// <param name="offset">Starting offset within the span</param>
+		/// <returns>Number of lines in the given span</returns>
+		public static int GetLineCount(ReadOnlySpan<byte> span, int offset)
+		{
+			int lineCount = 0;
+			for (; offset < span.Length; offset++)
+			{
+				if (span[offset] == (byte)'\n')
+				{
+					lineCount++;
+				}
+			}
+			return lineCount;
+		}
+
+		/// <summary>
+		/// Find the offsets of the start of each line within the given span
+		/// </summary>
+		/// <param name="span">Span to search</param>
+		/// <param name="offset">Offset within the span to start searching</param>
+		/// <param name="lineOffsets">Receives the list of line offsets</param>
+		/// <param name="lineCount">The current line count</param>
+		public static void FindLineOffsets(ReadOnlySpan<byte> span, int offset, Span<int> lineOffsets, int lineCount)
+		{
+			Debug.Assert(span.Length == 0 || span[^1] == (byte)'\n');
+			for(; ; )
+			{
+				// Store the start of this line
+				lineOffsets[lineCount++] = offset;
+
+				// Check if we're at the end
+				if(offset >= span.Length)
+				{
+					break;
+				}
+
+				// Move to the end of this line
+				while (span[offset] != '\n')
+				{
+					offset++;
+				}
+
+				// Move to the start of the next line
+				offset++;
+			}
+		}
+	}
+
+	/// <summary>
+	/// Immutable block of log text
+	/// </summary>
+	public class ReadOnlyLogText : ILogText
+	{
+		/// <inheritdoc/>
+		public ReadOnlyMemory<byte> Data { get; private set; }
+
+		/// <inheritdoc/>
+		public int Length => Data.Length;
+
+		/// <inheritdoc/>
+		public IReadOnlyList<int> LineOffsets { get; private set; }
+
+		/// <inheritdoc/>
+		public int LineCount => LineOffsets.Count - 1;
+
+		/// <inheritdoc/>
+		public int AllocatedSize => Data.Length + (LineOffsets.Count * sizeof(int));
+
+		/// <summary>
+		/// Constructor
+		/// </summary>
+		/// <param name="data">Data to construct from</param>
+		public ReadOnlyLogText(ReadOnlyMemory<byte> data)
+		{
+			ReadOnlySpan<byte> span = data.Span;
+
+			int lineCount = LogTextHelpers.GetLineCount(span, 0);
+
+			int[] lineOffsets = new int[lineCount + 1];
+			LogTextHelpers.FindLineOffsets(span, 0, lineOffsets, 0);
+			Debug.Assert(lineOffsets[lineCount] == data.Length);
+
+			Data = data;
+			LineOffsets = lineOffsets;
+		}
+	}
+
+	/// <summary>
+	/// Mutable log text object
+	/// </summary>
+	public class LogText : ILogText
+	{
+		/// <summary>
+		/// Accessor for Data
+		/// </summary>
+		byte[] _internalData;
+
+		/// <inheritdoc/>
+		public ReadOnlyMemory<byte> Data => _internalData.AsMemory(0, Length);
+
+		/// <summary>
+		/// Hacky accessor for InternalData, for serializing to Redis
+		/// </summary>
+		public byte[] InternalDataHack => _internalData;
+
+		/// <inheritdoc/>
+		public int Length
+		{
+			get;
+			private set;
+		}
+
+		/// <summary>
+		/// Offsets of the start of each line within the data
+		/// </summary>
+		readonly List<int> _internalLineOffsets = new List<int> { 0 };
+
+		/// <inheritdoc/>
+		public IReadOnlyList<int> LineOffsets => _internalLineOffsets;
+
+		/// <inheritdoc/>
+		public int LineCount => _internalLineOffsets.Count - 1;
+
+		/// <summary>
+		/// Chunk Size
+		/// </summary>
+		public int MaxLength => _internalData.Length;
+
+		/// <summary>
+		/// Computes the size of data used by this object
+		/// </summary>
+		public int AllocatedSize => _internalData.Length + (_internalLineOffsets.Capacity * sizeof(int));
+
+		/// <summary>
+		/// Constructor
+		/// </summary>
+		public LogText()
+		{
+			_internalData = Array.Empty<byte>();
+		}
+
+		/// <summary>
+		/// Constructor
+		/// </summary>
+		/// <param name="data">Data to initialize this chunk with. Ownership of this array is transferred to the chunk, and its length determines the chunk size.</param>
+		/// <param name="length">Number of valid bytes within the initial data array</param>
+		public LogText(byte[] data, int length)
+		{
+			_internalData = data;
+			Length = 0; // Updated below
+
+			UpdateLength(length);
+		}
+
+		/// <summary>
+		/// Create a new chunk data object with the given data appended. The internal buffers are reused, with the assumption that
+		/// there is no contention over writing to the same location in the chunk.
+		/// </summary>
+		/// <param name="textData">The data to append</param>
+		/// <returns>New chunk data object</returns>
+		public void Append(ReadOnlySpan<byte> textData)
+		{
+			textData.CopyTo(_internalData.AsSpan(Length, textData.Length));
+			UpdateLength(Length + textData.Length);
+		}
+
+		/// <summary>
+		/// Updates the plain text representation of this chunk
+		/// </summary>
+		public void AppendPlainText(ILogText srcText, int srcLineIndex, int srcLineCount, ILogger logger)
+		{
+			// Convert all the lines to plain text
+			for (int idx = 0; idx < srcLineCount; idx++)
+			{
+				int lineOffset = srcText.LineOffsets[srcLineIndex + idx];
+				int nextLineOffset = srcText.LineOffsets[srcLineIndex + idx + 1];
+				ReadOnlySpan<byte> inputLine = srcText.Data.Slice(lineOffset, nextLineOffset - lineOffset).Span;
+
+				// Make sure the output buffer is large enough
+				int requiredSpace = nextLineOffset - lineOffset;
+				if (requiredSpace > _internalData.Length - Length)
+				{
+					int maxLineOffset = srcText.LineOffsets[srcLineIndex + srcLineCount];
+					Array.Resize(ref _internalData, _internalData.Length + requiredSpace + (maxLineOffset - nextLineOffset) / 2);
+				}
+
+				// Convert the line to plain text
+				try
+				{
+					Length = ConvertToPlainText(inputLine, _internalData, Length);
+				}
+				catch (Exception ex)
+				{
+					inputLine.CopyTo(_internalData.AsSpan(Length));
+					logger.LogWarning(ex, "Exception while attempting to parse log text as JSON. Line: \"{Line}\"", Encoding.UTF8.GetString(inputLine).Trim());
+				}
+				_internalLineOffsets.Add(Length);
+			}
+		}
+
+		/// <summary>
+		/// Determines if the given line is empty
+		/// </summary>
+		/// <param name="input">The input data</param>
+		/// <returns>True if the given text is empty</returns>
+		public static bool IsEmptyOrWhitespace(ReadOnlySpan<byte> input)
+		{
+			for(int idx = 0; idx < input.Length; idx++)
+			{
+				byte v = input[idx];
+				if(v != (byte)'\n' && v != '\r' && v != ' ')
+				{
+					return false;
+				}
+			}
+			return true;
+		}
+
+		/// <summary>
+		/// Converts a JSON log line to plain text
+		/// </summary>
+		/// <param name="input">The JSON data</param>
+		/// <param name="output">Output buffer for the converted line</param>
+		/// <param name="outputOffset">Offset within the buffer to write the converted data</param>
+		/// <returns></returns>
+		public static int ConvertToPlainText(ReadOnlySpan<byte> input, byte[] output, int outputOffset)
+		{
+			if(IsEmptyOrWhitespace(input))
+			{
+				output[outputOffset] = (byte)'\n';
+				return outputOffset + 1;
+			}
+
+			Utf8JsonReader reader = new Utf8JsonReader(input);
+			if (reader.Read() && reader.TokenType == JsonTokenType.StartObject)
+			{
+				while (reader.Read() && reader.TokenType == JsonTokenType.PropertyName)
+				{
+					if (!reader.ValueTextEquals("message"))
+					{
+						reader.Skip();
+						continue;
+					}
+					if (!reader.Read() || reader.TokenType != JsonTokenType.String)
+					{
+						reader.Skip();
+						continue;
+					}
+
+					int unescapedLength = UnescapeUtf8(reader.ValueSpan, output.AsSpan(outputOffset));
+					outputOffset += unescapedLength;
+
+					output[outputOffset] = (byte)'\n';
+					outputOffset++;
+
+					break;
+				}
+			}
+			return outputOffset;
+		}
+
+		/// <summary>
+		/// Unescape a json utf8 string
+		/// </summary>
+		/// <param name="source">Source span of bytes</param>
+		/// <param name="target">Target span of bytes</param>
+		/// <returns>Length of the converted data</returns>
+		static int UnescapeUtf8(ReadOnlySpan<byte> source, Span<byte> target)
+		{
+			int length = 0;
+			for (; ; )
+			{
+				// Copy up to the next backslash
+				int backslash = source.IndexOf((byte)'\\');
+				if (backslash == -1)
+				{
+					source.CopyTo(target);
+					length += source.Length;
+					break;
+				}
+				else if (backslash > 0)
+				{
+					source.Slice(0, backslash).CopyTo(target);
+					source = source.Slice(backslash);
+					target = target.Slice(backslash);
+					length += backslash;
+				}
+
+				// Check what the escape code is
+				if (source[1] == 'u')
+				{
+					char[] chars = { (char)((StringUtils.ParseHexByte(source, 2) << 8) | StringUtils.ParseHexByte(source, 4)) };
+					int encodedLength = Encoding.UTF8.GetBytes(chars.AsSpan(), target);
+					source = source.Slice(6);
+					target = target.Slice(encodedLength);
+					length += encodedLength;
+				}
+				else
+				{
+					target[0] = source[1] switch
+					{
+						(byte)'\"' => (byte)'\"',
+						(byte)'\\' => (byte)'\\',
+						(byte)'b' => (byte)'\b',
+						(byte)'f' => (byte)'\f',
+						(byte)'n' => (byte)'\n',
+						(byte)'r' => (byte)'\r',
+						(byte)'t' => (byte)'\t',
+						_ => source[1]
+					};
+					source = source.Slice(2);
+					target = target.Slice(1);
+					length++;
+				}
+			}
+			return length;
+		}
+
+		/// <summary>
+		/// Generates placeholder data for a missing span
+		/// </summary>
+		/// <param name="chunkIdx">Index of the chunk</param>
+		/// <param name="hostName">Host name of the machine that was holding the data</param>
+		/// <param name="targetLength">Desired length of the buffer</param>
+		/// <param name="targetLineCount">Desired line count for the buffer</param>
+		public void AppendMissingDataInfo(int chunkIdx, string? hostName, int targetLength, int targetLineCount)
+		{
+			int newLength = Length;
+
+			if(_internalData.Length < targetLength)
+			{
+				Array.Resize(ref _internalData, targetLength);
+			}
+
+			string suffix = $" (server: {hostName})";
+			for(int newLineCount = LineCount; newLineCount < targetLineCount; newLineCount++)
+			{
+				byte[] errorBytes = Encoding.ASCII.GetBytes($"{{ \"level\":\"Error\",\"message\":\"[Missing data at chunk {chunkIdx}, line {newLineCount}{suffix}]\" }}\n");
+				if (newLength + errorBytes.Length > targetLength)
+				{
+					break;
+				}
+
+				errorBytes.AsSpan().CopyTo(_internalData.AsSpan(newLength));
+				newLength += errorBytes.Length;
+
+				suffix = String.Empty;
+			}
+
+			if (newLength < targetLength)
+			{
+				newLength = Math.Max(newLength - 1, 0);
+				_internalData.AsSpan(newLength, (int)((targetLength - 1) - newLength)).Fill((byte)' ');
+				_internalData[targetLength - 1] = (byte)'\n';
+			}
+
+			UpdateLength((int)targetLength);
+		}
+
+		/// <summary>
+		/// Shrinks the data allocated by this chunk to the minimum required
+		/// </summary>
+		public void Shrink()
+		{
+			if (Length < _internalData.Length)
+			{
+				Array.Resize(ref _internalData, Length);
+			}
+		}
+
+		/// <summary>
+		/// Updates the length of this chunk, computing all the newline offsets
+		/// </summary>
+		/// <param name="newLength">New length of the chunk</param>
+		private void UpdateLength(int newLength)
+		{
+			if (newLength > Length)
+			{
+				// Make sure the data ends with a newline
+				if (_internalData[newLength - 1] != '\n')
+				{
+					throw new InvalidDataException("Chunk data must end with a newline");
+				}
+
+				// Calculate the new number of newlines
+				ReadOnlySpan<byte> newData = _internalData.AsSpan(0, newLength);
+				for(int idx = Length; idx < newLength; idx++)
+				{
+					if(newData[idx] == '\n')
+					{
+						_internalLineOffsets.Add(idx + 1);
+					}
+				}
+
+				// Update the size of the buffer
+				Length = newLength;
+			}
+		}
+	}
+
+	/// <summary>
+	/// Extension methods for ILogText
+	/// </summary>
+	public static class LogTextExtensions
+	{
+		/// <summary>
+		/// Find the line index for a particular offset
+		/// </summary>
+		/// <param name="logText">The text to search</param>
+		/// <param name="offset">Offset within the text</param>
+		/// <returns>The line index</returns>
+		public static int GetLineIndexForOffset(this ILogText logText, int offset)
+		{
+			int lineIdx = logText.LineOffsets.BinarySearch(offset);
+			if(lineIdx < 0)
+			{
+				lineIdx = ~lineIdx - 1;
+			}
+			return lineIdx;
+		}
+
+		/// <summary>
+		/// Converts a log text instance to plain text
+		/// </summary>
+		/// <param name="logText">The text to convert</param>
+		/// <param name="logger">Logger for conversion warnings</param>
+		/// <returns>The plain text instance</returns>
+		public static ILogText ToPlainText(this ILogText logText, ILogger logger)
+		{
+			LogText other = new LogText();
+			other.AppendPlainText(logText, 0, logText.LineCount, logger);
+			return other;
+		}
+	}
+}
diff --git a/Engine/Source/Programs/Horde/Horde.Server/Logs/Data/LogToken.cs b/Engine/Source/Programs/Horde/Horde.Server/Logs/Data/LogToken.cs
new file mode 100644
--- /dev/null
+++ b/Engine/Source/Programs/Horde/Horde.Server/Logs/Data/LogToken.cs
@@ -0,0 +1,205 @@
+// Copyright Epic Games, Inc. All Rights Reserved.
+
+using System;
+using System.Collections.Generic;
+
+namespace Horde.Server.Logs.Data
+{
+	/// <summary>
+	/// Functionality for decomposing log text into tokens
+	/// </summary>
+	public static class LogToken
+	{
+		/// <summary>
+		/// Maximum number of bytes in each token
+		/// </summary>
+		public const int MaxTokenBytes = 4;
+
+		/// <summary>
+		/// Number of bits in each token
+		/// </summary>
+		public const int MaxTokenBits = MaxTokenBytes * 8;
+
+		/// <summary>
+		/// Lookup from input byte to token type
+		/// </summary>
+		static readonly byte[] s_tokenTypes = GetTokenTypes();
+
+		/// <summary>
+		/// Lookup from input byte to token char
+		/// </summary>
+		static readonly byte[] s_tokenChars = GetTokenChars();
+
+		/// <summary>
+		/// Get a set of tokens in a trie
+		/// </summary>
+		/// <param name="text">Text to scan</param>
+		/// <returns></returns>
+		public static ReadOnlyTrie GetTrie(ReadOnlySpan<byte> text)
+		{
+			ReadOnlyTrieBuilder builder = new ReadOnlyTrieBuilder();
+			GetTokens(text, builder.Add);
+			return builder.Build();
+		}
+
+		/// <summary>
+		/// Gets a single token
+		/// </summary>
+		/// <param name="text">The text to parse</param>
+		/// <returns>The token value</returns>
+		public static ulong GetToken(ReadOnlySpan<byte> text)
+		{
+			ulong token = 0;
+			GetTokens(text, x => token = x);
+			return token;
+		}
+
+		/// <summary>
+		/// Get a set of tokens in a trie
+		/// </summary>
+		/// <param name="text">Text to scan</param>
+		/// <returns></returns>
+		public static HashSet<ulong> GetTokenSet(ReadOnlySpan<byte> text)
+		{
+			HashSet<ulong> tokens = new HashSet<ulong>();
+			GetTokens(text, x => tokens.Add(x));
+			return tokens;
+		}
+
+		/// <summary>
+		/// Decompose a span of text into tokens
+		/// </summary>
+		/// <param name="text">Text to scan</param>
+		/// <param name="addToken">Receives a set of tokens</param>
+		public static void GetTokens(ReadOnlySpan<byte> text, Action<ulong> addToken)
+		{
+			if (text.Length > 0)
+			{
+				int type = s_tokenTypes[text[0]];
+				int numTokenBits = 8;
+				ulong token = s_tokenChars[text[0]];
+
+				for (int textIdx = 1; textIdx < text.Length; textIdx++)
+				{
+					byte nextChar = s_tokenChars[text[textIdx]];
+					int nextType = s_tokenTypes[nextChar];
+					if (type != nextType || numTokenBits + 8 > MaxTokenBits)
+					{
+						addToken(token << (MaxTokenBits - numTokenBits));
+						token = 0;
+						numTokenBits = 0;
+						type = nextType;
+					}
+					token = (token << 8) | nextChar;
+					numTokenBits += 8;
+				}
+
+				addToken(token << (MaxTokenBits - numTokenBits));
+			}
+		}
+
+		/// <summary>
+		/// Gets the length of the first token in the given span
+		/// </summary>
+		/// <param name="text">The text to search</param>
+		/// <param name="pos">Start position for the search</param>
+		/// <returns>Length of the first token</returns>
+		public static ReadOnlySpan<byte> GetTokenText(ReadOnlySpan<byte> text, int pos)
+		{
+			int type = s_tokenTypes[text[pos]];
+			for (int end = pos + 1; ; end++)
+			{
+				if (end == text.Length || s_tokenTypes[text[end]] != type)
+				{
+					return text.Slice(pos, end - pos);
+				}
+			}
+		}
+
+		/// <summary>
+		/// Gets the length of the first token in the given span
+		/// </summary>
+		/// <param name="text">The text to search</param>
+		/// <param name="offset">Offset of the window to read from the token</param>
+		/// <returns>Length of the first token</returns>
+		public static ulong GetWindowedTokenValue(ReadOnlySpan<byte> text, int offset)
+		{
+			ulong token = 0;
+			for (int idx = 0; idx < MaxTokenBytes; idx++)
+			{
+				token <<= 8;
+				if (offset >= 0 && offset < text.Length)
+				{
+					token |= s_tokenChars[text[offset]];
+				}
+				offset++;
+			}
+			return token;
+		}
+
+		/// <summary>
+		/// Gets the length of the first token in the given span
+		/// </summary>
+		/// <param name="text">The text to search</param>
+		/// <param name="offset">Offset of the window to read from the token</param>
+		/// <param name="allowPartialMatch">Whether to allow only matching the start of the string</param>
+		/// <returns>Length of the first token</returns>
+		public static ulong GetWindowedTokenMask(ReadOnlySpan<byte> text, int offset, bool allowPartialMatch)
+		{
+			ulong token = 0;
+			for (int idx = 0; idx < MaxTokenBytes; idx++)
+			{
+				token <<= 8;
+				if (offset >= 0 && (offset < text.Length || !allowPartialMatch))
+				{
+					token |= 0xff;
+				}
+				offset++;
+			}
+			return token;
+		}
+
+		/// <summary>
+		/// Build the lookup table for token types
+		/// </summary>
+		/// <returns>Array whose elements map from an input byte to token type</returns>
+		static byte[] GetTokenTypes()
+		{
+			byte[] charTypes = new byte[256];
+			for (int idx = 'a'; idx <= 'z'; idx++)
+			{
+				charTypes[idx] = 1;
+			}
+			for (int idx = 'A'; idx <= 'Z'; idx++)
+			{
+				charTypes[idx] = 1;
+			}
+			for (int idx = '0'; idx <= '9'; idx++)
+			{
+				charTypes[idx] = 2;
+			}
+			charTypes[' '] = 3;
+			charTypes['\t'] = 3;
+			charTypes['\n'] = 4;
+			return charTypes;
+		}
+
+		/// <summary>
+		/// Build the lookup table for token types
+		/// </summary>
+		/// <returns>Array whose elements map from an input byte to token type</returns>
+		static byte[] GetTokenChars()
+		{
+			byte[] chars = new byte[256];
+			for(int idx = 0; idx < 256; idx++)
+			{
+				chars[idx] = (byte)idx;
+			}
+			for (int idx = 'A'; idx <= 'Z'; idx++)
+			{
+				chars[idx] = (byte)('a' + idx - 'A');
+			}
+			return chars;
+		}
+	}
+}
diff --git a/Engine/Source/Programs/Horde/Horde.Server/Logs/ILogBuilder.cs b/Engine/Source/Programs/Horde/Horde.Server/Logs/ILogBuilder.cs
new file mode 100644
--- /dev/null
+++ b/Engine/Source/Programs/Horde/Horde.Server/Logs/ILogBuilder.cs
@@ -0,0 +1,74 @@
+// Copyright Epic Games, Inc. All Rights Reserved.
+
+using System;
+using System.Collections.Generic;
+using System.Threading.Tasks;
+using Horde.Server.Logs.Data;
+using Horde.Server.Utilities;
+
+namespace Horde.Server.Logs
+{
+	/// <summary>
+	/// Interface for the log file write cache
+	/// </summary>
+	public interface ILogBuilder
+	{
+		/// <summary>
+		/// Whether the cache should be flushed on shutdown
+		/// </summary>
+		bool FlushOnShutdown { get; }
+
+		/// <summary>
+		/// Append data to a key
+		/// </summary>
+		/// <param name="logId">The log file id</param>
+		/// <param name="chunkOffset">Offset of the chunk within the log file</param>
+		/// <param name="writeOffset">Offset to write to</param>
+		/// <param name="writeLineIndex">Line index being written</param>
+		/// <param name="writeLineCount">Line count being written</param>
+		/// <param name="data">Data to be appended</param>
+		/// <param name="type">Type of data stored in this log file</param>
+		/// <returns>True if the data was appended to the given chunk. False if the chunk has been completed.</returns>
+		Task<bool> AppendAsync(LogId logId, long chunkOffset, long writeOffset, int writeLineIndex, int writeLineCount, ReadOnlyMemory<byte> data, LogType type);
+
+		/// <summary>
+		/// Finish the current sub chunk
+		/// </summary>
+		/// <param name="logId">The log file id</param>
+		/// <param name="offset">Offset of the chunk within the log file</param>
+		/// <returns>Async task</returns>
+		Task CompleteSubChunkAsync(LogId logId, long offset);
+
+		/// <summary>
+		/// Finish the current chunk
+		/// </summary>
+		/// <param name="logId">The log file id</param>
+		/// <param name="offset">Offset of the chunk within the log file</param>
+		/// <returns>Async task</returns>
+		Task CompleteChunkAsync(LogId logId, long offset);
+
+		/// <summary>
+		/// Remove a complete chunk from the builder
+		/// </summary>
+		/// <param name="logId">The log file id</param>
+		/// <param name="offset">Offset of the chunk within the log file</param>
+		/// <returns>Async task</returns>
+		Task RemoveChunkAsync(LogId logId, long offset);
+
+		/// <summary>
+		/// Gets the current chunk for the given log file
+		/// </summary>
+		/// <param name="logId">The log file id</param>
+		/// <param name="offset">Offset of the chunk within the log file</param>
+		/// <param name="lineIndex">Line index of the chunk within the log file</param>
+		/// <returns></returns>
+		Task<LogChunkData?> GetChunkAsync(LogId logId, long offset, int lineIndex);
+
+		/// <summary>
+		/// Touches the timestamps of all the chunks after the given age, and returns them. Used for flushing the builder.
+		/// </summary>
+		/// <param name="minAge">Minimum age of the chunks to enumerate. If specified, only chunks last modified longer than this period will be returned.</param>
+		/// <returns>List of chunks, identified by log id and chunk index</returns>
+		Task<List<(LogId, long)>> TouchChunksAsync(TimeSpan minAge);
+	}
+}
diff --git a/Engine/Source/Programs/Horde/Horde.Server/Logs/ILogEventCollection.cs b/Engine/Source/Programs/Horde/Horde.Server/Logs/ILogEventCollection.cs
new file mode 100644
--- /dev/null
+++ b/Engine/Source/Programs/Horde/Horde.Server/Logs/ILogEventCollection.cs
@@ -0,0 +1,62 @@
+// Copyright Epic Games, Inc. All Rights Reserved.
+
+using System.Collections.Generic;
+using System.Threading.Tasks;
+using Horde.Server.Utilities;
+using MongoDB.Bson;
+
+namespace Horde.Server.Logs
+{
+	/// <summary>
+	/// Interface for a collection of event documents
+	/// </summary>
+	public interface ILogEventCollection
+	{
+		/// <summary>
+		/// Creates a new event
+		/// </summary>
+		/// <param name="newEvent">The new event to vreate</param>
+		Task AddAsync(NewLogEventData newEvent);
+
+		/// <summary>
+		/// Creates a new event
+		/// </summary>
+		/// <param name="newEvents">List of events to create</param>
+		Task AddManyAsync(List<NewLogEventData> newEvents);
+
+		/// <summary>
+		/// Finds events within a log file
+		/// </summary>
+		/// <param name="logId">Unique id of the log containing this event</param>
+		/// <param name="spanId">Optional span to filter events by</param>
+		/// <param name="index">Start index within the matching results</param>
+		/// <param name="count">Maximum number of results to return</param>
+		/// <returns>List of events matching the query</returns>
+		Task<List<ILogEvent>> FindAsync(LogId logId, ObjectId? spanId, int? index = null, int? count = null);
+
+		/// <summary>
+		/// Finds a list of events for a set of spans
+		/// </summary>
+		/// <param name="spanIds">The span ids</param>
+		/// <param name="logIds">List of log ids to query</param>
+		/// <param name="index">Index of the first result to return</param>
+		/// <param name="count">Number of results to return</param>
+		/// <returns>List of events for this issue</returns>
+		Task<List<ILogEvent>> FindEventsForSpansAsync(IEnumerable<ObjectId> spanIds, LogId[]? logIds = null, int index = 0, int count = 10);
+
+		/// <summary>
+		/// Delete all the events for a log file
+		/// </summary>
+		/// <param name="logId">Unique id of the log</param>
+		/// <returns>Async task</returns>
+		Task DeleteLogAsync(LogId logId);
+
+		/// <summary>
+		/// Update the span for an event
+		/// </summary>
+		/// <param name="events">The events to update</param>
+		/// <param name="spanId">New span id</param>
+		/// <returns>Async task</returns>
+		Task AddSpanToEventsAsync(IEnumerable<ILogEvent> events, ObjectId spanId);
+	}
+}
diff --git a/Engine/Source/Programs/Horde/Horde.Server/Logs/ILogFile.cs b/Engine/Source/Programs/Horde/Horde.Server/Logs/ILogFile.cs
new file mode 100644
--- /dev/null
+++ b/Engine/Source/Programs/Horde/Horde.Server/Logs/ILogFile.cs
@@ -0,0 +1,144 @@
+// Copyright Epic Games, Inc. All Rights Reserved.
+
+using System;
+using System.Collections.Generic;
+using EpicGames.Core;
+using EpicGames.Horde.Storage;
+using Horde.Server.Agents.Leases;
+using Horde.Server.Agents.Sessions;
+using Horde.Server.Jobs;
+using Horde.Server.Storage;
+using Horde.Server.Utilities;
+
+namespace Horde.Server.Logs
+{
+	/// <summary>
+	/// Information about a log file chunk
+	/// </summary>
+	public interface ILogChunk
+	{
+		/// <summary>
+		/// Offset of the chunk within the log
+		/// </summary>
+		long Offset { get; }
+
+		/// <summary>
+		/// Length of this chunk. If zero, the chunk is still being written to.
+		/// </summary>
+		int Length { get; }
+
+		/// <summary>
+		/// Index of the first line within this chunk. If a line straddles two chunks, this is the index of the split line.
+		/// </summary>
+		int LineIndex { get; }
+
+		/// <summary>
+		/// If the chunk has yet to be pushed to persistent storage, includes the name of the server that is currently storing it.
+		/// </summary>
+		string? Server { get; }
+	}
+
+	/// <summary>
+	/// Information about a log file
+	/// </summary>
+	public interface ILogFile
+	{
+		/// <summary>
+		/// Identifier for the LogFile. Randomly generated.
+		/// </summary>
+		public LogId Id { get; }
+
+		/// <summary>
+		/// Unique id of the job containing this log
+		/// </summary>
+		public JobId JobId { get; }
+
+		/// <summary>
+		/// The lease allowed to write to this log
+		/// </summary>
+		public LeaseId? LeaseId { get; }
+
+		/// <summary>
+		/// The session allowed to write to this log
+		/// </summary>
+		public SessionId? SessionId { get; }
+
+		/// <summary>
+		/// Whether to use the new storage backend for log data
+		/// </summary>
+		public bool UseNewStorageBackend { get; }
+
+		/// <summary>
+		/// Maximum line index in the file
+		/// </summary>
+		public int? MaxLineIndex { get; }
+
+		/// <summary>
+		/// Length of the file which is indexed
+		/// </summary>
+		public long? IndexLength { get; }
+
+		/// <summary>
+		/// Type of data stored in this log 
+		/// </summary>
+		public LogType Type { get; }
+
+		/// <summary>
+		/// Chunks within this file
+		/// </summary>
+		public IReadOnlyList<ILogChunk> Chunks { get; }
+
+		/// <summary>
+		/// Name of the ref used to store data for this log
+		/// </summary>
+		public RefName RefName { get; }
+
+		/// <summary>
+		/// Number of lines (V2 storage backend)
+		/// </summary>
+		public int LineCount { get; }
+
+		/// <summary>
+		/// Whether the log is complete (V2 storage backend)
+		/// </summary>
+		public bool Complete { get; }
+	}
+
+	/// <summary>
+	/// Extension methods for log files
+	/// </summary>
+	public static class LogFileExtensions
+	{
+		/// <summary>
+		/// Gets the chunk index containing the given offset.
+		/// </summary>
+		/// <param name="chunks">The chunks to search</param>
+		/// <param name="offset">The offset to search for</param>
+		/// <returns>The chunk index containing the given offset</returns>
+		public static int GetChunkForOffset(this IReadOnlyList<ILogChunk> chunks, long offset)
+		{
+			int chunkIndex = chunks.BinarySearch(x => x.Offset, offset);
+			if (chunkIndex < 0)
+			{
+				chunkIndex = ~chunkIndex - 1;
+			}
+			return chunkIndex;
+		}
+
+		/// <summary>
+		/// Gets the starting chunk index for the given line
+		/// </summary>
+		/// <param name="chunks">The chunks to search</param>
+		/// <param name="lineIndex">Index of the line to query</param>
+		/// <returns>Index of the chunk to fetch</returns>
+		public static int GetChunkForLine(this IReadOnlyList<ILogChunk> chunks, int lineIndex)
+		{
+			int chunkIndex = chunks.BinarySearch(x => x.LineIndex, lineIndex);
+			if(chunkIndex < 0)
+			{
+				chunkIndex = ~chunkIndex - 1;
+			}
+			return chunkIndex;
+		}
+	}
+}
diff --git a/Engine/Source/Programs/Horde/Horde.Server/Logs/ILogFileCollection.cs b/Engine/Source/Programs/Horde/Horde.Server/Logs/ILogFileCollection.cs
new file mode 100644
--- /dev/null
+++ b/Engine/Source/Programs/Horde/Horde.Server/Logs/ILogFileCollection.cs
@@ -0,0 +1,120 @@
+// Copyright Epic Games, Inc. All Rights Reserved.
+
+using System.Collections.Generic;
+using System.Threading;
+using System.Threading.Tasks;
+using Horde.Server.Agents.Leases;
+using Horde.Server.Agents.Sessions;
+using Horde.Server.Jobs;
+using Horde.Server.Utilities;
+
+namespace Horde.Server.Logs
+{
+	/// <summary>
+	/// Updates a log file chunk
+	/// </summary>
+	public class CompleteLogChunkUpdate
+	{
+		/// <summary>
+		/// Index of the chunk
+		/// </summary>
+		public int Index { get; set; }
+
+		/// <summary>
+		/// New length for the chunk
+		/// </summary>
+		public int Length { get; set; }
+
+		/// <summary>
+		/// Number of lines in the chunk
+		/// </summary>
+		public int LineCount { get; set; }
+
+		/// <summary>
+		/// Constructor
+		/// </summary>
+		/// <param name="index">Index of the chunk</param>
+		/// <param name="length">New length for the chunk</param>
+		/// <param name="lineCount">Number of lines in the chunk</param>
+		public CompleteLogChunkUpdate(int index, int length, int lineCount)
+		{
+			Index = index;
+			Length = length;
+			LineCount = lineCount;
+		}
+	}
+
+	/// <summary>
+	/// Wrapper around the jobs collection in a mongo DB
+	/// </summary>
+	public interface ILogFileCollection
+	{
+		/// <summary>
+		/// Creates a new log
+		/// </summary>
+		/// <param name="jobId">Unique id of the job that owns this log file</param>
+		/// <param name="leaseId">Agent lease allowed to update the log</param>
+		/// <param name="sessionId">Agent session allowed to update the log</param>
+		/// <param name="type">Type of events to be stored in the log</param>
+		/// <param name="newStorageBackend">Whether to use the new storage backend</param>
+		/// <param name="logId">ID of the log file (optional)</param>
+		/// <param name="cancellationToken">Cancellation token for the call</param>
+		/// <returns>The new log file document</returns>
+		Task<ILogFile> CreateLogFileAsync(JobId jobId, LeaseId? leaseId, SessionId? sessionId, LogType type, bool newStorageBackend, LogId? logId = null, CancellationToken cancellationToken = default);
+
+		/// <summary>
+		/// Updates the line count for a log file (v2 backend only)
+		/// </summary>
+		/// <param name="logFileInterface">The current log file</param>
+		/// <param name="lineCount">New line count for the log file</param>
+		/// <param name="complete">Flag indicating whether the log is complete, or can still be tailed for additional data</param>
+		/// <param name="cancellationToken">Cancellation token for the call</param>
+		/// <returns>The updated log file document</returns>
+		Task<ILogFile> UpdateLineCountAsync(ILogFile logFileInterface, int lineCount, bool complete, CancellationToken cancellationToken);
+
+		/// <summary>
+		/// Adds a new chunk
+		/// </summary>
+		/// <param name="logFileInterface">The current log file</param>
+		/// <param name="offset">Offset of the new chunk</param>
+		/// <param name="lineIndex">Line index for the start of the chunk</param>
+		/// <param name="cancellationToken">Cancellation token for the call</param>
+		/// <returns>The updated log file document</returns>
+		Task<ILogFile?> TryAddChunkAsync(ILogFile logFileInterface, long offset, int lineIndex, CancellationToken cancellationToken);
+
+		/// <summary>
+		/// Update the log file with final information about certain chunks
+		/// </summary>
+		/// <param name="logFileInterface">The current log file</param>
+		/// <param name="chunks">Chunks to update. New chunks will be inserted</param>
+		/// <param name="cancellationToken">Cancellation token for the call</param>
+		/// <returns>The updated log file document</returns>
+		Task<ILogFile?> TryCompleteChunksAsync(ILogFile logFileInterface, IEnumerable<CompleteLogChunkUpdate> chunks, CancellationToken cancellationToken);
+
+		/// <summary>
+		/// Update the log file with final information about the index
+		/// </summary>
+		/// <param name="logFileInterface">The current log file</param>
+		/// <param name="newIndexLength">New length of the index</param>
+		/// <param name="cancellationToken">Cancellation token for the call</param>
+		/// <returns>The updated log file document</returns>
+		Task<ILogFile?> TryUpdateIndexAsync(ILogFile logFileInterface, long newIndexLength, CancellationToken cancellationToken);
+
+		/// <summary>
+		/// Gets a logfile by ID
+		/// </summary>
+		/// <param name="logFileId">Unique id of the log file</param>
+		/// <param name="cancellationToken">Cancellation token for the call</param>
+		/// <returns>The logfile document</returns>
+		Task<ILogFile?> GetLogFileAsync(LogId logFileId, CancellationToken cancellationToken);
+
+		/// <summary>
+		/// Gets all the log files
+		/// </summary>
+		/// <param name="index">Index of the first result to return</param>
+		/// <param name="count">Number of results to return</param>
+		/// <param name="cancellationToken">Cancellation token for the call</param>
+		/// <returns>List of log files</returns>
+		Task<List<ILogFile>> GetLogFilesAsync(int? index = null, int? count = null, CancellationToken cancellationToken = default);
+	}
+}
diff --git a/Engine/Source/Programs/Horde/Horde.Server/Logs/ILogStorage.cs b/Engine/Source/Programs/Horde/Horde.Server/Logs/ILogStorage.cs
new file mode 100644
--- /dev/null
+++ b/Engine/Source/Programs/Horde/Horde.Server/Logs/ILogStorage.cs
@@ -0,0 +1,50 @@
+// Copyright Epic Games, Inc. All Rights Reserved.
+
+using System;
+using System.Threading.Tasks;
+using Horde.Server.Logs.Data;
+using Horde.Server.Utilities;
+
+namespace Horde.Server.Logs
+{
+	/// <summary>
+	/// Caching interface for reading and writing log data.
+	/// </summary>
+	public interface ILogStorage : IDisposable
+	{
+		/// <summary>
+		/// Attempts to read an index for the given log file
+		/// </summary>
+		/// <param name="logId">Unique id of the log file</param>
+		/// <param name="length">Length of the file covered by the index</param>
+		/// <returns>Index for the log file</returns>
+		Task<LogIndexData?> ReadIndexAsync(LogId logId, long length);
+
+		/// <summary>
+		/// Log file to write an index for
+		/// </summary>
+		/// <param name="logId">Unique id of the log file</param>
+		/// <param name="length">Length of the file covered by the index</param>
+		/// <param name="index">The log file index</param>
+		/// <returns>Async task</returns>
+		Task WriteIndexAsync(LogId logId, long length, LogIndexData index);
+
+		/// <summary>
+		/// Retrieves an item from the cache
+		/// </summary>
+		/// <param name="logId">Unique id of the log file</param>
+		/// <param name="offset">Offset of the chunk to read</param>
+		/// <param name="lineIndex">First line of the chunk</param>
+		/// <returns>Data for the given key, or null if it's not present</returns>
+		Task<LogChunkData?> ReadChunkAsync(LogId logId, long offset, int lineIndex);
+
+		/// <summary>
+		/// Writes a chunk to storage
+		/// </summary>
+		/// <param name="logId">Unique id of the log file</param>
+		/// <param name="offset">Offset of the chunk to write</param>
+		/// <param name="chunkData">Information about the chunk data</param>
+		/// <returns>Async task</returns>
+		Task WriteChunkAsync(LogId logId, long offset, LogChunkData chunkData);
+	}
+}
diff --git a/Engine/Source/Programs/Horde/Horde.Server/Logs/LogEventCollection.cs b/Engine/Source/Programs/Horde/Horde.Server/Logs/LogEventCollection.cs
new file mode 100644
--- /dev/null
+++ b/Engine/Source/Programs/Horde/Horde.Server/Logs/LogEventCollection.cs
@@ -0,0 +1,187 @@
+// Copyright Epic Games, Inc. All Rights Reserved.
+
+using System;
+using System.Collections.Generic;
+using System.Linq;
+using System.Threading.Tasks;
+using Horde.Server.Server;
+using Horde.Server.Utilities;
+using HordeCommon;
+using Microsoft.Extensions.Logging;
+using MongoDB.Bson;
+using MongoDB.Bson.Serialization.Attributes;
+using MongoDB.Driver;
+
+namespace Horde.Server.Logs
+{
+	/// <summary>
+	/// Collection of event documents
+	/// </summary>
+	public class LogEventCollection : ILogEventCollection
+	{
+		class LogEventId
+		{
+			[BsonElement("l")]
+			public LogId LogId { get; set; }
+
+			[BsonElement("n")]
+			public int LineIndex { get; set; }
+		}
+
+		class LogEventDocument : ILogEvent
+		{
+			[BsonId]
+			public LogEventId Id { get; set; }
+
+			[BsonElement("w"), BsonIgnoreIfDefault, BsonDefaultValue(false)]
+			public bool IsWarning { get; set; }
+
+			[BsonElement("c"), BsonIgnoreIfNull]
+			public int? LineCount { get; set; }
+
+			[BsonElement("s")]
+			public ObjectId? SpanId { get; set; }
+
+			LogId ILogEvent.LogId => Id.LogId;
+			int ILogEvent.LineIndex => Id.LineIndex;
+			int ILogEvent.LineCount => LineCount ?? 1;
+			EventSeverity ILogEvent.Severity => IsWarning ? EventSeverity.Warning : EventSeverity.Error;
+
+			public LogEventDocument()
+			{
+				Id = new LogEventId();
+			}
+
+			public LogEventDocument(LogId logId, EventSeverity severity, int lineIndex, int lineCount, ObjectId? spanId)
+			{
+				Id = new LogEventId { LogId = logId, LineIndex = lineIndex };
+				IsWarning = severity == EventSeverity.Warning;
+				LineCount = (lineCount > 1) ? (int?)lineCount : null;
+				SpanId = spanId;
+			}
+
+			public LogEventDocument(NewLogEventData data)
+				: this(data.LogId, data.Severity, data.LineIndex, data.LineCount, data.SpanId)
+			{
+			}
+		}
+
+		class LegacyEventDocument
+		{
+			public ObjectId Id { get; set; }
+			public DateTime Time { get; set; }
+			public EventSeverity Severity { get; set; }
+			public LogId LogId { get; set; }
+			public int LineIndex { get; set; }
+			public int LineCount { get; set; }
+
+			public string? Message { get; set; }
+
+			[BsonIgnoreIfNull, BsonElement("IssueId2")]
+			public int? IssueId { get; set; }
+
+			public BsonDocument? Data { get; set; }
+
+			public int UpgradeVersion { get; set; }
+		}
+
+		/// <summary>
+		/// Collection of event documents
+		/// </summary>
+		readonly IMongoCollection<LogEventDocument> _logEvents;
+
+		/// <summary>
+		/// Collection of legacy event documents
+		/// </summary>
+		readonly IMongoCollection<LegacyEventDocument> _legacyEvents;
+
+		/// <summary>
+		/// Logger for upgrade messages
+		/// </summary>
+		readonly ILogger<LogEventCollection> _logger;
+
+		/// <summary>
+		/// Constructor
+		/// </summary>
+		/// <param name="mongoService">The database service instance</param>
+		/// <param name="logger">The logger instance</param>
+		public LogEventCollection(MongoService mongoService, ILogger<LogEventCollection> logger)
+		{
+			_logger = logger;
+
+			List<MongoIndex<LogEventDocument>> logEventIndexes = new List<MongoIndex<LogEventDocument>>();
+			logEventIndexes.Add(keys => keys.Ascending(x => x.Id.LogId));
+			logEventIndexes.Add(keys => keys.Ascending(x => x.SpanId).Ascending(x => x.Id));
+			_logEvents = mongoService.GetCollection<LogEventDocument>("LogEvents", logEventIndexes);
+
+			_legacyEvents = mongoService.GetCollection<LegacyEventDocument>("Events", keys => keys.Ascending(x => x.LogId));
+		}
+
+		/// <inheritdoc/>
+		public Task AddAsync(NewLogEventData newEvent)
+		{
+			return _logEvents.InsertOneAsync(new LogEventDocument(newEvent));
+		}
+
+		/// <inheritdoc/>
+		public Task AddManyAsync(List<NewLogEventData> newEvents)
+		{
+			return _logEvents.InsertManyAsync(newEvents.ConvertAll(x => new LogEventDocument(x)));
+		}
+
+		/// <inheritdoc/>
+		public async Task<List<ILogEvent>> FindAsync(LogId logId, ObjectId? spanId = null, int? index = null, int? count = null)
+		{
+			_logger.LogInformation("Querying for log events for log {LogId} creation time {CreateTime}", logId, logId.Id.CreationTime);
+
+			FilterDefinitionBuilder<LogEventDocument> builder = Builders<LogEventDocument>.Filter;
+
+			FilterDefinition<LogEventDocument> filter = builder.Eq(x => x.Id.LogId, logId);
+			if(spanId != null)
+			{
+				filter &= builder.Eq(x => x.SpanId, spanId.Value);
+			}
+
+			IFindFluent<LogEventDocument, LogEventDocument> results = _logEvents.Find(filter).SortBy(x => x.Id);
+			if (index != null)
+			{
+				results = results.Skip(index.Value);
+			}
+			if (count != null)
+			{
+				results = results.Limit(count.Value);
+			}
+
+			List<LogEventDocument> eventDocuments = await results.ToListAsync();
+			return eventDocuments.ConvertAll<ILogEvent>(x => x);
+		}
+
+		/// <inheritdoc/>
+		public async Task<List<ILogEvent>> FindEventsForSpansAsync(IEnumerable<ObjectId> spanIds, LogId[]? logIds, int index, int count)
+		{
+			FilterDefinition<LogEventDocument> filter = Builders<LogEventDocument>.Filter.In(x => x.SpanId, spanIds.Select<ObjectId, ObjectId?>(x => x));
+			if (logIds != null && logIds.Length > 0)
+			{
+				filter &= Builders<LogEventDocument>.Filter.In(x => x.Id.LogId, logIds);
+			}
+
+			List<LogEventDocument> eventDocuments = await _logEvents.Find(filter).Skip(index).Limit(count).ToListAsync();
+			return eventDocuments.ConvertAll<ILogEvent>(x => x);
+		}
+
+		/// <inheritdoc/>
+		public async Task DeleteLogAsync(LogId logId)
+		{
+			await _logEvents.DeleteManyAsync(x => x.Id.LogId == logId);
+			await _legacyEvents.DeleteManyAsync(x => x.LogId == logId);
+		}
+
+		/// <inheritdoc/>
+		public async Task AddSpanToEventsAsync(IEnumerable<ILogEvent> events, ObjectId spanId)
+		{
+			FilterDefinition<LogEventDocument> eventFilter = Builders<LogEventDocument>.Filter.In(x => x.Id, events.OfType<LogEventDocument>().Select(x => x.Id));
+			UpdateDefinition<LogEventDocument> eventUpdate = Builders<LogEventDocument>.Update.Set(x => x.SpanId, spanId);
+			await _logEvents.UpdateManyAsync(eventFilter, eventUpdate);
+		}
+	}
+}
diff --git a/Engine/Source/Programs/Horde/Horde.Server/Logs/LogEventMessages.cs b/Engine/Source/Programs/Horde/Horde.Server/Logs/LogEventMessages.cs
new file mode 100644
--- /dev/null
+++ b/Engine/Source/Programs/Horde/Horde.Server/Logs/LogEventMessages.cs
@@ -0,0 +1,61 @@
+// Copyright Epic Games, Inc. All Rights Reserved.
+
+using System.Collections.Generic;
+using System.Text.Json;
+using EpicGames.Core;
+using HordeCommon;
+
+namespace Horde.Server.Logs
+{
+	/// <summary>
+	/// Information about an uploaded event
+	/// </summary>
+	public class GetLogEventResponse
+	{
+		/// <summary>
+		/// Unique id of the log containing this event
+		/// </summary>
+		public string LogId { get; set; }
+
+		/// <summary>
+		/// Severity of this event
+		/// </summary>
+		public EventSeverity Severity { get; set; }
+
+		/// <summary>
+		/// Index of the first line for this event
+		/// </summary>
+		public int LineIndex { get; set; }
+
+		/// <summary>
+		/// Number of lines in the event
+		/// </summary>
+		public int LineCount { get; set; }
+
+		/// <summary>
+		/// The issue id associated with this event
+		/// </summary>
+		public int? IssueId { get; set; }
+
+		/// <summary>
+		/// The structured message data for this event
+		/// </summary>
+		public List<JsonElement> Lines { get; set; }
+
+		/// <summary>
+		/// Constructor
+		/// </summary>
+		/// <param name="logEvent">The event to construct from</param>
+		/// <param name="eventData">The event data</param>
+		/// <param name="issueId">The issue for this event</param>
+		public GetLogEventResponse(ILogEvent logEvent, ILogEventData eventData, int? issueId)
+		{
+			Severity = logEvent.Severity;
+			LogId = logEvent.LogId.ToString();
+			LineIndex = logEvent.LineIndex;
+			LineCount = logEvent.LineCount;
+			IssueId = issueId;
+			Lines = eventData.Lines.ConvertAll(x => x.Data);
+		}
+	}
+}
diff --git a/Engine/Source/Programs/Horde/Horde.Server/Logs/LogFileCollection.cs b/Engine/Source/Programs/Horde/Horde.Server/Logs/LogFileCollection.cs
new file mode 100644
--- /dev/null
+++ b/Engine/Source/Programs/Horde/Horde.Server/Logs/LogFileCollection.cs
@@ -0,0 +1,256 @@
+// Copyright Epic Games, Inc. All Rights Reserved.
+
+using System.Collections.Generic;
+using System.Linq;
+using System.Net;
+using System.Threading;
+using System.Threading.Tasks;
+using EpicGames.Horde.Storage;
+using Horde.Server.Agents.Leases;
+using Horde.Server.Agents.Sessions;
+using Horde.Server.Jobs;
+using Horde.Server.Server;
+using Horde.Server.Streams;
+using Horde.Server.Utilities;
+using MongoDB.Bson.Serialization.Attributes;
+using MongoDB.Driver;
+
+namespace Horde.Server.Logs
+{
+	/// <summary>
+	/// Wrapper around the jobs collection in a mongo DB
+	/// </summary>
+	public class LogFileCollection : ILogFileCollection
+	{
+		class LogChunkDocument : ILogChunk
+		{
+			public long Offset { get; set; }
+			public int Length { get; set; }
+			public int LineIndex { get; set; }
+
+			[BsonIgnoreIfNull]
+			public string? Server { get; set; }
+
+			[BsonConstructor]
+			public LogChunkDocument()
+			{
+			}
+
+			public LogChunkDocument(LogChunkDocument other)
+			{
+				Offset = other.Offset;
+				Length = other.Length;
+				LineIndex = other.LineIndex;
+				Server = other.Server;
+			}
+
+			public LogChunkDocument Clone()
+			{
+				return (LogChunkDocument)MemberwiseClone();
+			}
+		}
+
+		class LogFileDocument : ILogFile
+		{
+			[BsonRequired, BsonId]
+			public LogId Id { get; set; }
+
+			[BsonRequired]
+			public JobId JobId { get; set; }
+
+			[BsonIgnoreIfNull]
+			public LeaseId? LeaseId { get; set; }
+
+			[BsonIgnoreIfNull]
+			public SessionId? SessionId { get; set; }
+
+			public LogType Type { get; set; }
+			public bool UseNewStorageBackend { get; set; }
+
+			[BsonIgnoreIfNull]
+			public int? MaxLineIndex { get; set; }
+
+			[BsonIgnoreIfNull]
+			public long? IndexLength { get; set; }
+
+			public List<LogChunkDocument> Chunks { get; set; } = new List<LogChunkDocument>();
+
+			public int LineCount { get; set; }
+			public RefName RefName { get; set; }
+
+			[BsonIgnoreIfDefault]
+			public bool Complete { get; set; }
+
+			[BsonRequired]
+			public int UpdateIndex { get; set; }
+
+			IReadOnlyList<ILogChunk> ILogFile.Chunks => Chunks;
+
+			[BsonConstructor]
+			private LogFileDocument()
+			{
+			}
+
+			public LogFileDocument(JobId jobId, LeaseId? leaseId, SessionId? sessionId, LogType type, bool newStorageBackend, LogId? logId)
+			{
+				Id = logId ?? LogId.GenerateNewId();
+				JobId = jobId;
+				LeaseId = leaseId;
+				SessionId = sessionId;
+				Type = type;
+				UseNewStorageBackend = newStorageBackend;
+				MaxLineIndex = 0;
+				RefName = new RefName(Id.ToString());
+			}
+
+			public LogFileDocument Clone()
+			{
+				LogFileDocument document = (LogFileDocument)MemberwiseClone();
+				document.Chunks = document.Chunks.ConvertAll(x => x.Clone());
+				return document;
+			}
+		}
+
+		/// <summary>
+		/// The jobs collection
+		/// </summary>
+		readonly IMongoCollection<LogFileDocument> _logFiles;
+
+		/// <summary>
+		/// Hostname for the current server
+		/// </summary>
+		readonly string _hostName;
+
+		/// <summary>
+		/// Constructor
+		/// </summary>
+		/// <param name="mongoService">The database service singleton</param>
+		public LogFileCollection(MongoService mongoService)
+		{
+			_logFiles = mongoService.GetCollection<LogFileDocument>("LogFiles");
+			_hostName = Dns.GetHostName();
+		}
+
+		/// <inheritdoc/>
+		public async Task<ILogFile> CreateLogFileAsync(JobId jobId, LeaseId? leaseId, SessionId? sessionId, LogType type, bool newStorageBackend, LogId? logId, CancellationToken cancellationToken)
+		{
+			LogFileDocument newLogFile = new (jobId, leaseId, sessionId, type, newStorageBackend, logId);
+			await _logFiles.InsertOneAsync(newLogFile, null, cancellationToken);
+			return newLogFile;
+		}
+
+		/// <inheritdoc/>
+		public async Task<ILogFile> UpdateLineCountAsync(ILogFile logFileInterface, int lineCount, bool complete, CancellationToken cancellationToken)
+		{
+			FilterDefinition<LogFileDocument> filter = Builders<LogFileDocument>.Filter.Eq(x => x.Id, logFileInterface.Id);
+			UpdateDefinition<LogFileDocument> update = Builders<LogFileDocument>.Update.Set(x => x.LineCount, lineCount).Set(x => x.Complete, complete).Inc(x => x.UpdateIndex, 1);
+			return await _logFiles.FindOneAndUpdateAsync(filter, update, new FindOneAndUpdateOptions<LogFileDocument, LogFileDocument> { ReturnDocument = ReturnDocument.After }, cancellationToken);
+		}
+
+		/// <inheritdoc/>
+		public async Task<ILogFile?> TryAddChunkAsync(ILogFile logFileInterface, long offset, int lineIndex, CancellationToken cancellationToken)
+		{
+			LogFileDocument logFile = ((LogFileDocument)logFileInterface).Clone();
+
+			int chunkIdx = logFile.Chunks.GetChunkForOffset(offset) + 1;
+
+			LogChunkDocument chunk = new LogChunkDocument();
+			chunk.Offset = offset;
+			chunk.LineIndex = lineIndex;
+			chunk.Server = _hostName;
+			logFile.Chunks.Insert(chunkIdx, chunk);
+
+			UpdateDefinition<LogFileDocument> update = Builders<LogFileDocument>.Update.Set(x => x.Chunks, logFile.Chunks);
+			if (chunkIdx == logFile.Chunks.Count - 1)
+			{
+				logFile.MaxLineIndex = null;
+				update = update.Unset(x => x.MaxLineIndex);
+			}
+
+			if (!await TryUpdateLogFileAsync(logFile, update, cancellationToken))
+			{
+				return null;
+			}
+
+			return logFile;
+		}
+
+		/// <inheritdoc/>
+		public async Task<ILogFile?> TryCompleteChunksAsync(ILogFile logFileInterface, IEnumerable<CompleteLogChunkUpdate> chunkUpdates, CancellationToken cancellationToken)
+		{
+			LogFileDocument logFile = ((LogFileDocument)logFileInterface).Clone();
+
+			// Update the length of any complete chunks
+			UpdateDefinitionBuilder<LogFileDocument> updateBuilder = Builders<LogFileDocument>.Update;
+			List<UpdateDefinition<LogFileDocument>> updates = new List<UpdateDefinition<LogFileDocument>>();
+			foreach (CompleteLogChunkUpdate chunkUpdate in chunkUpdates)
+			{
+				LogChunkDocument chunk = logFile.Chunks[chunkUpdate.Index];
+				chunk.Length = chunkUpdate.Length;
+				updates.Add(updateBuilder.Set(x => x.Chunks[chunkUpdate.Index].Length, chunkUpdate.Length));
+
+				if (chunkUpdate.Index == logFile.Chunks.Count - 1)
+				{
+					logFile.MaxLineIndex = chunk.LineIndex + chunkUpdate.LineCount;
+					updates.Add(updateBuilder.Set(x => x.MaxLineIndex, logFile.MaxLineIndex));
+				}
+			}
+
+			// Try to apply the updates
+			if (updates.Count > 0 && !await TryUpdateLogFileAsync(logFile, updateBuilder.Combine(updates), cancellationToken))
+			{
+				return null;
+			}
+
+			return logFile;
+		}
+
+		/// <inheritdoc/>
+		public async Task<ILogFile?> TryUpdateIndexAsync(ILogFile logFileInterface, long newIndexLength, CancellationToken cancellationToken)
+		{
+			LogFileDocument logFile = ((LogFileDocument)logFileInterface).Clone();
+
+			UpdateDefinition<LogFileDocument> update = Builders<LogFileDocument>.Update.Set(x => x.IndexLength, newIndexLength);
+			if (!await TryUpdateLogFileAsync(logFile, update, cancellationToken))
+			{
+				return null;
+			}
+
+			logFile.IndexLength = newIndexLength;
+			return logFile;
+		}
+
+		/// <inheritdoc/>
+		private async Task<bool> TryUpdateLogFileAsync(LogFileDocument current, UpdateDefinition<LogFileDocument> update, CancellationToken cancellationToken)
+		{
+			int prevUpdateIndex = current.UpdateIndex;
+			current.UpdateIndex++;
+			UpdateResult result = await _logFiles.UpdateOneAsync<LogFileDocument>(x => x.Id == current.Id && x.UpdateIndex == prevUpdateIndex, update.Set(x => x.UpdateIndex, current.UpdateIndex), cancellationToken: cancellationToken);
+			return result.ModifiedCount == 1;
+		}
+
+		/// <inheritdoc/>
+		public async Task<ILogFile?> GetLogFileAsync(LogId logFileId, CancellationToken cancellationToken)
+		{
+			LogFileDocument logFile = await _logFiles.Find<LogFileDocument>(x => x.Id == logFileId).FirstOrDefaultAsync(cancellationToken);
+			return logFile;
+		}
+
+		/// <inheritdoc/>
+		public async Task<List<ILogFile>> GetLogFilesAsync(int? index = null, int? count = null, CancellationToken cancellationToken = default)
+		{
+			IFindFluent<LogFileDocument, LogFileDocument> query = _logFiles.Find(FilterDefinition<LogFileDocument>.Empty);
+			if(index != null)
+			{
+				query = query.Skip(index.Value);
+			}
+			if(count != null)
+			{
+				query = query.Limit(count.Value);
+			}
+
+			List<LogFileDocument> results = await query.ToListAsync(cancellationToken);
+			return results.ConvertAll<ILogFile>(x => x);
+		}
+	}
+}
diff --git a/Engine/Source/Programs/Horde/Horde.Server/Logs/LogFileMessages.cs b/Engine/Source/Programs/Horde/Horde.Server/Logs/LogFileMessages.cs
new file mode 100644
--- /dev/null
+++ b/Engine/Source/Programs/Horde/Horde.Server/Logs/LogFileMessages.cs
@@ -0,0 +1,103 @@
+// Copyright Epic Games, Inc. All Rights Reserved.
+
+using System;
+using System.Collections.Generic;
+using EpicGames.Horde.Logs;
+using Horde.Server.Logs.Data;
+
+namespace Horde.Server.Logs
+{
+	/// <summary>
+	/// The type of data stored in this log file
+	/// </summary>
+	public enum LogType
+	{
+		/// <summary>
+		/// Plain text data
+		/// </summary>
+		Text,
+
+		/// <summary>
+		/// Structured json objects, output as one object per line (without trailing commas)
+		/// </summary>
+		Json
+	}
+
+	/// <summary>
+	/// Creates a new log file
+	/// </summary>
+	public class CreateLogFileRequest
+	{
+		/// <summary>
+		/// Type of the log file
+		/// </summary>
+		public LogType Type { get; set; } = LogType.Json;
+	}
+
+	/// <summary>
+	/// Response from creating a log file
+	/// </summary>
+	public class CreateLogFileResponse
+	{
+		/// <summary>
+		/// Identifier for the created log file
+		/// </summary>
+		public string Id { get; set; } = String.Empty;
+	}
+
+	/// <summary>
+	/// Response describing a log file
+	/// </summary>
+	public class GetLogFileResponse
+	{
+        /// <summary>
+		/// Unique id of the log file
+		/// </summary>
+		public string Id { get; set; }
+
+        /// <summary>
+        /// Unique id of the job for this log file
+		/// </summary>
+        public string JobId { get; set; }
+
+		/// <summary>
+		/// Type of events stored in this log
+		/// </summary>
+		public LogType Type { get; set; }
+
+		/// <summary>
+        /// Number of lines in the file
+        /// </summary>
+        public int LineCount { get; set; }
+
+		/// <summary>
+		/// Constructor
+		/// </summary>
+		/// <param name="logFile">The logfile to construct from</param>
+		/// <param name="metadata">Metadata about the log file</param>
+		public GetLogFileResponse(ILogFile logFile, LogMetadata metadata)
+		{
+            Id = logFile.Id.ToString();
+            JobId = logFile.JobId.ToString();
+			Type = logFile.Type;
+            LineCount = metadata.MaxLineIndex;
+		}
+	}
+
+	/// <summary>
+	/// Response describing a log file
+	/// </summary>
+	public class SearchLogFileResponse
+	{
+		/// <summary>
+		/// List of line numbers containing the search text
+		/// </summary>
+		public List<int> Lines { get; set; } = new List<int>();
+
+		/// <summary>
+		/// Stats for the search
+		/// </summary>
+		public SearchStats? Stats { get; set; }
+	}
+}
+
