Horde: Rename the Horde.Build project to Horde.Server.

diff --git a/Engine/Source/Programs/Horde/Horde.Server/Logs/Storage/RedisLogStorage.cs b/Engine/Source/Programs/Horde/Horde.Server/Logs/Storage/RedisLogStorage.cs
new file mode 100644
--- /dev/null
+++ b/Engine/Source/Programs/Horde/Horde.Server/Logs/Storage/RedisLogStorage.cs
@@ -0,0 +1,167 @@
+// Copyright Epic Games, Inc. All Rights Reserved.
+
+using System;
+using System.Threading.Tasks;
+using Horde.Server.Logs.Data;
+using Horde.Server.Utilities;
+using Microsoft.Extensions.Logging;
+using StackExchange.Redis;
+
+namespace Horde.Server.Logs.Storage
+{
+	/// <summary>
+	/// Redis log file storage
+	/// </summary>
+	public sealed class RedisLogStorage : ILogStorage
+	{
+		/// <summary>
+		/// The Redis database
+		/// </summary>
+		readonly IDatabase _redisDb;
+
+		/// <summary>
+		/// Logging device
+		/// </summary>
+		readonly ILogger _logger;
+
+		/// <summary>
+		/// Inner storage layer
+		/// </summary>
+		readonly ILogStorage _inner;
+
+		/// <summary>
+		/// Constructor
+		/// </summary>
+		/// <param name="redisDb">Redis database</param>
+		/// <param name="logger">Logger instance</param>
+		/// <param name="inner">Inner storage layer</param>
+		public RedisLogStorage(IDatabase redisDb, ILogger logger, ILogStorage inner)
+		{
+			_redisDb = redisDb;
+			_logger = logger;
+			_inner = inner;
+		}
+
+		/// <inheritdoc/>
+		public void Dispose()
+		{
+			_inner.Dispose();
+		}
+
+		/// <summary>
+		/// Gets the key for a log file's index
+		/// </summary>
+		/// <param name="logId">The log file id</param>
+		/// <param name="length">Length of the file covered by the index</param>
+		/// <returns>The index key</returns>
+		static string IndexKey(LogId logId, long length)
+		{
+			return $"log-{logId}-index-{length}";
+		}
+
+		/// <summary>
+		/// Gets the key for a chunk's data
+		/// </summary>
+		/// <param name="logId">The log file id</param>
+		/// <param name="offset">Offset of the chunk within the log file</param>
+		/// <returns>The chunk key</returns>
+		static string ChunkKey(LogId logId, long offset)
+		{
+			return $"log-{logId}-chunk-{offset}";
+		}
+
+		/// <summary>
+		/// Adds an index to the cache
+		/// </summary>
+		/// <param name="key">Key for the item to store</param>
+		/// <param name="value">Value to be stored</param>
+		/// <returns>Async task</returns>
+		async Task AddAsync(string key, ReadOnlyMemory<byte> value)
+		{
+			try
+			{
+				await _redisDb.StringSetAsync(key, value, expiry: TimeSpan.FromHours(1.0));
+				_logger.LogDebug("Added key {Key} to Redis cache ({Size} bytes)", key, value.Length);
+			}
+			catch (Exception ex)
+			{
+				_logger.LogError(ex, "Error writing Redis key {Key}", key);
+			}
+		}
+
+		/// <summary>
+		/// Adds an index to the cache
+		/// </summary>
+		/// <param name="key">Key for the item to store</param>
+		/// <param name="deserialize">Delegate to deserialize the item</param>
+		/// <returns>Async task</returns>
+		async Task<T?> GetAsync<T>(string key, Func<ReadOnlyMemory<byte>, T> deserialize) where T : class
+		{
+			try
+			{
+				RedisValue value = await _redisDb.StringGetAsync(key);
+				if (value.IsNullOrEmpty)
+				{
+					_logger.LogDebug("Redis cache miss for {Key}", key);
+					return null;
+				}
+				else
+				{
+					_logger.LogDebug("Redis cache hit for {Key}", key);
+					return deserialize(value);
+				}
+			}
+			catch (Exception ex)
+			{
+				_logger.LogError(ex, "Error reading Redis key {Key}", key);
+				return null;
+			}
+		}
+
+		/// <inheritdoc/>
+		public async Task<LogIndexData?> ReadIndexAsync(LogId logId, long length)
+		{
+			string key = IndexKey(logId, length);
+
+			LogIndexData? indexData = await GetAsync(key, memory => LogIndexData.FromMemory(memory));
+			if(indexData == null)
+			{
+				indexData = await _inner.ReadIndexAsync(logId, length);
+				if(indexData != null)
+				{
+					await AddAsync(key, indexData.ToByteArray());
+				}
+			}
+			return indexData;
+		}
+
+		/// <inheritdoc/>
+		public async Task WriteIndexAsync(LogId logId, long length, LogIndexData index)
+		{
+			await _inner.WriteIndexAsync(logId, length, index);
+		}
+
+		/// <inheritdoc/>
+		public async Task<LogChunkData?> ReadChunkAsync(LogId logId, long offset, int lineIndex)
+		{
+			string key = ChunkKey(logId, offset);
+
+			LogChunkData? chunkData = await GetAsync(key, memory => LogChunkData.FromMemory(memory, offset, lineIndex));
+			if (chunkData == null)
+			{
+				chunkData = await _inner.ReadChunkAsync(logId, offset, lineIndex);
+				if (chunkData != null)
+				{
+					await AddAsync(key, chunkData.ToByteArray(_logger));
+				}
+			}
+			return chunkData;
+		}
+
+		/// <inheritdoc/>
+		public async Task WriteChunkAsync(LogId logId, long offset, LogChunkData chunkData)
+		{
+			await _inner.WriteChunkAsync(logId, offset, chunkData);
+		}
+	}
+}
diff --git a/Engine/Source/Programs/Horde/Horde.Server/Logs/Storage/SequencedLogStorage.cs b/Engine/Source/Programs/Horde/Horde.Server/Logs/Storage/SequencedLogStorage.cs
new file mode 100644
--- /dev/null
+++ b/Engine/Source/Programs/Horde/Horde.Server/Logs/Storage/SequencedLogStorage.cs
@@ -0,0 +1,191 @@
+// Copyright Epic Games, Inc. All Rights Reserved.
+
+using System.Collections.Generic;
+using System.Threading.Tasks;
+using Horde.Server.Logs.Data;
+using Horde.Server.Utilities;
+
+namespace Horde.Server.Logs.Storage
+{
+	/// <summary>
+	/// Storage layer which caches pending read tasks, to avoid fetching the same item more than once
+	/// </summary>
+	class SequencedLogStorage : ILogStorage
+	{
+		/// <summary>
+		/// Inner storage implementation
+		/// </summary>
+		readonly ILogStorage _inner;
+
+		/// <summary>
+		/// Pending index reads
+		/// </summary>
+		readonly Dictionary<(LogId, long), Task<LogIndexData?>> _indexReadTasks = new Dictionary<(LogId, long), Task<LogIndexData?>>();
+
+		/// <summary>
+		/// Pending index reads
+		/// </summary>
+		readonly Dictionary<(LogId, long), Task> _indexWriteTasks = new Dictionary<(LogId, long), Task>();
+
+		/// <summary>
+		/// Pending chunk reads
+		/// </summary>
+		readonly Dictionary<(LogId, long), Task<LogChunkData?>> _chunkReadTasks = new Dictionary<(LogId, long), Task<LogChunkData?>>();
+
+		/// <summary>
+		/// Pending chunk reads
+		/// </summary>
+		readonly Dictionary<(LogId, long), Task> _chunkWriteTasks = new Dictionary<(LogId, long), Task>();
+
+		/// <summary>
+		/// Constructor
+		/// </summary>
+		/// <param name="inner">The inner storage provider</param>
+		public SequencedLogStorage(ILogStorage inner)
+		{
+			_inner = inner;
+		}
+
+		/// <inheritdoc/>
+		public void Dispose()
+		{
+			_inner.Dispose();
+		}
+
+		/// <inheritdoc/>
+		public Task<LogIndexData?> ReadIndexAsync(LogId logId, long length)
+		{
+			Task<LogIndexData?>? task;
+			lock (_indexReadTasks)
+			{
+				if (!_indexReadTasks.TryGetValue((logId, length), out task))
+				{
+					task = InnerReadIndexAsync(logId, length);
+					_indexReadTasks.Add((logId, length), task);
+				}
+			}
+			return task;
+		}
+
+		/// <inheritdoc/>
+		public Task WriteIndexAsync(LogId logId, long length, LogIndexData indexData)
+		{
+			Task? task;
+			lock (_indexWriteTasks)
+			{
+				if (!_indexWriteTasks.TryGetValue((logId, length), out task))
+				{
+					task = InnerWriteIndexAsync(logId, length, indexData);
+					_indexWriteTasks.Add((logId, length), task);
+				}
+			}
+			return task;
+		}
+
+		/// <inheritdoc/>
+		public Task<LogChunkData?> ReadChunkAsync(LogId logId, long offset, int lineIndex)
+		{
+			Task<LogChunkData?>? task;
+			lock (_chunkReadTasks)
+			{
+				if (!_chunkReadTasks.TryGetValue((logId, offset), out task))
+				{
+					task = InnerReadChunkAsync(logId, offset, lineIndex);
+					_chunkReadTasks[(logId, offset)] = task;
+				}
+			}
+			return task;
+		}
+
+		/// <inheritdoc/>
+		public Task WriteChunkAsync(LogId logId, long offset, LogChunkData chunkData)
+		{
+			Task? task;
+			lock (_chunkWriteTasks)
+			{
+				if (!_chunkWriteTasks.TryGetValue((logId, offset), out task))
+				{
+					task = InnerWriteChunkAsync(logId, offset, chunkData);
+					_chunkWriteTasks[(logId, offset)] = task;
+				}
+			}
+			return task;
+		}
+
+		/// <summary>
+		/// Wrapper for reading an index from the inner storage provider
+		/// </summary>
+		/// <param name="logId">The log file to read the index for</param>
+		/// <param name="length">Length of the file that's indexed</param>
+		/// <returns>The index data</returns>
+		async Task<LogIndexData?> InnerReadIndexAsync(LogId logId, long length)
+		{
+			await Task.Yield();
+
+			LogIndexData? indexData = await _inner.ReadIndexAsync(logId, length);
+			lock (_indexReadTasks)
+			{
+				_indexReadTasks.Remove((logId, length));
+			}
+			return indexData;
+		}
+
+		/// <summary>
+		/// Wrapper for reading an index from the inner storage provider
+		/// </summary>
+		/// <param name="logId">The log file to read the index for</param>
+		/// <param name="length">Length of the indexed data</param>
+		/// <param name="indexData">The index data to write</param>
+		/// <returns>The index data</returns>
+		async Task InnerWriteIndexAsync(LogId logId, long length, LogIndexData indexData)
+		{
+			await Task.Yield();
+
+			await _inner.WriteIndexAsync(logId, length, indexData);
+
+			lock (_indexWriteTasks)
+			{
+				_indexWriteTasks.Remove((logId, length));
+			}
+		}
+
+		/// <summary>
+		/// Wrapper for reading a chunk from the inner storage provider
+		/// </summary>
+		/// <param name="logId">The log file to read the index for</param>
+		/// <param name="offset">Offset of the chunk to read</param>
+		/// <param name="lineIndex">Index of the first line in this chunk</param>
+		/// <returns>The index data</returns>
+		async Task<LogChunkData?> InnerReadChunkAsync(LogId logId, long offset, int lineIndex)
+		{
+			await Task.Yield();
+
+			LogChunkData? chunkData = await _inner.ReadChunkAsync(logId, offset, lineIndex);
+			lock (_chunkReadTasks)
+			{
+				_chunkReadTasks.Remove((logId, offset));
+			}
+			return chunkData;
+		}
+
+		/// <summary>
+		/// Wrapper for reading a chunk from the inner storage provider
+		/// </summary>
+		/// <param name="logId">The log file to write the chunk for</param>
+		/// <param name="offset">Offset of the chunk within the log</param>
+		/// <param name="chunkData">The chunk data</param>
+		/// <returns>The index data</returns>
+		async Task<LogChunkData> InnerWriteChunkAsync(LogId logId, long offset, LogChunkData chunkData)
+		{
+			await Task.Yield();
+
+			await _inner.WriteChunkAsync(logId, offset, chunkData);
+
+			lock (_chunkWriteTasks)
+			{
+				_chunkWriteTasks.Remove((logId, offset));
+			}
+			return chunkData;
+		}
+	}
+}
diff --git a/Engine/Source/Programs/Horde/Horde.Server/Notifications/INotificationSink.cs b/Engine/Source/Programs/Horde/Horde.Server/Notifications/INotificationSink.cs
new file mode 100644
--- /dev/null
+++ b/Engine/Source/Programs/Horde/Horde.Server/Notifications/INotificationSink.cs
@@ -0,0 +1,123 @@
+// Copyright Epic Games, Inc. All Rights Reserved.
+
+using System;
+using System.Collections.Generic;
+using System.Threading.Tasks;
+using Horde.Server.Configuration;
+using Horde.Server.Devices;
+using Horde.Server.Issues;
+using Horde.Server.Jobs;
+using Horde.Server.Jobs.Graphs;
+using Horde.Server.Logs;
+using Horde.Server.Streams;
+using Horde.Server.Users;
+using HordeCommon;
+
+namespace Horde.Server.Notifications
+{
+	/// <summary>
+	/// Implements a notification method
+	/// </summary>
+	public interface INotificationSink
+	{
+		/// <summary>
+		/// Send notifications that a job has been scheduled
+		/// </summary>
+		/// <param name="notifications">List of notifications to send</param>
+		/// <returns>Async task</returns>
+		Task NotifyJobScheduledAsync(List<JobScheduledNotification> notifications);
+		
+		/// <summary>
+		/// Send notifications that a job has completed
+		/// </summary>
+		/// <param name="job">The job containing the step</param>
+		/// <param name="graph"></param>
+		/// <param name="outcome"></param>
+		/// <returns>Async task</returns>
+		Task NotifyJobCompleteAsync(IJob job, IGraph graph, LabelOutcome outcome);
+
+		/// <summary>
+		/// Send notifications that a job has completed
+		/// </summary>
+		/// <param name="user">User to notify</param>
+		/// <param name="job">The job containing the step</param>
+		/// <param name="graph"></param>
+		/// <param name="outcome"></param>
+		/// <returns>Async task</returns>
+		Task NotifyJobCompleteAsync(IUser user, IJob job, IGraph graph, LabelOutcome outcome);
+
+		/// <summary>
+		/// Send notifications that a job step has completed
+		/// </summary>
+		/// <param name="user">User to notify</param>
+		/// <param name="job">The job containing the step</param>
+		/// <param name="batch">Unique id of the batch</param>
+		/// <param name="step">The step id</param>
+		/// <param name="node">Corresponding node for the step</param>
+		/// <param name="jobStepEventData"></param>
+		/// <returns>Async task</returns>
+		Task NotifyJobStepCompleteAsync(IUser user, IJob job, IJobStepBatch batch, IJobStep step, INode node, List<ILogEventData> jobStepEventData);
+
+		/// <summary>
+		/// Send notifications that a job step has completed
+		/// </summary>
+		/// <param name="user">User to notify</param>
+		/// <param name="job">The job containing the step</param>
+		/// <param name="label"></param>
+		/// <param name="labelIdx"></param>
+		/// <param name="outcome"></param>
+		/// <param name="stepData"></param>
+		/// <returns>Async task</returns>
+		Task NotifyLabelCompleteAsync(IUser user, IJob job, ILabel label, int labelIdx, LabelOutcome outcome, List<(string, JobStepOutcome, Uri)> stepData);
+
+		#region Issues
+
+		/// <summary>
+		/// Send notifications that a new issue has been created or assigned
+		/// </summary>
+		/// <param name="issue"></param>
+		/// <returns></returns>
+		Task NotifyIssueUpdatedAsync(IIssue issue);
+
+		/// <summary>
+		/// Notify the status of issues in a stream
+		/// </summary>
+		/// <param name="report"></param>
+		/// <returns></returns>
+		Task SendIssueReportAsync(IssueReportGroup report);
+
+		#endregion
+
+		/// <summary>
+		/// Notification that the configuration state has changed
+		/// </summary>
+		/// <param name="ex">Exception during update. Null if the update completed successfully.</param>
+		Task NotifyConfigUpdateAsync(Exception? ex);
+
+		/// <summary>
+		/// Notification that a stream has failed to update
+		/// </summary>
+		/// <param name="errorMessage"></param>
+		/// <param name="fileName"></param>
+		/// <param name="change"></param>
+		/// <param name="author"></param>
+		/// <param name="description"></param>
+		/// <returns></returns>
+		Task NotifyConfigUpdateFailureAsync(string errorMessage, string fileName, int? change = null, IUser? author = null, string? description = null);
+
+		/// <summary>
+		/// Notification for device service
+		/// </summary>
+		/// <param name="message"></param>
+		/// <param name="device"></param>
+		/// <param name="pool"></param>
+		/// <param name="streamConfig"></param>
+		/// <param name="job"></param>
+		/// <param name="step"></param>
+		/// <param name="node"></param>
+		/// <param name="user"></param>
+		/// <returns></returns>
+		Task NotifyDeviceServiceAsync(string message, IDevice? device = null, IDevicePool? pool = null, StreamConfig? streamConfig = null, IJob? job = null, IJobStep? step = null, INode? node = null, IUser? user = null);
+    }
+}
+
diff --git a/Engine/Source/Programs/Horde/Horde.Server/Notifications/INotificationTriggerCollection.cs b/Engine/Source/Programs/Horde/Horde.Server/Notifications/INotificationTriggerCollection.cs
new file mode 100644
--- /dev/null
+++ b/Engine/Source/Programs/Horde/Horde.Server/Notifications/INotificationTriggerCollection.cs
@@ -0,0 +1,61 @@
+// Copyright Epic Games, Inc. All Rights Reserved.
+
+using System.Collections.Generic;
+using System.Threading.Tasks;
+using Horde.Server.Users;
+using Horde.Server.Utilities;
+using MongoDB.Bson;
+
+namespace Horde.Server.Notifications
+{
+	/// <summary>
+	/// Collection of notification triggers
+	/// </summary>
+	public interface INotificationTriggerCollection
+	{
+		/// <summary>
+		/// Finds or adds a trigger
+		/// </summary>
+		/// <param name="triggerId">The trigger id</param>
+		/// <returns>New trigger document</returns>
+		Task<INotificationTrigger> FindOrAddAsync(ObjectId triggerId);
+
+		/// <summary>
+		/// Finds an existing trigger with the given id, or adds one if it does not exist
+		/// </summary>
+		/// <param name="triggerId"></param>
+		/// <returns></returns>
+		Task<INotificationTrigger?> GetAsync(ObjectId triggerId);
+
+		/// <summary>
+		/// Deletes a trigger
+		/// </summary>
+		/// <param name="triggerId">The unique trigger id</param>
+		/// <returns>Async task</returns>
+		Task DeleteAsync(ObjectId triggerId);
+
+		/// <summary>
+		/// Deletes a trigger
+		/// </summary>
+		/// <param name="triggerIds">The unique trigger id</param>
+		/// <returns>Async task</returns>
+		Task DeleteAsync(List<ObjectId> triggerIds);
+
+		/// <summary>
+		/// Fires the trigger, and marks notifications has having been sent
+		/// </summary>
+		/// <param name="trigger">The trigger to fire</param>
+		/// <returns>The trigger document</returns>
+		Task<INotificationTrigger?> FireAsync(INotificationTrigger trigger);
+
+		/// <summary>
+		/// Adds a subscriber to a particular trigger
+		/// </summary>
+		/// <param name="trigger">The trigger to subscribe to</param>
+		/// <param name="userId">The user name</param>
+		/// <param name="email">Whether to receive email notifications</param>
+		/// <param name="slack">Whether to receive Slack notifications</param>
+		/// <returns>The new trigger state</returns>
+		Task<INotificationTrigger?> UpdateSubscriptionsAsync(INotificationTrigger trigger, UserId userId, bool? email, bool? slack);
+	}
+}
diff --git a/Engine/Source/Programs/Horde/Horde.Server/Notifications/NotificationService.cs b/Engine/Source/Programs/Horde/Horde.Server/Notifications/NotificationService.cs
new file mode 100644
--- /dev/null
+++ b/Engine/Source/Programs/Horde/Horde.Server/Notifications/NotificationService.cs
@@ -0,0 +1,790 @@
+// Copyright Epic Games, Inc. All Rights Reserved.
+
+using System;
+using System.Collections.Concurrent;
+using System.Collections.Generic;
+using System.Linq;
+using System.Security.Claims;
+using System.Threading;
+using System.Threading.Tasks;
+using EpicGames.Core;
+using EpicGames.Redis;
+using Horde.Server.Agents.Pools;
+using Horde.Server.Configuration;
+using Horde.Server.Devices;
+using Horde.Server.Issues;
+using Horde.Server.Jobs;
+using Horde.Server.Jobs.Graphs;
+using Horde.Server.Logs;
+using Horde.Server.Server;
+using Horde.Server.Streams;
+using Horde.Server.Users;
+using Horde.Server.Utilities;
+using HordeCommon;
+using Microsoft.Extensions.Caching.Memory;
+using Microsoft.Extensions.Hosting;
+using Microsoft.Extensions.Logging;
+using Microsoft.Extensions.Options;
+using MongoDB.Bson;
+using MongoDB.Driver;
+using StackExchange.Redis;
+using StatsdClient;
+using JsonSerializer = System.Text.Json.JsonSerializer;
+
+namespace Horde.Server.Notifications
+{
+	/// <summary>
+	/// Wraps functionality for delivering notifications.
+	/// </summary>
+	public class NotificationService : BackgroundService, INotificationService
+	{
+		/// <summary>
+		/// The available notification sinks
+		/// </summary>
+		private readonly List<INotificationSink> _sinks;
+
+		/// <summary>
+		/// Collection of subscriptions
+		/// </summary>
+		private readonly ISubscriptionCollection _subscriptionCollection;
+
+		/// <summary>
+		/// Collection of notification request documents.
+		/// </summary>
+		private readonly INotificationTriggerCollection _triggerCollection;
+
+		/// <summary>
+		/// Instance of the <see cref="_graphCollection"/>.
+		/// </summary>
+		private readonly IGraphCollection _graphCollection;
+
+		/// <summary>
+		/// 
+		/// </summary>
+		private readonly IUserCollection _userCollection;
+
+		/// <summary>
+		/// Job service instance
+		/// </summary>
+		private readonly JobService _jobService;
+
+		/// <summary>
+		/// Instance of the <see cref="IStreamCollection"/>.
+		/// </summary>
+		private readonly IStreamCollection _streamCollection;
+
+		/// <summary>
+		/// 
+		/// </summary>
+		private readonly IssueService _issueService;
+
+		/// <summary>
+		/// Instance of the <see cref="_logFileService"/>.
+		/// </summary>
+		private readonly ILogFileService _logFileService;
+		
+		/// <summary>
+		/// Instance of the <see cref="IDogStatsd"/>.
+		/// </summary>
+		private readonly IDogStatsd _dogStatsd;
+		
+		/// <summary>
+		/// Cache for de-duplicating queued notifications
+		/// </summary>
+		private readonly IMemoryCache _cache;
+		
+		/// <summary>
+		/// Lock object for manipulating the above cache
+		/// Used since batch notification queue handling is run async.
+		/// </summary>
+		private readonly object _cacheLock = new object();
+
+		/// <summary>
+		/// Connection pool for Redis databases
+		/// </summary>
+		private readonly RedisConnectionPool _redisConnectionPool;
+		
+		/// <summary>
+		/// Settings for the application.
+		/// </summary>
+		private readonly IOptionsMonitor<ServerSettings> _settings;
+
+		/// <summary>
+		/// List of asychronous tasks currently executing
+		/// </summary>
+		private readonly ConcurrentQueue<Task> _newTasks = new ConcurrentQueue<Task>();
+
+		/// <summary>
+		/// Set when there are new tasks to wait for
+		/// </summary>
+		private AsyncEvent _newTaskEvent = new AsyncEvent();
+
+		/// <summary>
+		/// Settings for the application.
+		/// </summary>
+		private readonly ILogger<NotificationService> _logger;
+		
+		/// <summary>
+		/// Ticker for running batch sender method
+		/// </summary>
+		internal ITicker _ticker;
+
+		/// <summary>
+		/// Interval at which queued notifications should be sent as a batch 
+		/// </summary>
+		internal TimeSpan _notificationBatchInterval = TimeSpan.FromHours(12);
+
+		readonly IOptionsMonitor<GlobalConfig> _globalConfig;
+
+		static string RedisQueueListKey(string notificationType) => "NotificationService.queued." + notificationType;
+
+		/// <summary>
+		/// Constructor
+		/// </summary>
+		public NotificationService(
+			IEnumerable<INotificationSink> sinks,
+			IOptionsMonitor<ServerSettings> settings,
+			ILogger<NotificationService> logger,
+			IGraphCollection graphCollection,
+			ISubscriptionCollection subscriptionCollection,
+			INotificationTriggerCollection triggerCollection,
+			IUserCollection userCollection,
+			JobService jobService,
+			IStreamCollection streamCollection,
+			IssueService issueService,
+			ILogFileService logFileService,
+			IDogStatsd dogStatsd,
+			IMemoryCache cache,
+			RedisService redisService,
+			ConfigService configService,
+			IOptionsMonitor<GlobalConfig> globalConfig,
+			IClock clock)
+		{
+			_sinks = sinks.ToList();
+			_settings = settings;
+			_logger = logger;
+			_graphCollection = graphCollection;
+			_subscriptionCollection = subscriptionCollection;
+			_triggerCollection = triggerCollection;
+			_userCollection = userCollection;
+			_jobService = jobService;
+			_streamCollection = streamCollection;
+			_issueService = issueService;
+			_logFileService = logFileService;
+			_dogStatsd = dogStatsd;
+			_cache = cache;
+			_redisConnectionPool = redisService.ConnectionPool;
+			_globalConfig = globalConfig;
+
+			issueService.OnIssueUpdated += NotifyIssueUpdated;
+			jobService.OnJobStepComplete += NotifyJobStepComplete;
+			jobService.OnJobScheduled += NotifyJobScheduled;
+			jobService.OnLabelUpdate += NotifyLabelUpdate;
+			configService.OnConfigUpdate += NotifyConfigUpdate;
+
+			_ticker = clock.AddSharedTicker<NotificationService>(_notificationBatchInterval, TickEveryTwelveHoursAsync, logger);
+		}
+
+		/// <inheritdoc/>
+		public override void Dispose()
+		{
+			base.Dispose();
+
+			_issueService.OnIssueUpdated -= NotifyIssueUpdated;
+			_jobService.OnJobStepComplete -= NotifyJobStepComplete;
+			_jobService.OnJobScheduled += NotifyJobScheduled;
+			_jobService.OnLabelUpdate -= NotifyLabelUpdate;
+
+			GC.SuppressFinalize(this);
+			_ticker.Dispose();
+		}
+
+		/// <inheritdoc/>
+		public async Task<bool> UpdateSubscriptionsAsync(ObjectId triggerId, ClaimsPrincipal user, bool? email, bool? slack)
+		{
+			UserId? userId = user.GetUserId();
+			if (userId == null)
+			{
+				_logger.LogWarning("Unable to find username for principal {User}", user.Identity?.Name);
+				return false;
+			}
+
+			INotificationTrigger trigger = await _triggerCollection.FindOrAddAsync(triggerId);
+			await _triggerCollection.UpdateSubscriptionsAsync(trigger, userId.Value, email, slack);
+			return true;
+		}
+
+		/// <inheritdoc/>
+		public async Task<INotificationSubscription?> GetSubscriptionsAsync(ObjectId triggerId, ClaimsPrincipal user)
+		{
+			UserId? userId = user.GetUserId();
+			if (userId == null)
+			{
+				return null;
+			}
+
+			INotificationTrigger? trigger = await _triggerCollection.GetAsync(triggerId);
+			if(trigger == null)
+			{
+				return null;
+			}
+
+			return trigger.Subscriptions.FirstOrDefault(x => x.UserId == userId.Value);
+		}
+
+		/// <inheritdoc/>
+		public void NotifyJobStepComplete(IJob job, IGraph graph, SubResourceId batchId, SubResourceId stepId)
+		{
+			// Enqueue job step complete notifications if needed
+			if (job.TryGetStep(batchId, stepId, out IJobStep? step))
+			{
+				_logger.LogInformation("Queuing step notifications for {JobId}:{BatchId}:{StepId}", job.Id, batchId, stepId);
+				EnqueueTask(() => SendJobStepNotificationsAsync(job, batchId, stepId));
+			}
+
+			// Enqueue job complete notifications if needed
+			if (job.GetState() == JobState.Complete)
+			{
+				_logger.LogInformation("Queuing job notifications for {JobId}:{BatchId}:{StepId}", job.Id, batchId, stepId);
+				EnqueueTask(() => SendJobNotificationsAsync(job, graph));
+				EnqueueTask(() => RecordJobCompleteMetrics(job));
+			}
+		}
+		
+		/// <inheritdoc/>
+		public void NotifyJobScheduled(IPool pool, bool poolHasAgentsOnline, IJob job, IGraph graph, SubResourceId batchId)
+		{
+			if (pool.EnableAutoscaling && !poolHasAgentsOnline)
+			{
+				EnqueueTasks(sink => EnqueueNotificationForBatchSending(new JobScheduledNotification(job.Id.ToString(), job.Name, pool.Name)));
+			}
+		}
+
+		/// <inheritdoc/>
+		public void NotifyLabelUpdate(IJob job, IReadOnlyList<(LabelState, LabelOutcome)> oldLabelStates, IReadOnlyList<(LabelState, LabelOutcome)> newLabelStates)
+		{
+			// If job has any label trigger IDs, send label complete notifications if needed
+			for (int idx = 0; idx < oldLabelStates.Count && idx < newLabelStates.Count; idx++)
+			{
+				if (oldLabelStates[idx] != newLabelStates[idx])
+				{
+					EnqueueTask(() => SendAllLabelNotificationsAsync(job, oldLabelStates, newLabelStates));
+					break;
+				}
+			}
+		}
+
+		/// <inheritdoc/>
+		public void NotifyIssueUpdated(IIssue issue)
+		{
+			_logger.LogInformation("Issue {IssueId} updated", issue.Id);
+			EnqueueTasks(sink => sink.NotifyIssueUpdatedAsync(issue));
+		}
+
+		/// <inheritdoc/>
+		public void NotifyConfigUpdate(Exception? ex)
+		{
+			_logger.LogInformation(ex, "Configuration updated ({Result})", (ex == null) ? "success" : "failure");
+			EnqueueTasks(sink => sink.NotifyConfigUpdateAsync(ex));
+		}
+
+		/// <inheritdoc/>
+		public void NotifyConfigUpdateFailure(string errorMessage, string fileName, int? change = null, IUser? author = null, string? description = null)
+		{
+			EnqueueTasks(sink => sink.NotifyConfigUpdateFailureAsync(errorMessage, fileName, change, author, description));
+		}
+
+		/// <inheritdoc/>
+		public void NotifyDeviceService(string message, IDevice? device = null, IDevicePool? pool = null, StreamConfig? streamConfig = null, IJob? job = null, IJobStep? step = null, INode? node = null, IUser? user = null)
+		{
+			EnqueueTasks(sink => sink.NotifyDeviceServiceAsync(message, device, pool, streamConfig, job, step, node, user));
+		}
+
+		/// <summary>
+		/// Enqueues an async task
+		/// </summary>
+		/// <param name="taskFunc">Function to generate an async task</param>
+		void EnqueueTask(Func<Task> taskFunc)
+		{
+			_newTasks.Enqueue(Task.Run(taskFunc));
+		}
+
+		/// <summary>
+		/// Enqueues an async task
+		/// </summary>
+		/// <param name="taskFunc">Function to generate an async task</param>
+		void EnqueueTasks(Func<INotificationSink, Task> taskFunc)
+		{
+			foreach (INotificationSink sink in _sinks)
+			{
+				EnqueueTask(() => taskFunc(sink));
+			}
+		}
+
+		/// <summary>
+		/// Enqueue a notification in Redis for batch sending later on 
+		/// </summary>
+		/// <param name="notification">Notification to enqueue</param>
+		/// <param name="deduplicate">True if notification should be deduplicated</param>
+		/// <typeparam name="T">Any INotification type</typeparam>
+		private async Task EnqueueNotificationForBatchSending<T>(T notification, bool deduplicate = true) where T : INotification<T>
+		{
+			lock (_cacheLock)
+			{
+				if (deduplicate)
+				{
+					// Use cache to deduplicate notifications
+					if (_cache.TryGetValue(notification, out object? _))
+					{
+						return;
+					}
+					_cache.Set(notification, notification, _notificationBatchInterval / 2);					
+				}
+			}
+
+			try
+			{
+				byte[] data = JsonSerializer.SerializeToUtf8Bytes(notification);
+				await _redisConnectionPool.GetDatabase().ListRightPushAsync(RedisQueueListKey(typeof(T).ToString()), data);
+			}
+			catch (Exception e)
+			{
+				_logger.LogError(e, "Unable to serialize and queue notification {Type} for batch sending in Redis", notification.GetType());
+			}
+		}
+
+		private async ValueTask TickEveryTwelveHoursAsync(CancellationToken stoppingToken)
+		{
+			List<JobScheduledNotification> jobScheduledNotifications = await GetAllQueuedNotificationsAsync<JobScheduledNotification>();
+			if (jobScheduledNotifications.Count > 0)
+			{
+				foreach (INotificationSink sink in _sinks)
+				{
+					await sink.NotifyJobScheduledAsync(jobScheduledNotifications);
+				}
+			}
+		}
+		
+		/// <summary>
+		/// Get and clear all queued notifications of type T from Redis
+		/// </summary>
+		/// <typeparam name="T"></typeparam>
+		/// <returns>Deserialized notifications</returns>
+		/// <exception cref="Exception"></exception>
+		private async Task<List<T>> GetAllQueuedNotificationsAsync<T>() where T : INotification<T>
+		{
+			IDatabase redis = _redisConnectionPool.GetDatabase();
+
+			// Reading and deleting the entire list from Redis in this way is not thread-safe.
+			// But likely not a big deal considering the alternative of distributed locks.
+			string redisKey = RedisQueueListKey(typeof(T).ToString());
+			RedisValue[] rawNotifications = await redis.ListRangeAsync(redisKey);
+			await redis.KeyDeleteAsync(redisKey);
+
+			List<T> notifications = new List<T>();
+			foreach (byte[]? data in rawNotifications)
+			{
+				try
+				{
+					T? notification = JsonSerializer.Deserialize<T>(data!);
+					if (notification == null)
+					{
+						throw new Exception("Unable to deserialize");
+					}
+					notifications.Add(notification);
+				}
+				catch (Exception e)
+				{
+					_logger.LogError(e, "Unable to deserialize notification data of {Type}", typeof(T));
+				}
+			}
+
+			return notifications;
+		}
+
+		/// <inheritdoc/>
+		protected override async Task ExecuteAsync(CancellationToken stoppingToken)
+		{
+			await _ticker.StartAsync();
+			
+			// This background service just waits for tasks to finish and prints any exception info. The only reason to do this is to
+			// ensure we finish processing everything before shutdown.
+			using (CancellationTask stoppingTask = new CancellationTask(stoppingToken))
+			{
+				List<Task> tasks = new List<Task>();
+				tasks.Add(_newTaskEvent.Task);
+
+				for (; ; )
+				{
+					// Add any new tasks to be monitored
+					Task? newTask;
+					while (_newTasks.TryDequeue(out newTask))
+					{
+						tasks.Add(newTask);
+					}
+
+					// If we don't have any
+					if (tasks.Count == 1)
+					{
+						await Task.WhenAny(_newTaskEvent.Task, stoppingTask.Task);
+						if (stoppingToken.IsCancellationRequested)
+						{
+							break;
+						}
+					}
+					else
+					{
+						// Wait for something to finish
+						Task task = await Task.WhenAny(tasks);
+						if (task == _newTaskEvent.Task)
+						{
+							_newTaskEvent = new AsyncEvent();
+							tasks[0] = _newTaskEvent.Task;
+						}
+						else
+						{
+							try
+							{
+								await task;
+							}
+							catch (Exception ex)
+							{
+								_logger.LogError(ex, "Exception while executing notification");
+							}
+							tasks.Remove(task);
+						}
+					}
+				}
+			}
+			
+			await _ticker.StopAsync();
+		}
+
+		internal Task ExecuteBackgroundForTest(CancellationToken stoppingToken)
+		{
+			return ExecuteAsync(stoppingToken);
+		}
+
+		/// <summary>
+		/// Gets the <see cref="INotificationTrigger"/> for a given trigger ID, if any.
+		/// </summary>
+		/// <param name="triggerId"></param>
+		/// <param name="fireTrigger">If true, the trigger is fired and cannot be reused</param>
+		/// <returns></returns>
+		private async Task<INotificationTrigger?> GetNotificationTrigger(ObjectId? triggerId, bool fireTrigger)
+		{
+			if (triggerId == null)
+			{
+				return null;
+			}
+
+			INotificationTrigger? trigger = await _triggerCollection.GetAsync(triggerId.Value);
+			if (trigger == null)
+			{
+				return null;
+			}
+
+			return fireTrigger ? await _triggerCollection.FireAsync(trigger) : trigger;
+		}
+	
+		private async Task SendJobNotificationsAsync(IJob job, IGraph graph)
+		{
+			using IDisposable scope = _logger.BeginScope("Sending notifications for job {JobId}", job.Id);
+
+			job.GetJobState(job.GetStepForNodeMap(), out _, out LabelOutcome outcome);
+			JobCompleteEventRecord jobCompleteEvent = new JobCompleteEventRecord(job.StreamId, job.TemplateId, outcome);
+
+			List<IUser> usersToNotify = await GetUsersToNotify(jobCompleteEvent, job.NotificationTriggerId, true);
+			foreach (IUser userToNotify in usersToNotify)
+			{
+				if(job.PreflightChange != 0)
+				{
+					if(userToNotify.Id != job.StartedByUserId)
+					{
+						continue;
+					}
+				}
+				EnqueueTasks(sink => sink.NotifyJobCompleteAsync(userToNotify, job, graph, outcome));
+			}
+
+			if (job.PreflightChange == 0)
+			{
+				EnqueueTasks(sink => sink.NotifyJobCompleteAsync(job, graph, outcome));
+			}
+
+			_logger.LogDebug("Finished sending notifications for job {JobId}", job.Id);
+		}
+
+		private Task RecordJobCompleteMetrics(IJob job)
+		{
+			void RecordMetric(string type, JobStepOutcome outcome, DateTimeOffset? startTime, DateTimeOffset? finishTime)
+			{
+				string outcomeStr = outcome switch
+				{
+					JobStepOutcome.Unspecified => "unspecified",
+					JobStepOutcome.Failure => "failure",
+					JobStepOutcome.Warnings => "warnings",
+					JobStepOutcome.Success => "success",
+					_ => "unspecified"
+				};
+				
+				string[] tags = {"stream:" + job.StreamId, "template:" + job.TemplateId};
+				_dogStatsd.Increment($"horde.{type}.{outcomeStr}.count", 1, tags: tags);
+
+				if (startTime == null || finishTime == null)
+				{
+					_logger.LogDebug("Completed job or step is missing start or finish time, cannot record duration metric. Job ID={JobId}", job.Id);
+					return;
+				}
+
+				TimeSpan duration = finishTime.Value - startTime.Value;
+				_dogStatsd.Timer($"horde.{type}.{outcomeStr}.duration", duration.TotalSeconds, tags: tags);
+			}
+
+			JobStepOutcome jobOutcome = job.Batches.SelectMany(x => x.Steps).Min(x => x.Outcome);
+			DateTime? startTime = job.Batches.Select(x => x.StartTimeUtc).Min();
+			DateTime? finishTime = job.Batches.Select(x => x.FinishTimeUtc).Max();
+			RecordMetric("job", jobOutcome, startTime, finishTime);
+
+			// TODO: record metrics for individual steps
+			// foreach (IJobStepBatch Batch in Job.Batches)
+			// {
+			// 	foreach (IJobStep Step in Batch.Steps)
+			// 	{
+			// 	}
+			// }
+
+			return Task.CompletedTask;
+		}
+
+		private async Task<List<IUser>> GetUsersToNotify(EventRecord? eventRecord, ObjectId? notificationTriggerId, bool fireTrigger)
+		{
+			List<UserId> userIds = new List<UserId>();
+
+			// Find the notifications for all steps of this type
+			if (eventRecord != null)
+			{
+				List<ISubscription> subscriptions = await _subscriptionCollection.FindSubscribersAsync(eventRecord);
+				foreach (ISubscription subscription in subscriptions)
+				{
+					if (subscription.NotificationType == NotificationType.Slack)
+					{
+						userIds.Add(subscription.UserId);
+					}
+				}
+			}
+
+			// Find the notifications for this particular step
+			if (notificationTriggerId != null)
+			{
+				INotificationTrigger? trigger = await GetNotificationTrigger(notificationTriggerId, fireTrigger);
+				if (trigger != null)
+				{
+					foreach (INotificationSubscription subscription in trigger.Subscriptions)
+					{
+						if (subscription.Email)
+						{
+							// TODO?
+						}
+						if (subscription.Slack)
+						{
+							userIds.Add(subscription.UserId);
+						}
+					}
+				}
+			}
+			return await _userCollection.FindUsersAsync(userIds);
+		}
+
+		private async Task SendJobStepNotificationsAsync(IJob job, SubResourceId batchId, SubResourceId stepId)
+		{
+			using IDisposable scope = _logger.BeginScope("Sending notifications for step {JobId}:{BatchId}:{StepId}", job.Id, batchId, stepId);
+
+			IJobStepBatch? batch;
+			if(!job.TryGetBatch(batchId, out batch))
+			{
+				_logger.LogError("Unable to find batch {BatchId} in job {JobId}", batchId, job.Id);
+				return;
+			}
+
+			IJobStep? step;
+			if (!batch.TryGetStep(stepId, out step))
+			{
+				_logger.LogError("Unable to find step {StepId} in batch {JobId}:{BatchId}", stepId, job.Id, batchId);
+				return;
+			}
+
+			IGraph jobGraph = await _graphCollection.GetAsync(job.GraphHash);
+			INode node = jobGraph.GetNode(new NodeRef(batch.GroupIdx, step.NodeIdx));
+
+			// Find the notifications for this particular step
+			EventRecord eventRecord = new StepCompleteEventRecord(job.StreamId, job.TemplateId, node.Name, step.Outcome);
+
+			List<IUser> usersToNotify = await GetUsersToNotify(eventRecord, step.NotificationTriggerId, true);
+
+			// If this is not a success notification and the author isn't in the list to notify, add them manually if this is the outcome has gotten worse.
+			int failures = job.Batches.Sum(x => x.Steps.Count(y => y.Outcome == JobStepOutcome.Failure));
+			bool firstFailure = step.Outcome == JobStepOutcome.Failure && failures == 1;
+			bool firstWarning = step.Outcome == JobStepOutcome.Warnings && failures == 0 && job.Batches.Sum(x => x.Steps.Count(y => y.Outcome == JobStepOutcome.Warnings)) == 1;
+			if (job.StartedByUserId.HasValue && !usersToNotify.Any(x => x.Id == job.StartedByUserId) && (firstFailure || firstWarning))
+			{
+				_logger.LogInformation("Author {AuthorUserId} is not in notify list but step outcome is {JobStepOutcome}, adding them to the list...", job.StartedByUserId, step.Outcome);
+				IUser? authorUser = await _userCollection.GetUserAsync(job.StartedByUserId.Value);
+				if (authorUser != null)
+				{
+					usersToNotify.Add(authorUser);
+				}
+			}
+
+			if (usersToNotify.Count == 0)
+			{
+				_logger.LogInformation("No users to notify for step {JobId}:{BatchId}:{StepId}", job.Id, batchId, stepId);
+				return;
+			}
+
+			if (step.LogId == null)
+			{
+				_logger.LogError("Step does not have a log file");
+				return;
+			}
+
+			ILogFile? logFile = await _logFileService.GetLogFileAsync(step.LogId.Value, CancellationToken.None);
+			if(logFile == null)
+			{
+				_logger.LogError("Step does not have a log file");
+				return;
+			}
+
+			List<ILogEvent> jobStepEvents = await _logFileService.FindEventsAsync(logFile);
+			List<ILogEventData> jobStepEventData = new List<ILogEventData>();
+			foreach (ILogEvent logEvent in jobStepEvents)
+			{
+				ILogEventData eventData = await _logFileService.GetEventDataAsync(logFile, logEvent.LineIndex, logEvent.LineCount);
+				jobStepEventData.Add(eventData);
+			}
+
+			foreach (IUser slackUser in usersToNotify)
+			{
+				if(job.PreflightChange != 0)
+				{
+					if(slackUser.Id != job.StartedByUserId)
+					{
+						continue;
+					}
+				}
+				EnqueueTasks(sink => sink.NotifyJobStepCompleteAsync(slackUser, job, batch, step, node, jobStepEventData));
+			}
+			_logger.LogDebug("Finished sending notifications for step {JobId}:{BatchId}:{StepId}", job.Id, batchId, stepId);
+		}
+
+		private async Task SendAllLabelNotificationsAsync(IJob job, IReadOnlyList<(LabelState State, LabelOutcome Outcome)> oldLabelStates, IReadOnlyList<(LabelState, LabelOutcome)> newLabelStates)
+		{
+			IGraph? graph = await _graphCollection.GetAsync(job.GraphHash);
+			if (graph == null)
+			{
+				_logger.LogError("Unable to find graph {GraphHash} for job {JobId}", job.GraphHash, job.Id);
+				return;
+			}
+
+			IReadOnlyDictionary<NodeRef, IJobStep> stepForNode = job.GetStepForNodeMap();
+			for (int labelIdx = 0; labelIdx < graph.Labels.Count; ++labelIdx)
+			{
+				(LabelState State, LabelOutcome Outcome) oldLabel = oldLabelStates[labelIdx];
+				(LabelState State, LabelOutcome Outcome) newLabel = newLabelStates[labelIdx];
+				if (oldLabel != newLabel)
+				{
+					// If the state transitioned from Unspecified to Running, don't update unless the outcome also changed.
+					if (oldLabel.State == LabelState.Unspecified && newLabel.State == LabelState.Running && oldLabel.Outcome == newLabel.Outcome)
+					{
+						continue;
+					}
+
+					// If the label isn't complete, don't report on outcome changing to success, this will be reported when the label state becomes complete.
+					if (newLabel.State != LabelState.Complete && newLabel.Outcome == LabelOutcome.Success)
+					{
+						return;
+					}
+
+					ILabel label = graph.Labels[labelIdx];
+
+					EventRecord? eventId;
+					if (String.IsNullOrEmpty(label.DashboardName))
+					{
+						eventId = null;
+					}
+					else
+					{
+						eventId = new LabelCompleteEventRecord(job.StreamId, job.TemplateId, label.DashboardCategory, label.DashboardName, newLabel.Outcome);
+					}
+
+					ObjectId? triggerId;
+					if (job.LabelIdxToTriggerId.TryGetValue(labelIdx, out ObjectId newTriggerId))
+					{
+						triggerId = newTriggerId;
+					}
+					else
+					{
+						triggerId = null;
+					}
+
+					bool fireTrigger = newLabel.State == LabelState.Complete;
+
+					List<IUser> usersToNotify = await GetUsersToNotify(eventId, triggerId, fireTrigger);
+
+					// filter preflight label notifications to only include initiator
+					if (usersToNotify.Count > 0 && job.PreflightChange != 0 && job.StartedByUserId != null)
+					{
+						usersToNotify = usersToNotify.Where(x => x.Id == job.StartedByUserId).ToList();
+					}
+
+					if (usersToNotify.Count > 0)
+					{
+						SendLabelUpdateNotifications(job, graph, stepForNode, graph.Labels[labelIdx], labelIdx, newLabel.Outcome, usersToNotify);
+					}
+					else
+					{
+						_logger.LogDebug("No users to notify for label {DashboardName}/{UgsName} in job {JobId}", graph.Labels[labelIdx].DashboardName, graph.Labels[labelIdx].UgsName, job.Id);
+					}
+				}
+			}
+		}
+
+		private void SendLabelUpdateNotifications(IJob job, IGraph graph, IReadOnlyDictionary<NodeRef, IJobStep> stepForNode, ILabel label, int labelIdx, LabelOutcome outcome, List<IUser> slackUsers)
+		{
+			List<(string, JobStepOutcome, Uri)> stepData = new List<(string, JobStepOutcome, Uri)>();
+			if (outcome != LabelOutcome.Success)
+			{
+				foreach (NodeRef includedNodeRef in label.IncludedNodes)
+				{
+					INode includedNode = graph.GetNode(includedNodeRef);
+					IJobStep includedStep = stepForNode[includedNodeRef];
+					stepData.Add((includedNode.Name, includedStep.Outcome, new Uri($"{_settings.CurrentValue.DashboardUrl}job/{job.Id}?step={includedStep.Id}")));
+				}
+			}
+
+			foreach (IUser slackUser in slackUsers)
+			{
+				EnqueueTasks(sink => sink.NotifyLabelCompleteAsync(slackUser, job, label, labelIdx, outcome, stepData));
+			}
+
+			_logger.LogDebug("Finished sending label notifications for label {DashboardName}/{UgsName} in job {JobId}", label.DashboardName, label.UgsName, job.Id);
+		}
+
+		/// <inheritdoc/>
+		public async Task SendIssueReportAsync(IssueReportGroup report)
+		{
+			foreach (INotificationSink sink in _sinks)
+			{
+				try
+				{
+					await sink.SendIssueReportAsync(report);
+				}
+				catch (Exception e)
+				{
+					string streamsWithChannel = String.Join(", ", report.Reports.Select(x => $"{x.StreamId} {x.TriageChannel}"));
+					_logger.LogError(e, "Failed sending issue report to {Channel} for streams/channels {StreamsWithChannels})", report.Channel, streamsWithChannel);
+				}
+			}
+		}
+	}
+}
diff --git a/Engine/Source/Programs/Horde/Horde.Server/Notifications/NotificationTriggerCollection.cs b/Engine/Source/Programs/Horde/Horde.Server/Notifications/NotificationTriggerCollection.cs
new file mode 100644
--- /dev/null
+++ b/Engine/Source/Programs/Horde/Horde.Server/Notifications/NotificationTriggerCollection.cs
@@ -0,0 +1,241 @@
+// Copyright Epic Games, Inc. All Rights Reserved.
+
+using System.Collections.Generic;
+using System.Linq;
+using System.Threading.Tasks;
+using Horde.Server.Server;
+using Horde.Server.Users;
+using Horde.Server.Utilities;
+using MongoDB.Bson;
+using MongoDB.Bson.Serialization.Attributes;
+using MongoDB.Driver;
+
+namespace Horde.Server.Notifications
+{
+	/// <summary>
+	/// Collection of notification triggers
+	/// </summary>
+	public class NotificationTriggerCollection : INotificationTriggerCollection
+	{
+		class SubscriptionDocument : INotificationSubscription
+		{
+			[BsonIgnoreIfNull]
+			public string? User { get; set; }
+
+			public UserId UserId { get; set; }
+			public bool Email { get; set; }
+			public bool Slack { get; set; }
+
+			[BsonConstructor]
+			private SubscriptionDocument()
+			{
+				User = null!;
+			}
+
+			public SubscriptionDocument(UserId userId, bool email, bool slack)
+			{
+				UserId = userId;
+				Email = email;
+				Slack = slack;
+			}
+
+			public SubscriptionDocument(INotificationSubscription subscription)
+			{
+				UserId = subscription.UserId;
+				Email = subscription.Email;
+				Slack = subscription.Slack;
+			}
+		}
+
+		class TriggerDocument : INotificationTrigger
+		{
+			public ObjectId Id { get; set; }
+			public bool Fired { get; set; }
+			public List<SubscriptionDocument> Subscriptions { get; set; } = new List<SubscriptionDocument>();
+			public int UpdateIndex { get; set; }
+
+			IReadOnlyList<INotificationSubscription> INotificationTrigger.Subscriptions => Subscriptions;
+		}
+
+		readonly IMongoCollection<TriggerDocument> _triggers;
+		readonly IUserCollection _userCollection;
+
+		/// <summary>
+		/// Constructor
+		/// </summary>
+		/// <param name="mongoService">The database singleton</param>
+		/// <param name="userCollection"></param>
+		public NotificationTriggerCollection(MongoService mongoService, IUserCollection userCollection)
+		{
+			_triggers = mongoService.GetCollection<TriggerDocument>("NotificationTriggers");
+			_userCollection = userCollection;
+		}
+
+		/// <inheritdoc/>
+		public async Task<INotificationTrigger?> GetAsync(ObjectId triggerId)
+		{
+			TriggerDocument? trigger = await _triggers.Find(x => x.Id == triggerId).FirstOrDefaultAsync();
+			if (trigger != null)
+			{
+				for (int idx = 0; idx < trigger.Subscriptions.Count; idx++)
+				{
+					SubscriptionDocument subscription = trigger.Subscriptions[idx];
+					if (subscription.User != null)
+					{
+						IUser? user = await _userCollection.FindUserByLoginAsync(subscription.User);
+						if (user == null)
+						{
+							trigger.Subscriptions.RemoveAt(idx);
+							idx--;
+						}
+						else
+						{
+							subscription.UserId = user.Id;
+							subscription.User = null;
+						}
+					}
+				}
+			}
+			return trigger;
+		}
+
+		/// <inheritdoc/>
+		public async Task<INotificationTrigger> FindOrAddAsync(ObjectId triggerId)
+		{
+			for (; ; )
+			{
+				// Find an existing trigger
+				INotificationTrigger? existing = await GetAsync(triggerId);
+				if (existing != null)
+				{
+					return existing;
+				}
+
+				// Try to insert a new document
+				try
+				{
+					TriggerDocument newDocument = new TriggerDocument();
+					newDocument.Id = triggerId;
+					await _triggers.InsertOneAsync(newDocument);
+					return newDocument;
+				}
+				catch (MongoWriteException ex)
+				{
+					if (ex.WriteError.Category != ServerErrorCategory.DuplicateKey)
+					{
+						throw;
+					}
+				}
+			}
+		}
+
+		/// <summary>
+		/// Updates an existing document
+		/// </summary>
+		/// <param name="trigger">The trigger to update</param>
+		/// <param name="transaction">The update definition</param>
+		/// <returns>The updated document</returns>
+		async Task<INotificationTrigger?> TryUpdateAsync(INotificationTrigger trigger, TransactionBuilder<TriggerDocument> transaction)
+		{
+			TriggerDocument document = (TriggerDocument)trigger;
+			int nextUpdateIndex = document.UpdateIndex + 1;
+
+			FilterDefinition<TriggerDocument> filter = Builders<TriggerDocument>.Filter.Expr(x => x.Id == trigger.Id && x.UpdateIndex == document.UpdateIndex);
+			UpdateDefinition<TriggerDocument> update = transaction.ToUpdateDefinition().Set(x => x.UpdateIndex, nextUpdateIndex);
+
+			UpdateResult result = await _triggers.UpdateOneAsync(filter, update);
+			if (result.ModifiedCount > 0)
+			{
+				transaction.ApplyTo(document);
+				document.UpdateIndex = nextUpdateIndex;
+				return document;
+			}
+
+			return null;
+		}
+
+		/// <inheritdoc/>
+		public async Task DeleteAsync(ObjectId triggerId)
+		{
+			await _triggers.DeleteOneAsync(x => x.Id == triggerId);
+		}
+
+		/// <inheritdoc/>
+		public async Task DeleteAsync(List<ObjectId> triggerIds)
+		{
+			FilterDefinition<TriggerDocument> filter = Builders<TriggerDocument>.Filter.In(x => x.Id, triggerIds);
+			await _triggers.DeleteManyAsync(filter);
+		}
+
+		/// <inheritdoc/>
+		public async Task<INotificationTrigger?> FireAsync(INotificationTrigger trigger)
+		{
+			if (trigger.Fired)
+			{
+				return null;
+			}
+
+			for (; ; )
+			{
+				TransactionBuilder<TriggerDocument> transaction = new TransactionBuilder<TriggerDocument>();
+				transaction.Set(x => x.Fired, true);
+
+				INotificationTrigger? newTrigger = await TryUpdateAsync(trigger, transaction);
+				if (newTrigger != null)
+				{
+					return newTrigger;
+				}
+
+				newTrigger = await FindOrAddAsync(trigger.Id); // Need to add to prevent race condition on triggering vs adding
+				if (newTrigger == null || newTrigger.Fired)
+				{
+					return null;
+				}
+			}
+		}
+
+		/// <inheritdoc/>
+		public async Task<INotificationTrigger?> UpdateSubscriptionsAsync(INotificationTrigger trigger, UserId userId, bool? email, bool? slack)
+		{
+			for (; ; )
+			{
+				// If the trigger has already fired, don't add a new subscription to it
+				if(trigger.Fired)
+				{
+					return trigger;
+				}
+
+				// Try to update the trigger
+				List<SubscriptionDocument> newSubscriptions = new List<SubscriptionDocument>();
+				newSubscriptions.AddRange(trigger.Subscriptions.Select(x => new SubscriptionDocument(x)));
+
+				SubscriptionDocument? newSubscription = newSubscriptions.FirstOrDefault(x => x.UserId == userId);
+				if (newSubscription == null)
+				{
+					newSubscription = new SubscriptionDocument(userId, email ?? false, slack ?? false);
+					newSubscriptions.Add(newSubscription);
+				}
+				else
+				{
+					newSubscription.Email = email ?? newSubscription.Email;
+					newSubscription.Slack = slack ?? newSubscription.Slack;
+				}
+
+				TransactionBuilder<TriggerDocument> transaction = new TransactionBuilder<TriggerDocument>();
+				transaction.Set(x => x.Subscriptions, newSubscriptions);
+
+				INotificationTrigger? newTrigger = await TryUpdateAsync(trigger, transaction);
+				if (newTrigger != null)
+				{
+					return newTrigger;
+				}
+
+				newTrigger = await GetAsync(trigger.Id);
+				if (newTrigger == null)
+				{
+					return null;
+				}
+			}
+		}
+	}
+}
