Horde: Rename the Horde.Build project to Horde.Server.

diff --git a/Engine/Source/Programs/Horde/Horde.Server/Notifications/Sinks/SlackNotificationSink.cs b/Engine/Source/Programs/Horde/Horde.Server/Notifications/Sinks/SlackNotificationSink.cs
new file mode 100644
--- /dev/null
+++ b/Engine/Source/Programs/Horde/Horde.Server/Notifications/Sinks/SlackNotificationSink.cs
@@ -0,0 +1,2829 @@
+// Copyright Epic Games, Inc. All Rights Reserved.
+
+using System;
+using System.Collections.Generic;
+using System.Diagnostics.CodeAnalysis;
+using System.Globalization;
+using System.IO;
+using System.Linq;
+using System.Net.Http;
+using System.Net.WebSockets;
+using System.Text;
+using System.Text.Json;
+using System.Text.Json.Serialization;
+using System.Text.RegularExpressions;
+using System.Threading;
+using System.Threading.Tasks;
+using System.Timers;
+using EpicGames.Core;
+using EpicGames.Redis;
+using EpicGames.Redis.Utility;
+using EpicGames.Slack;
+using EpicGames.Slack.Blocks;
+using EpicGames.Slack.Elements;
+using Horde.Server.Agents;
+using Horde.Server.Configuration;
+using Horde.Server.Devices;
+using Horde.Server.Issues;
+using Horde.Server.Issues.External;
+using Horde.Server.Jobs;
+using Horde.Server.Jobs.Graphs;
+using Horde.Server.Logs;
+using Horde.Server.Server;
+using Horde.Server.Streams;
+using Horde.Server.Users;
+using Horde.Server.Utilities;
+using HordeCommon;
+using Microsoft.AspNetCore.Hosting;
+using Microsoft.Extensions.Caching.Memory;
+using Microsoft.Extensions.Hosting;
+using Microsoft.Extensions.Logging;
+using Microsoft.Extensions.Options;
+using MongoDB.Bson;
+using MongoDB.Bson.Serialization.Attributes;
+using MongoDB.Driver;
+
+namespace Horde.Server.Notifications.Sinks
+{
+	/// <summary>
+	/// Maintains a connection to Slack, in order to receive socket-mode notifications of user interactions
+	/// </summary>
+	public sealed class SlackNotificationSink : BackgroundService, INotificationSink, IAvatarService
+	{
+		const bool DefaultAllowMentions = true;
+
+		const int MaxLineLength = 2048;
+
+		// The color to use for error messages.
+		const string ErrorColor = "#ec4c47";
+		
+		// The color to use for warning messages.
+		const string WarningColor = "#f7d154";
+
+		// The color to use for success messages.
+		const string SuccessColor = "#4db507";
+
+		class SocketResponse
+		{
+			[JsonPropertyName("ok")]
+			public bool Ok { get; set; }
+
+			[JsonPropertyName("url")]
+			public Uri? Url { get; set; }
+		}
+
+		class EventMessage
+		{
+			[JsonPropertyName("envelope_id")]
+			public string? EnvelopeId { get; set; }
+
+			[JsonPropertyName("type")]
+			public string? Type { get; set; }
+
+			[JsonPropertyName("payload")]
+			public EventPayload? Payload { get; set; }
+		}
+
+		class EventResponse
+		{
+			[JsonPropertyName("envelope_id")]
+			public string? EnvelopeId { get; set; }
+
+			[JsonPropertyName("payload")]
+			public object? Payload { get; set; }
+		}
+
+		class EventPayload
+		{
+			[JsonPropertyName("type")]
+			public string? Type { get; set; }
+
+			[JsonPropertyName("trigger_id")]
+			public string? TriggerId { get; set; }
+
+			[JsonPropertyName("user")]
+			public SlackUser? User { get; set; }
+
+			[JsonPropertyName("response_url")]
+			public string? ResponseUrl { get; set; }
+
+			[JsonPropertyName("actions")]
+			public List<ActionInfo> Actions { get; set; } = new List<ActionInfo>();
+
+			[JsonPropertyName("view")]
+			public SlackView? View { get; set; }
+		}
+
+		class ActionInfo
+		{
+			[JsonPropertyName("blockId")]
+			public string BlockId { get; set; } = String.Empty;
+
+			[JsonPropertyName("action_id")]
+			public string ActionId { get; set; } = String.Empty;
+
+			[JsonPropertyName("value")]
+			public string Value { get; set; } = String.Empty;
+		}
+
+		class MessageStateDocument
+		{
+			[BsonId]
+			public ObjectId Id { get; set; }
+
+			[BsonElement("uid")]
+			public string Recipient { get; set; } = String.Empty;
+
+			[BsonElement("usr")]
+			public UserId? UserId { get; set; }
+
+			[BsonElement("eid")]
+			public string EventId { get; set; } = String.Empty;
+
+			[BsonElement("ch")]
+			public string Channel { get; set; } = String.Empty;
+
+			[BsonElement("ts")]
+			public string Ts { get; set; } = String.Empty;
+
+			[BsonElement("dig")]
+			public string Digest { get; set; } = String.Empty;
+
+			[BsonElement("lnk"), BsonIgnoreIfNull]
+			public string? Permalink { get; set; }
+
+			[BsonIgnore]
+			public SlackMessageId MessageId => new SlackMessageId(Channel, null, Ts);
+		}
+
+		class SlackUserDocument : IAvatar
+		{
+			public const int CurrentVersion = 2;
+			
+			public UserId Id { get; set; }
+
+			[BsonElement("u")]
+			public string? SlackUserId { get; set; }
+
+			[BsonElement("i24")]
+			public string? Image24 { get; set; }
+
+			[BsonElement("i32")]
+			public string? Image32 { get; set; }
+			
+			[BsonElement("i48")]
+			public string? Image48 { get; set; }
+
+			[BsonElement("i72")]
+			public string? Image72 { get; set; }
+
+			[BsonElement("t")]
+			public DateTime Time { get; set; }
+
+			[BsonElement("v")]
+			public int Version { get; set; }
+
+			private SlackUserDocument()
+			{
+			}
+
+			public SlackUserDocument(UserId id, SlackUser? info)
+			{
+				Id = id;
+				SlackUserId = info?.Id;
+				if (info != null && info.Profile != null && info.Profile.IsCustomImage)
+				{
+					Image24 = info.Profile.Image24;
+					Image32 = info.Profile.Image32;
+					Image48 = info.Profile.Image48;
+					Image72 = info.Profile.Image72;
+				}
+				Time = DateTime.UtcNow;
+				Version = CurrentVersion;
+			}
+		}
+
+		readonly RedisService _redisService;
+		readonly IssueService _issueService;
+		readonly IUserCollection _userCollection;
+		readonly ILogFileService _logFileService;
+		readonly IStreamCollection _streamCollection;
+		readonly IWebHostEnvironment _environment;
+		readonly ServerSettings _settings;
+		readonly IMongoCollection<MessageStateDocument> _messageStates;
+		readonly IMongoCollection<SlackUserDocument> _slackUsers;
+		readonly HashSet<string>? _allowUsers;
+		readonly IExternalIssueService _externalIssueService;
+		readonly JsonSerializerOptions _jsonSerializerOptions;
+		readonly ITicker _escalateTicker;
+		static readonly RedisSortedSetKey<int> _escalateIssues = "slack/escalate";
+		readonly IClock _clock;
+		readonly IOptionsMonitor<GlobalConfig> _globalConfig;
+		readonly ILogger _logger;
+
+		readonly ITicker _issueQueueTicker;
+		static readonly RedisListKey<int> _redisIssueQueue = "slack/issue-queue";
+		readonly string _redisIssueLockPrefix;
+
+		readonly HttpClient _httpClient;
+		readonly SlackClient _slackClient;
+
+		readonly HttpClient? _adminHttpClient;
+		readonly SlackClient? _adminSlackClient;
+
+		/// <summary>
+		/// Map of email address to Slack user ID.
+		/// </summary>
+		private readonly MemoryCache _userCache = new MemoryCache(new MemoryCacheOptions());
+
+		/// <summary>
+		/// Constructor
+		/// </summary>
+		public SlackNotificationSink(MongoService mongoService, RedisService redisService, IssueService issueService, IUserCollection userCollection, ILogFileService logFileService, IStreamCollection streamCollection, IExternalIssueService externalIssueService, IWebHostEnvironment environment, IOptions<ServerSettings> settings, IClock clock, IOptionsMonitor<GlobalConfig> globalConfig, ILogger<SlackNotificationSink> logger)
+		{
+			_redisService = redisService;
+			_issueService = issueService;
+			_userCollection = userCollection;
+			_logFileService = logFileService;
+			_streamCollection = streamCollection;
+			_externalIssueService = externalIssueService;
+			_environment = environment;
+			_settings = settings.Value;
+			_messageStates = mongoService.GetCollection<MessageStateDocument>("SlackV2", keys => keys.Ascending(x => x.Recipient).Ascending(x => x.EventId), unique: true);
+			_slackUsers = mongoService.GetCollection<SlackUserDocument>("Slack.UsersV2");
+			_clock = clock;
+			_globalConfig = globalConfig;
+			_logger = logger;
+
+			_escalateTicker = clock.AddSharedTicker<SlackNotificationSink>(TimeSpan.FromMinutes(1.0), EscalateAsync, _logger);
+
+			_issueQueueTicker = clock.AddSharedTicker("slack-issues", TimeSpan.FromMinutes(1.0), ProcessIssueQueueAsync, _logger);
+			_redisIssueLockPrefix = "slack/issues/";
+
+			_httpClient = new HttpClient();
+			_httpClient.DefaultRequestHeaders.Add("Authorization", $"Bearer {_settings.SlackToken ?? ""}");
+			_slackClient = new SlackClient(_httpClient, _logger);
+
+			if (!String.IsNullOrEmpty(_settings.SlackAdminToken))
+			{
+				_adminHttpClient = new HttpClient();
+				_adminHttpClient.DefaultRequestHeaders.Add("Authorization", $"Bearer {_settings.SlackAdminToken}");
+				_adminSlackClient = new SlackClient(_adminHttpClient, _logger);
+			}
+
+			_jsonSerializerOptions = new JsonSerializerOptions();
+			Startup.ConfigureJsonSerializer(_jsonSerializerOptions);
+
+			if (!String.IsNullOrEmpty(settings.Value.SlackUsers))
+			{
+				_allowUsers = new HashSet<string>(settings.Value.SlackUsers.Split(','), StringComparer.OrdinalIgnoreCase);
+			}
+		}
+
+		/// <inheritdoc/>
+		public override void Dispose()
+		{
+			base.Dispose();
+
+			_userCache.Dispose();
+			_httpClient.Dispose();
+			_adminHttpClient?.Dispose();
+			_issueQueueTicker.Dispose();
+			_escalateTicker.Dispose();
+		}
+
+		/// <inheritdoc/>
+		public override async Task StartAsync(CancellationToken cancellationToken)
+		{
+			await base.StartAsync(cancellationToken);
+
+			await _issueQueueTicker.StartAsync();
+			await _escalateTicker.StartAsync();
+		}
+
+		/// <inheritdoc/>
+		public override async Task StopAsync(CancellationToken cancellationToken)
+		{
+			await base.StopAsync(cancellationToken);
+
+			await _escalateTicker.StopAsync();
+			await _issueQueueTicker.StopAsync();
+		}
+
+		#region Avatars
+
+		/// <inheritdoc/>
+		public async Task<IAvatar?> GetAvatarAsync(IUser user)
+		{
+			return await GetSlackUser(user);
+		}
+
+		#endregion
+
+		#region Message state 
+
+		async Task<(MessageStateDocument, bool)> AddOrUpdateMessageStateAsync(string recipient, string eventId, UserId? userId, string digest, SlackMessageId? messageId)
+		{
+			ObjectId newId = ObjectId.GenerateNewId();
+
+			FilterDefinition<MessageStateDocument> filter = Builders<MessageStateDocument>.Filter.Eq(x => x.Recipient, recipient) & Builders<MessageStateDocument>.Filter.Eq(x => x.EventId, eventId);
+			UpdateDefinition<MessageStateDocument> update = Builders<MessageStateDocument>.Update.SetOnInsert(x => x.Id, newId).Set(x => x.UserId, userId).Set(x => x.Digest, digest);
+
+			if (messageId != null)
+			{
+				update = update.Set(x => x.Channel, messageId.Channel).Set(x => x.Ts, messageId.Ts);
+			}
+
+			MessageStateDocument state = await _messageStates.FindOneAndUpdateAsync(filter, update, new FindOneAndUpdateOptions<MessageStateDocument> { IsUpsert = true, ReturnDocument = ReturnDocument.After });
+			if (state.Id == newId)
+			{
+				_logger.LogInformation("Posted message {StateId} (recipient: {Recipient}, user: {UserId}, event: {EventId}, messageId: {MessageId}, digest: {Digest})", state.Id, recipient, userId ?? UserId.Empty, eventId, state.MessageId, digest);
+			}
+			else
+			{
+				_logger.LogInformation("Updated message {StateId} (recipient: {Recipient}, user: {UserId}, event: {EventId}, messageId: {MessageId}, digest: {Digest})", state.Id, recipient, userId ?? UserId.Empty, eventId, state.MessageId, digest);
+			}
+
+			return (state, state.Id == newId);
+		}
+
+		async Task<MessageStateDocument?> GetMessageStateAsync(string recipient, string eventId)
+		{
+			FilterDefinition<MessageStateDocument> filter = Builders<MessageStateDocument>.Filter.Eq(x => x.Recipient, recipient) & Builders<MessageStateDocument>.Filter.Eq(x => x.EventId, eventId);
+			return await _messageStates.Find(filter).FirstOrDefaultAsync();
+		}
+
+		async Task<bool> DeleteMessageStateAsync(string recipient, string eventId)
+		{
+			FilterDefinition<MessageStateDocument> filter = Builders<MessageStateDocument>.Filter.Eq(x => x.Recipient, recipient) & Builders<MessageStateDocument>.Filter.Eq(x => x.EventId, eventId);
+			DeleteResult result = await _messageStates.DeleteOneAsync(filter);
+			return result.DeletedCount > 0;
+		}
+
+		async Task UpdateMessageStateAsync(ObjectId stateId, SlackMessageId id, string? permalink = null)
+		{
+			FilterDefinition<MessageStateDocument> filter = Builders<MessageStateDocument>.Filter.Eq(x => x.Id, stateId);
+			UpdateDefinition<MessageStateDocument> update = Builders<MessageStateDocument>.Update.Set(x => x.Channel, id.Channel).Set(x => x.Ts, id.Ts);
+
+			if (permalink == null)
+			{
+				update = update.Unset(x => x.Permalink);
+			}
+			else
+			{
+				update = update.Set(x => x.Permalink, permalink);
+			}
+
+			await _messageStates.FindOneAndUpdateAsync(filter, update);
+			_logger.LogInformation("Updated message {StateId} (messageId: {MessageId}, permalink: {Permalink})", stateId, id, permalink ?? "(n/a)");
+		}
+
+		#endregion
+
+		/// <inheritdoc/>
+		public async Task NotifyJobScheduledAsync(List<JobScheduledNotification> notifications)
+		{
+			if (_settings.JobNotificationChannel != null)
+			{
+				string jobIds = String.Join(", ", notifications.Select(x => x.JobId));
+				_logger.LogInformation("Sending Slack notification for scheduled job IDs {JobIds} to channel {SlackChannel}", jobIds, _settings.JobNotificationChannel);
+				foreach (string channel in _settings.JobNotificationChannel.Split(';'))
+				{
+					await SendJobScheduledOnEmptyAutoScaledPoolMessageAsync($"#{channel}", notifications);
+				}
+			}
+		}
+		
+		private async Task SendJobScheduledOnEmptyAutoScaledPoolMessageAsync(string recipient, List<JobScheduledNotification> notifications)
+		{
+			const int MaxItems = 10;
+			string jobIds = StringUtils.FormatList(notifications.Select(x => x.JobId.ToString()).ToArray(), MaxItems);
+				
+			StringBuilder sb = new();
+			for(int idx = 0; idx < notifications.Count; idx++)
+			{
+				if (idx >= MaxItems && notifications.Count > MaxItems + 2)
+				{
+					sb.AppendLine($"...and {notifications.Count - idx} others.");
+					break;
+				}
+
+				JobScheduledNotification notification = notifications[idx];
+				string jobUrl = _settings.DashboardUrl + "/job/" + notification.JobId;
+				sb.AppendLine($"Job `{notification.JobName}` with ID <{jobUrl}|{notification.JobId}> in pool `{notification.PoolName}`");
+			}
+
+			string outcomeColor = WarningColor;
+			SlackAttachment attachment = new SlackAttachment();
+			attachment.Color = outcomeColor;
+			attachment.FallbackText = $"Job(s) scheduled in an auto-scaled pool with no agents online. Job IDs {jobIds}";
+
+			attachment.AddHeader($"Jobs scheduled in empty pool", true);
+			attachment.AddSection($"One or more jobs were scheduled in an auto-scaled pool but with no current agents online.");
+			if (sb.Length > 0)
+			{
+				attachment.AddSection(sb.ToString());
+			}
+
+			await SendMessageAsync(recipient, attachment);
+		}
+
+		#region Job Complete
+
+		/// <inheritdoc/>
+		public async Task NotifyJobCompleteAsync(IJob job, IGraph graph, LabelOutcome outcome)
+		{
+			StreamConfig? streamConfig;
+			if (!_globalConfig.CurrentValue.TryGetStream(job.StreamId, out streamConfig))
+			{
+				return;
+			}
+			if (job.NotificationChannel != null)
+			{
+				await SendJobCompleteNotificationToChannelAsync(job.NotificationChannel, job.NotificationChannelFilter, streamConfig, job, graph, outcome);
+			}
+			if (streamConfig.NotificationChannel != null)
+			{
+				await SendJobCompleteNotificationToChannelAsync(streamConfig.NotificationChannel, streamConfig.NotificationChannelFilter, streamConfig, job, graph, outcome);
+			}
+		}
+
+		async Task SendJobCompleteNotificationToChannelAsync(string notificationChannel, string? notificationFilter, StreamConfig streamConfig, IJob job, IGraph graph, LabelOutcome outcome)
+		{
+			if (notificationFilter != null)
+			{
+				List<LabelOutcome> outcomes = new List<LabelOutcome>();
+				foreach (string filterOption in notificationFilter.Split('|'))
+				{
+					LabelOutcome result;
+					if (Enum.TryParse(filterOption, out result))
+					{
+						outcomes.Add(result);
+					}
+					else
+					{
+						_logger.LogWarning("Invalid filter option {Option} specified in job filter {NotificationChannelFilter} in job {JobId} or stream {StreamId}", filterOption, notificationFilter, job.Id, job.StreamId);
+					}
+				}
+				if (!outcomes.Contains(outcome))
+				{
+					return;
+				}
+			}
+			foreach (string channel in notificationChannel.Split(';'))
+			{
+				await SendJobCompleteMessageAsync(channel, streamConfig, job, graph);
+			}
+		}
+
+		/// <inheritdoc/>
+		public async Task NotifyJobCompleteAsync(IUser slackUser, IJob job, IGraph graph, LabelOutcome outcome)
+		{
+			StreamConfig? streamConfig;
+			if (!_globalConfig.CurrentValue.TryGetStream(job.StreamId, out streamConfig))
+			{
+				return;
+			}
+
+			string? slackUserId = await GetSlackUserId(slackUser);
+			if (slackUserId != null)
+			{
+				await SendJobCompleteMessageAsync(slackUserId, streamConfig, job, graph);
+			}
+		}
+
+		private Task SendJobCompleteMessageAsync(string recipient, StreamConfig streamConfig, IJob job, IGraph graph)
+		{
+			JobStepOutcome jobOutcome = job.Batches.SelectMany(x => x.Steps).Min(x => x.Outcome);
+			_logger.LogInformation("Sending Slack notification for job {JobId} outcome {Outcome} to {SlackUser}", job.Id, jobOutcome, recipient);
+
+			Uri jobLink = new Uri($"{_settings.DashboardUrl}job/{job.Id}");
+
+			string outcomeColor = jobOutcome == JobStepOutcome.Failure ? ErrorColor : jobOutcome == JobStepOutcome.Warnings ? WarningColor : SuccessColor;
+
+			SlackAttachment attachment = new SlackAttachment();
+			attachment.FallbackText = $"{streamConfig.Name} - {GetJobChangeText(job)} - {job.Name} - {jobOutcome}";
+			attachment.Color = outcomeColor;
+			attachment.AddSection($"*<{jobLink}|{streamConfig.Name} - {GetJobChangeText(job)} - {job.Name}>*");
+
+			if (!String.IsNullOrEmpty(job.PreflightDescription))
+			{
+				string description = job.PreflightDescription.TrimEnd();
+				if (jobOutcome == JobStepOutcome.Success)
+				{
+					description += $"\n#preflight {job.Id}";
+				}
+				attachment.AddSection($"```{description}```");
+			}
+
+			if (jobOutcome == JobStepOutcome.Success)
+			{
+				attachment.AddSection($"*Job Succeeded*");
+			}
+			else
+			{
+				List<string> failedStepStrings = new List<string>();
+				List<string> warningStepStrings = new List<string>();
+
+				IReadOnlyDictionary<NodeRef, IJobStep> nodeToStep = job.GetStepForNodeMap();
+				foreach ((NodeRef nodeRef, IJobStep step) in nodeToStep)
+				{
+					if (step.State == JobStepState.Completed)
+					{
+						INode stepNode = graph.GetNode(nodeRef);
+						string stepName = $"<{jobLink}?step={step.Id}|{stepNode.Name}>";
+						if (step.Outcome == JobStepOutcome.Failure)
+						{
+							failedStepStrings.Add(stepName);
+						}
+						else if (step.Outcome == JobStepOutcome.Warnings)
+						{
+							warningStepStrings.Add(stepName);
+						}
+					}
+				}
+
+				if (failedStepStrings.Any())
+				{
+					string msg = $"*Errors*\n{String.Join(", ", failedStepStrings)}";
+					attachment.AddSection(msg.Substring(0, Math.Min(msg.Length, 3000)));
+				}
+				else if (warningStepStrings.Any())
+				{
+					string msg = $"*Warnings*\n{String.Join(", ", warningStepStrings)}";
+					attachment.AddSection(msg.Substring(0, Math.Min(msg.Length, 3000)));
+				}
+			}
+
+			if (job.AutoSubmit)
+			{
+				attachment.AddDivider();
+				if (job.AutoSubmitChange != null)
+				{
+					attachment.AddSection($"Shelved files were submitted in CL {job.AutoSubmitChange}.");
+				}
+				else
+				{
+					attachment.Color = WarningColor;
+
+					string autoSubmitMessage = String.Empty;
+					if (!String.IsNullOrEmpty(job.AutoSubmitMessage))
+					{
+						autoSubmitMessage = $"\n\n```{job.AutoSubmitMessage}```";
+					}
+
+					attachment.AddSection($"Files in CL *{job.PreflightChange}* were *not submitted*. Please resolve these issues and submit manually.{autoSubmitMessage}");
+				}
+			}
+
+			return SendMessageAsync(recipient, attachment);
+		}
+
+		#endregion
+
+		#region Job step complete
+
+		/// <inheritdoc/>
+		public async Task NotifyJobStepCompleteAsync(IUser slackUser, IJob job, IJobStepBatch batch, IJobStep step, INode node, List<ILogEventData> jobStepEventData)
+		{
+			_logger.LogInformation("Sending Slack notification for job {JobId}, batch {BatchId}, step {StepId}, outcome {Outcome} to {SlackUser} ({UserId})", job.Id, batch.Id, step.Id, step.Outcome, slackUser.Name, slackUser.Id);
+
+			string? slackUserId = await GetSlackUserId(slackUser);
+			if (slackUserId != null)
+			{
+				await SendJobStepCompleteMessageAsync(slackUserId, job, step, node, jobStepEventData);
+			}
+		}
+
+		/// <summary>
+		/// Creates a Slack message about a completed step job.
+		/// </summary>
+		/// <param name="recipient"></param>
+		/// <param name="job">The job that contains the step that completed.</param>
+		/// <param name="step">The job step that completed.</param>
+		/// <param name="node">The node for the job step.</param>
+		/// <param name="events">Any events that occurred during the job step.</param>
+		private Task SendJobStepCompleteMessageAsync(string recipient, IJob job, IJobStep step, INode node, List<ILogEventData> events)
+		{
+			GlobalConfig globalConfig = _globalConfig.CurrentValue;
+			if (!globalConfig.TryGetStream(job.StreamId, out StreamConfig? streamConfig))
+			{
+				return Task.CompletedTask;
+			}
+
+			Uri jobStepLink = new Uri($"{_settings.DashboardUrl}job/{job.Id}?step={step.Id}");
+			Uri jobStepLogLink = new Uri($"{_settings.DashboardUrl}log/{step.LogId}");
+
+			string outcomeColor = step.Outcome == JobStepOutcome.Failure ? ErrorColor : step.Outcome == JobStepOutcome.Warnings ? WarningColor : SuccessColor;
+
+			SlackAttachment attachment = new SlackAttachment();
+			attachment.FallbackText = $"{streamConfig.Name} - {GetJobChangeText(job)} - {job.Name} - {node.Name} - {step.Outcome}";
+			attachment.Color = outcomeColor;
+			attachment.AddSection($"*<{jobStepLink}|{streamConfig.Name} - {GetJobChangeText(job)} - {job.Name} - {node.Name}>*");
+			if (step.Outcome == JobStepOutcome.Success)
+			{
+				attachment.AddSection($"*Job Step Succeeded*");
+			}
+			else
+			{
+				List<ILogEventData> errors = events.Where(x => x.Severity == EventSeverity.Error).ToList();
+				List<ILogEventData> warnings = events.Where(x => x.Severity == EventSeverity.Warning).ToList();
+				List<string> eventStrings = new List<string>();
+				if (errors.Any())
+				{
+					string errorSummary = errors.Count > MaxJobStepEvents ? $"*Errors (First {MaxJobStepEvents} shown)*" : $"*Errors*";
+					attachment.AddSection(errorSummary);
+					foreach (ILogEventData error in errors.Take(MaxJobStepEvents))
+					{
+						attachment.AddSection(QuoteText(error.Message));
+					}
+				}
+				else if (warnings.Any())
+				{
+					string warningSummary = warnings.Count > MaxJobStepEvents ? $"*Warnings (First {MaxJobStepEvents} shown)*" : $"*Warnings*";
+					eventStrings.Add(warningSummary);
+					foreach (ILogEventData warning in warnings.Take(MaxJobStepEvents))
+					{
+						attachment.AddSection(QuoteText(warning.Message));
+					}
+				}
+
+				attachment.AddSection($"<{jobStepLogLink}|View Job Step Log>");
+			}
+
+			return SendMessageAsync(recipient, attachment);
+		}
+
+		#endregion
+
+		#region Label complete
+
+		/// <inheritdoc/>
+		public async Task NotifyLabelCompleteAsync(IUser user, IJob job, ILabel label, int labelIdx, LabelOutcome outcome, List<(string, JobStepOutcome, Uri)> stepData)
+		{
+			if (!_globalConfig.CurrentValue.TryGetStream(job.StreamId, out StreamConfig? streamConfig))
+			{
+				return;
+			}
+
+			_logger.LogInformation("Sending Slack notification for job {JobId} outcome {Outcome} to {Name} ({UserId})", job.Id, outcome, user.Name, user.Id);
+
+			string? slackUserId = await GetSlackUserId(user);
+			if (slackUserId != null)
+			{
+				await SendLabelUpdateMessageAsync(slackUserId, streamConfig, job, label, labelIdx, outcome, stepData);
+			}
+		}
+
+		Task SendLabelUpdateMessageAsync(string recipient, StreamConfig streamConfig, IJob job, ILabel label, int labelIdx, LabelOutcome outcome, List<(string, JobStepOutcome, Uri)> jobStepData)
+		{
+			Uri labelLink = new Uri($"{_settings.DashboardUrl}job/{job.Id}?label={labelIdx}");
+
+			string outcomeColor = outcome == LabelOutcome.Failure ? ErrorColor : outcome == LabelOutcome.Warnings ? WarningColor : SuccessColor;
+
+			SlackAttachment attachment = new SlackAttachment();
+			attachment.FallbackText = $"{streamConfig.Name} - {GetJobChangeText(job)} - {job.Name} - Label {label.DashboardName} - {outcome}";
+			attachment.Color = outcomeColor;
+			attachment.AddSection($"*<{labelLink}|{streamConfig.Name} - {GetJobChangeText(job)} - {job.Name} - Label {label.DashboardName}>*");
+			if (outcome == LabelOutcome.Success)
+			{
+				attachment.AddSection($"*Label Succeeded*");
+			}
+			else
+			{
+				List<string> failedStepStrings = new List<string>();
+				List<string> warningStepStrings = new List<string>();
+				foreach ((string Name, JobStepOutcome StepOutcome, Uri Link) jobStep in jobStepData)
+				{
+					string stepString = $"<{jobStep.Link}|{jobStep.Name}>";
+					if (jobStep.StepOutcome == JobStepOutcome.Failure)
+					{
+						failedStepStrings.Add(stepString);
+					}
+					else if (jobStep.StepOutcome == JobStepOutcome.Warnings)
+					{
+						warningStepStrings.Add(stepString);
+					}
+				}
+
+				if (failedStepStrings.Any())
+				{
+					attachment.AddSection($"*Errors*\n{String.Join(", ", failedStepStrings)}");
+				}
+				else if (warningStepStrings.Any())
+				{
+					attachment.AddSection($"*Warnings*\n{String.Join(", ", warningStepStrings)}");
+				}
+			}
+
+			return SendMessageAsync(recipient, attachment);
+		}
+
+		#endregion
+
+		#region Issues
+
+		/// <inheritdoc/>
+		public async Task NotifyIssueUpdatedAsync(IIssue issue)
+		{
+			// Do not send notifications for quarantined issues
+			if (issue.QuarantinedByUserId != null)
+			{
+				_logger.LogInformation("Skipping notifications for quarantined issue {IssueId}", issue.Id);
+				return;
+			}
+
+			// Otherwise add it to the redis queue, and attempt to process the queue immediately.
+			await _redisService.GetDatabase().ListRightPushAsync(_redisIssueQueue, issue.Id);
+			await ProcessIssueQueueAsync(CancellationToken.None);
+		}
+
+		/// <summary>
+		/// Processes the issue queue
+		/// </summary>
+		/// <param name="cancellationToken">Cancellation token for the operation</param>
+		async ValueTask ProcessIssueQueueAsync(CancellationToken cancellationToken)
+		{
+			GlobalConfig globalConfig = _globalConfig.CurrentValue;
+			HashSet<int> testedIssueIds = new HashSet<int>();
+
+			// Execute loop number of times based on the length of the queue at the start. This should bound the number of iterations while allowing us
+			// to re-queue items that we're unable to process now, without having to track whether we've reached the end of the list in its original state.
+			long count = await _redisService.GetDatabase().ListLengthAsync(_redisIssueQueue);
+			for (; count > 0; count--)
+			{
+				int issueId = await _redisService.GetDatabase().ListLeftPopAsync(_redisIssueQueue);
+				if (!testedIssueIds.Add(issueId) || !await TryUpdateIssueAsync(globalConfig, issueId))
+				{
+					await _redisService.GetDatabase().ListRightPushAsync(_redisIssueQueue, issueId);
+				}
+				cancellationToken.ThrowIfCancellationRequested();
+			}
+		}
+
+		async ValueTask<bool> TryUpdateIssueAsync(GlobalConfig globalConfig, int issueId)
+		{
+			using (RedisLock issueLock = new RedisLock(_redisService.GetDatabase(), $"{_redisIssueLockPrefix}/{issueId}"))
+			{
+				if (await issueLock.AcquireAsync(TimeSpan.FromSeconds(30.0)))
+				{
+					await NotifyIssueUpdatedInternalAsync(globalConfig, issueId);
+					return true;
+				}
+				else
+				{
+					_logger.LogDebug("Unable to aquire lock for updating issue {IssueId}.", issueId);
+					return false;
+				}
+			}
+		}
+
+		async Task NotifyIssueUpdatedInternalAsync(GlobalConfig globalConfig, int issueId)
+		{
+			using IDisposable scope = _logger.BeginScope("Slack notifications for issue {IssueId}", issueId);
+			_logger.LogInformation("Updating Slack notifications for issue {IssueId}", issueId);
+
+			IIssueDetails? details = await _issueService.GetIssueDetailsAsync(issueId);
+			if (details == null)
+			{
+				return;
+			}
+
+			IIssue issue = details.Issue;
+
+			bool notifyOwner = true;
+			bool notifySuspects = issue.Promoted || details.Spans.Any(x => x.LastFailure.Annotations.NotifySubmitters ?? false);
+
+			WorkflowConfig? workflow = null;
+			if (details.Spans.Count > 0)
+			{
+				IIssueSpan span = details.Spans[0];
+
+				WorkflowId? workflowId = span.LastFailure.Annotations.WorkflowId;
+				if (workflowId != null)
+				{
+					StreamConfig? streamConfig;
+					if (globalConfig.TryGetStream(span.StreamId, out streamConfig) && streamConfig.TryGetWorkflow(workflowId.Value, out workflow) && workflow.TriageChannel != null)
+					{
+						await CreateOrUpdateWorkflowThreadAsync(workflow.TriageChannel, issue, span, details.Spans, workflow);
+						notifyOwner = notifySuspects = false;
+					}
+				}
+			}
+
+			HashSet<UserId> userIds = new HashSet<UserId>();
+			if (notifySuspects)
+			{
+				if (details.Suspects.Count > 5)
+				{
+					_issueService.Collection.GetLogger(issueId).LogInformation("Not notifying suspects for issue; too many users ({Count}).", details.Suspects.Count);
+				}
+				else
+				{
+					userIds.UnionWith(details.Suspects.Select(x => x.AuthorId));
+				}
+			}
+			if (notifyOwner && issue.OwnerId.HasValue)
+			{
+				userIds.Add(issue.OwnerId.Value);
+			}
+
+			HashSet<string> channels = new HashSet<string>(StringComparer.OrdinalIgnoreCase);
+
+			List<MessageStateDocument> existingMessages = await _messageStates.Find(x => x.EventId == GetIssueEventId(issue)).ToListAsync();
+			foreach (MessageStateDocument existingMessage in existingMessages)
+			{
+				if (existingMessage.UserId != null)
+				{
+					userIds.Add(existingMessage.UserId.Value);
+				}
+				else
+				{
+					channels.Add(existingMessage.Recipient);
+				}
+			}
+
+			foreach (IIssueSpan span in details.Spans)
+			{
+				StreamConfig? streamConfig;
+				if (globalConfig.TryGetStream(span.StreamId, out streamConfig))
+				{
+					TemplateRefConfig? templateRefConfig;
+					if (streamConfig.TryGetTemplate(span.TemplateRefId, out templateRefConfig) && templateRefConfig.TriageChannel != null)
+					{
+						channels.Add(templateRefConfig.TriageChannel);
+					}
+					else if (streamConfig.TriageChannel != null)
+					{
+						channels.Add(streamConfig.TriageChannel);
+					}
+				}
+			}
+
+			if (userIds.Count > 0)
+			{
+				foreach (UserId userId in userIds)
+				{
+					IUser? user = await _userCollection.GetUserAsync(userId);
+					if (user == null)
+					{
+						_logger.LogWarning("Unable to find user {UserId}", userId);
+					}
+					else
+					{
+						await NotifyIssueUpdatedAsync(globalConfig, user, issue, details);
+					}
+				}
+			}
+
+			if (channels.Count > 0)
+			{
+				foreach (string channel in channels)
+				{
+					await SendIssueMessageAsync(globalConfig, channel, issue, details, null, DefaultAllowMentions);
+				}
+			}
+
+			await UpdateReportsAsync(globalConfig, issue, details.Spans);
+		}
+
+		static IIssueSpan? GetFixFailedSpan(IIssue issue, IReadOnlyList<IIssueSpan> spans)
+		{
+			if (issue.FixChange != null)
+			{
+				HashSet<StreamId> originStreams = new HashSet<StreamId>(issue.Streams.Where(x => x.MergeOrigin ?? false).Select(x => x.StreamId));
+				foreach (IIssueSpan originSpan in spans.Where(x => originStreams.Contains(x.StreamId)).OrderBy(x => x.StreamId.ToString()))
+				{
+					if (originSpan.LastFailure.Change > issue.FixChange.Value)
+					{
+						return originSpan;
+					}
+				}
+			}
+			return null;
+		}
+
+		async Task InviteUsersAsync(string channel, IEnumerable<UserId> userIds, bool inviteUsersAsAdmin)
+		{
+			List<string> slackUserIds = new List<string>();
+			foreach (UserId userId in userIds)
+			{
+				IUser? user = await _userCollection.GetUserAsync(userId);
+				if (user != null)
+				{
+					string? slackUserId = await GetSlackUserId(user);
+					if (slackUserId != null)
+					{
+						slackUserIds.Add(slackUserId);
+					}
+				}
+			}
+
+			foreach (string slackUserId in slackUserIds)
+			{
+				string? errorCode = await _slackClient.TryInviteUsersAsync(channel, new[] { slackUserId });
+				if (errorCode != null)
+				{
+					if (errorCode.Equals("user_is_restricted", StringComparison.Ordinal) && inviteUsersAsAdmin && _adminSlackClient != null)
+					{
+						try
+						{
+							await _adminSlackClient.AdminInviteUsersAsync(channel, new[] { slackUserId });
+						}
+						catch (SlackException ex)
+						{
+							_logger.LogWarning(ex, "Unable to invite user {UserId} to {Channel} (as admin): {Error}", slackUserId, channel, ex.Code);
+						}
+					}
+					else
+					{
+						_logger.LogWarning("Unable to invite user {UserId} to {Channel}: {Error}", slackUserId, channel, errorCode);
+					}
+				}
+			}
+		}
+
+		static string GetTriageThreadEventId(int issueId) => $"issue_triage_{issueId}";
+
+		async Task CreateOrUpdateWorkflowThreadAsync(string triageChannel, IIssue issue, IIssueSpan span, IReadOnlyList<IIssueSpan> spans, WorkflowConfig workflow)
+		{
+			Uri issueUrl = GetIssueUrl(issue, span.FirstFailure);
+
+			string eventId = GetTriageThreadEventId(issue.Id);
+
+			string prefix = workflow.TriagePrefix ?? String.Empty;
+			if (issue.Severity == IssueSeverity.Error && (span.LastFailure.Annotations.BuildBlocker ?? false))
+			{
+				prefix = $"{prefix}*[BUILD BLOCKER]* ";
+			}
+
+			string issueSummary = issue.UserSummary ?? issue.Summary;
+			string text = $"{workflow.TriagePrefix}{GetSeverityPrefix(issue.Severity)}Issue <{issueUrl}|{issue.Id}>: {issueSummary}{workflow.TriageSuffix}";
+			bool closed = !spans.Any(x => x.NextSuccess == null && x.LastFailure.Annotations.WorkflowId != null);
+			if (closed) // Thread may be shared by multiple workflows
+			{
+				text = $"~{text}~";
+			}
+
+			// Get the suspects for the issue
+			List<IIssueSuspect> suspects = await _issueService.Collection.FindSuspectsAsync(issue);
+
+			(MessageStateDocument state, bool isNew) = await SendOrUpdateMessageAsync(triageChannel, eventId, null, text);
+			SlackMessageId threadId = state.MessageId;
+
+			if (isNew)
+			{
+				// Create the summary text
+				List<ILogEvent> events = new List<ILogEvent>();
+				List<ILogEventData> eventDataItems = new List<ILogEventData>();
+
+				if (span.FirstFailure.LogId != null)
+				{
+					LogId logId = span.FirstFailure.LogId.Value;
+					ILogFile? logFile = await _logFileService.GetLogFileAsync(logId, CancellationToken.None);
+					if (logFile != null)
+					{
+						events = await _logFileService.FindEventsAsync(logFile, span.Id, 0, 50);
+						if (events.Any(x => x.Severity == EventSeverity.Error))
+						{
+							events.RemoveAll(x => x.Severity == EventSeverity.Warning);
+						}
+
+						List<string> eventStrings = new List<string>();
+						for (int idx = 0; idx < Math.Min(events.Count, 3); idx++)
+						{
+							ILogEventData data = await _logFileService.GetEventDataAsync(logFile, events[idx].LineIndex, events[idx].LineCount);
+							eventDataItems.Add(data);
+						}
+					}
+				}
+
+				SlackMessage message = new SlackMessage();
+
+				message.Blocks.Add(new SectionBlock(new TextObject($"From {FormatJobStep(span.FirstFailure, span.NodeName)}:")));
+				foreach (ILogEventData eventDataItem in eventDataItems)
+				{
+					message.Blocks.Add(new SectionBlock(new TextObject(QuoteText(eventDataItem.Message, MaxLineLength))));
+				}
+				if (events.Count > eventDataItems.Count)
+				{
+					message.Blocks.Add(new SectionBlock(new TextObject("```...```")));
+				}
+
+				SlackMessageId summaryId = await _slackClient.PostMessageToThreadAsync(threadId, message);
+
+				// Permalink to the summary text so we link inside the thread rather than just to the original message
+				string? permalink = await _slackClient.GetPermalinkAsync(summaryId);
+				await UpdateMessageStateAsync(state.Id, state.MessageId, permalink);
+
+				_issueService.Collection.GetLogger(issue.Id).LogInformation("Created Slack thread: {SlackLink}", permalink);
+				if (permalink != null)
+				{
+					try
+					{
+						await _issueService.UpdateIssueAsync(issue.Id, workflowThreadUrl: new Uri(permalink));
+					}
+					catch (Exception Ex)
+					{
+						_issueService.Collection.GetLogger(issue.Id).LogInformation("Error associating workflow thread with issue, bad URI format? {ErrorMessage}", Ex.Message);
+					}					
+				}				
+			}
+
+			// Post a message containing the controls and status
+			{
+				SlackMessage message = new SlackMessage();
+
+				if (workflow.TriageInstructions != null)
+				{
+					message.AddSection(workflow.TriageInstructions);
+				}
+
+				if (!closed)
+				{
+					ActionsBlock actions = message.AddActions();
+					actions.AddButton("Assign to Me", value: $"issue_{issue.Id}_ack", style: ButtonStyle.Primary);
+					actions.AddButton("Not Me", value: $"issue_{issue.Id}_decline", style: ButtonStyle.Danger);
+					actions.AddButton("Mark Fixed", value: $"issue_{issue.Id}_markfixed");
+
+					string? context = null;
+					if (issue.OwnerId != null)
+					{
+						string user = await FormatNameAsync(issue.OwnerId.Value);
+						if (issue.AcknowledgedAt == null)
+						{
+							context = $"Assigned to {user} (unacknowledged).";
+						}
+						else
+						{
+							context = $"Acknowledged by {user}.";
+						}
+					}
+					else if (suspects.Any(x => x.DeclinedAt != null))
+					{
+						HashSet<UserId> userIds = new HashSet<UserId>();
+						foreach (IIssueSuspect suspect in suspects)
+						{
+							if (suspect.DeclinedAt != null)
+							{
+								userIds.Add(suspect.AuthorId);
+							}
+						}
+
+						List<string> users = new List<string>();
+						foreach (UserId userId in userIds)
+						{
+							users.Add(await FormatNameAsync(userId));
+						}
+						users.Sort(StringComparer.OrdinalIgnoreCase);
+
+						context = $"Declined by {StringUtils.FormatList(users)}.";
+					}
+
+					if (context != null)
+					{
+						message.AddContext(context);
+					}
+				}
+
+				if (message.Blocks.Count == 0)
+				{
+					message.AddSection("Issue has been closed.");
+				}
+
+				await SendOrUpdateMessageToThreadAsync(triageChannel, $"{eventId}_buttons", null, threadId, message);
+			}
+
+			bool notifyTriageAlias = false;
+			if (isNew)
+			{
+				// If it has an owner, show that
+				HashSet<UserId> inviteUserIds = new HashSet<UserId>();
+				if (issue.OwnerId != null)
+				{
+					string mention = await FormatMentionAsync(issue.OwnerId.Value, workflow.AllowMentions);
+
+					string changes = String.Join(", ", suspects.Where(x => x.AuthorId == issue.OwnerId).Select(x => FormatChange(x.Change)));
+					if (changes.Length > 0)
+					{
+						mention += $" ({changes})";
+					}
+
+					await _slackClient.PostMessageToThreadAsync(threadId, $"Assigned to {mention}");
+					inviteUserIds.Add(issue.OwnerId.Value);
+				}
+				else
+				{
+					IGrouping<UserId, IIssueSuspect>[] suspectGroups = suspects.GroupBy(x => x.AuthorId).ToArray();
+					if (suspectGroups.Length > 0 && suspectGroups.Length <= workflow.MaxMentions)
+					{
+						List<string> suspectList = new List<string>();
+						foreach (IGrouping<UserId, IIssueSuspect> suspectGroup in suspectGroups)
+						{
+							string mention = await FormatMentionAsync(suspectGroup.Key, workflow.AllowMentions);
+							string changes = String.Join(", ", suspectGroup.Select(x => FormatChange(x.Change)));
+							suspectList.Add($"{mention} ({changes})");
+							inviteUserIds.Add(suspectGroup.Key);
+						}
+
+						string suspectMessage = $"Possibly {StringUtils.FormatList(suspectList, "or")}.";
+						await _slackClient.PostMessageToThreadAsync(threadId, suspectMessage);
+					}
+					else
+					{
+						notifyTriageAlias = true;
+					}
+				}
+
+				if (_environment.IsProduction() && workflow.AllowMentions)
+				{
+					await InviteUsersAsync(state.Channel, inviteUserIds, workflow.InviteRestrictedUsers);
+				}
+			}
+
+			if (workflow.EscalateAlias != null && workflow.EscalateTimes.Count > 0)
+			{
+				DateTime escalateTime = _clock.UtcNow + TimeSpan.FromMinutes(workflow.EscalateTimes[0]);
+				if (await _redisService.GetDatabase().SortedSetAddAsync(_escalateIssues, issue.Id, (escalateTime - DateTime.UnixEpoch).TotalSeconds, StackExchange.Redis.When.NotExists))
+				{
+					_logger.LogInformation("First escalation time for issue {IssueId} is {Time}", issue.Id, escalateTime);
+				}
+			}
+
+			if ((workflow.TriageAlias != null || workflow.TriageTypeAliases != null) && issue.OwnerId == null && (suspects.All(x => x.DeclinedAt != null) || notifyTriageAlias) && !closed)
+			{
+				string? triageAlias;
+					
+				if (workflow.TriageTypeAliases == null || issue.Fingerprints.Count == 0 || !workflow.TriageTypeAliases.TryGetValue(issue.Fingerprints[0].Type, out triageAlias))
+				{
+					triageAlias = workflow.TriageAlias;
+				}
+
+				if (triageAlias != null)
+				{
+					string triageMessage = $"(cc {FormatUserOrGroupMention(triageAlias)} for triage).";
+					await SendOrUpdateMessageToThreadAsync(triageChannel, eventId + "_triage", null, threadId, triageMessage);
+				}
+			}
+
+			if (issue.AcknowledgedAt != null)
+			{
+				await _slackClient.AddReactionAsync(threadId, "eyes");
+			}
+			else
+			{
+				await _slackClient.RemoveReactionAsync(threadId, "eyes");
+			}
+
+			if (issue.QuarantinedByUserId != null)
+			{
+				await _slackClient.AddReactionAsync(threadId, "mask");
+			}
+			else
+			{
+				await _slackClient.RemoveReactionAsync(threadId, "mask");
+			}
+
+			IIssueSpan? fixFailedSpan = null;
+			if (issue.FixChange != null)
+			{
+				fixFailedSpan = GetFixFailedSpan(issue, spans);
+
+				if (fixFailedSpan == null)
+				{
+					string fixedEventId = $"issue_{issue.Id}_fixed_{issue.FixChange}";
+					string fixedMessage = $"Marked as fixed in {FormatChange(issue.FixChange.Value)}";
+					await PostSingleMessageToThreadAsync(triageChannel, fixedEventId, threadId, fixedMessage);
+				}
+				else
+				{
+					string fixFailedEventId = $"issue_{issue.Id}_fixfailed_{issue.FixChange}";
+					string fixFailedMessage = $"Issue not fixed by {FormatChange(issue.FixChange.Value)}; see {FormatJobStep(fixFailedSpan.LastFailure, fixFailedSpan.NodeName)} at CL {fixFailedSpan.LastFailure.Change} in {fixFailedSpan.StreamName}.";
+					if (issue.OwnerId.HasValue)
+					{
+						string mention = await FormatMentionAsync(issue.OwnerId.Value, workflow.AllowMentions);
+						fixFailedMessage += $" ({mention})";
+					}
+					await PostSingleMessageToThreadAsync(triageChannel, fixFailedEventId, threadId, fixFailedMessage);
+				}
+
+				if (fixFailedSpan == null)
+				{
+					foreach (IIssueStream stream in issue.Streams)
+					{
+						if ((stream.MergeOrigin ?? false) && !(stream.ContainsFix ?? false))
+						{
+							string streamName = spans.FirstOrDefault(x => x.StreamId == stream.StreamId)?.StreamName ?? stream.StreamId.ToString();
+							string missingEventId = $"issue_{issue.Id}_fixmissing_{issue.FixChange}_{stream.StreamId}";
+							string missingMessage = $"Note: Fix may need manually merging to {streamName}";
+							await PostSingleMessageToThreadAsync(triageChannel, missingEventId, threadId, missingMessage);
+						}
+					}
+				}
+			}
+
+			if (fixFailedSpan != null)
+			{
+				await _slackClient.AddReactionAsync(threadId, "x");
+			}
+			else
+			{
+				await _slackClient.RemoveReactionAsync(threadId, "x");
+			}
+
+			if (issue.ResolvedAt != null && fixFailedSpan == null)
+			{
+				await _slackClient.AddReactionAsync(threadId, "tick");
+			}
+			else
+			{
+				await _slackClient.RemoveReactionAsync(threadId, "tick");
+			}
+
+			if (issue.ExternalIssueKey != null)
+			{
+				string extIssueEventId = $"issue_{issue.Id}_ext_{issue.ExternalIssueKey}";
+				string extIssueMessage = $"Linked to issue {FormatExternalIssue(issue.ExternalIssueKey)}";
+				await PostSingleMessageToThreadAsync(triageChannel, extIssueEventId, threadId, extIssueMessage);
+			}
+		}
+
+		async Task PostSingleMessageToThreadAsync(string recipient, string eventId, SlackMessageId threadId, string message)
+		{
+			(MessageStateDocument state, bool isNew) = await AddOrUpdateMessageStateAsync(recipient, eventId, null, "", null);
+			if (isNew)
+			{
+				SlackMessageId messageId = await _slackClient.PostMessageToThreadAsync(threadId, message);
+				await UpdateMessageStateAsync(state.Id, messageId);
+			}
+		}
+
+		async Task NotifyIssueUpdatedAsync(GlobalConfig globalConfig, IUser user, IIssue issue, IIssueDetails details)
+		{
+			string? slackUserId = await GetSlackUserId(user);
+			if (slackUserId == null)
+			{
+				return;
+			}
+
+			await SendIssueMessageAsync(globalConfig, slackUserId, issue, details, user.Id, DefaultAllowMentions);
+		}
+
+		Uri GetJobUrl(JobId jobId)
+		{
+			return new Uri(_settings.DashboardUrl, $"job/{jobId}");
+		}
+
+		Uri GetStepUrl(JobId jobId, SubResourceId stepId)
+		{
+			return new Uri(_settings.DashboardUrl, $"job/{jobId}?step={stepId}");
+		}
+
+		Uri GetIssueUrl(IIssue issue, IIssueStep step)
+		{
+			return new Uri(_settings.DashboardUrl, $"job/{step.JobId}?step={step.StepId}&issue={issue.Id}");
+		}
+
+		async Task SendIssueMessageAsync(GlobalConfig globalConfig, string recipient, IIssue issue, IIssueDetails details, UserId? userId, bool allowMentions)
+		{
+			using IDisposable scope = _logger.BeginScope("SendIssueMessageAsync (User: {SlackUser}, Issue: {IssueId})", recipient, issue.Id);
+
+			SlackAttachment attachment = new SlackAttachment();
+			attachment.Color = ErrorColor;
+
+			Uri issueUrl = _settings.DashboardUrl;
+			if (details.Steps.Count > 0)
+			{
+				issueUrl = GetIssueUrl(issue, details.Steps[0]);
+			}
+			attachment.AddSection($"*<{issueUrl}|Issue #{issue.Id}: {issue.Summary}>*");
+
+			string streamList = StringUtils.FormatList(details.Spans.Select(x => $"*{x.StreamName}*").Distinct(StringComparer.OrdinalIgnoreCase).OrderBy(x => x, StringComparer.OrdinalIgnoreCase));
+			StringBuilder summaryBuilder = new StringBuilder($"Occurring in {streamList}.");
+
+			IIssueSpan span = details.Spans[0];
+			WorkflowId? workflowId = span.LastFailure.Annotations.WorkflowId;
+			if (workflowId != null)
+			{
+				StreamConfig? streamConfig;
+				if(globalConfig.TryGetStream(span.StreamId, out streamConfig) && streamConfig.TryGetWorkflow(workflowId.Value, out WorkflowConfig? workflow) && workflow.TriageChannel != null && workflow.AllowMentions)
+				{
+					MessageStateDocument? state = await GetMessageStateAsync(workflow.TriageChannel, GetTriageThreadEventId(issue.Id));
+					if (state != null)
+					{
+						summaryBuilder.Append($" See *<{state.Permalink}|discussion thread>*.");
+					}
+				}
+			}
+			attachment.AddSection(summaryBuilder.ToString());
+
+			IIssueSpan? lastSpan = details.Spans.OrderByDescending(x => x.LastFailure.StepTime).FirstOrDefault();
+			if (lastSpan != null && lastSpan.LastFailure.LogId != null)
+			{
+				LogId logId = lastSpan.LastFailure.LogId.Value;
+				ILogFile? logFile = await _logFileService.GetLogFileAsync(logId, CancellationToken.None);
+				if(logFile != null)
+				{
+					List<ILogEvent> events = await _logFileService.FindEventsAsync(logFile, lastSpan.Id, 0, 20);
+					if (events.Any(x => x.Severity == EventSeverity.Error))
+					{
+						events.RemoveAll(x => x.Severity == EventSeverity.Warning);
+					}
+
+					for (int idx = 0; idx < Math.Min(events.Count, 3); idx++)
+					{
+						ILogEventData data = await _logFileService.GetEventDataAsync(logFile, events[idx].LineIndex, events[idx].LineCount);
+						attachment.AddSection(QuoteText(data.Message));
+					}
+					if (events.Count > 3)
+					{
+						attachment.AddSection("```...```");
+					}
+				}
+			}
+
+			if (issue.FixChange != null)
+			{
+				IIssueStep? fixFailedStep = issue.FindFixFailedStep(details.Spans);
+
+				string text;
+				if (issue.FixChange.Value < 0)
+				{
+					text = ":tick: Marked as a systemic issue.";
+				}
+				else if (fixFailedStep != null)
+				{
+					Uri fixFailedUrl = new Uri(_settings.DashboardUrl, $"job/{fixFailedStep.JobId}?step={fixFailedStep.StepId}&issue={issue.Id}");
+					text = $":cross: Marked fixed in *CL {issue.FixChange.Value}*, but seen again at *<{fixFailedUrl}|CL {fixFailedStep.Change}>*";
+				}
+				else
+				{
+					text = $":tick: Marked fixed in *CL {issue.FixChange.Value}*.";
+				}
+				attachment.AddSection(text);
+			}
+			else if (userId != null && issue.OwnerId == userId)
+			{
+				if (issue.AcknowledgedAt.HasValue)
+				{
+					attachment.AddSection($":+1: Acknowledged at {FormatSlackTime(issue.AcknowledgedAt.Value)}");
+				}
+				else
+				{
+					if (issue.NominatedById != null)
+					{
+						string mention = await FormatMentionAsync(issue.NominatedById.Value, allowMentions);
+						string text = $"You were nominated to fix this issue by {mention} at {FormatSlackTime(issue.NominatedAt ?? DateTime.UtcNow)}";
+						attachment.AddSection(text);
+					}
+					else
+					{
+						List<int> changes = details.Suspects.Where(x => x.AuthorId == userId).Select(x => x.Change).OrderBy(x => x).ToList();
+						if (changes.Count > 0)
+						{
+							string text = $"Horde has determined that {StringUtils.FormatList(changes.Select(x => $"CL {x}"), "or")} is the most likely cause for this issue.";
+							attachment.AddSection(text);
+						}
+					}
+
+					ActionsBlock actions = new ActionsBlock();
+					actions.AddButton("Acknowledge", value: $"issue_{issue.Id}_ack_{userId}", style: ButtonStyle.Primary);
+					actions.AddButton("Not Me", value: $"issue_{issue.Id}_decline_{userId}", style: ButtonStyle.Danger);
+					attachment.Blocks.Add(actions);
+				}
+			}
+			else if (issue.OwnerId != null)
+			{
+				string ownerMention = await FormatMentionAsync(issue.OwnerId.Value, allowMentions);
+				if (issue.AcknowledgedAt.HasValue)
+				{
+					attachment.AddSection($":+1: Acknowledged by {ownerMention} at {FormatSlackTime(issue.AcknowledgedAt.Value)}");
+				}
+				else if (issue.NominatedById == null)
+				{
+					attachment.AddSection($"Assigned to {ownerMention}");
+				}
+				else if (issue.NominatedById == userId)
+				{
+					attachment.AddSection($"You nominated {ownerMention} to fix this issue.");
+				}
+				else
+				{
+					attachment.AddSection($"{ownerMention} was nominated to fix this issue by {await FormatMentionAsync(issue.NominatedById.Value, allowMentions)}");
+				}
+			}
+			else if (userId != null)
+			{
+				IIssueSuspect? suspect = details.Suspects.FirstOrDefault(x => x.AuthorId == userId);
+				if (suspect != null)
+				{
+					if (suspect.DeclinedAt != null)
+					{
+						attachment.AddSection($":downvote: Declined at {FormatSlackTime(suspect.DeclinedAt.Value)}");
+					}
+					else
+					{
+						attachment.AddSection("Please check if any of your recently submitted changes have caused this issue.");
+
+						ActionsBlock actions = attachment.AddActions();
+						actions.AddButton("Will Fix", value: $"issue_{issue.Id}_accept_{userId}", style: ButtonStyle.Primary);
+						actions.AddButton("Not Me", value: $"issue_{issue.Id}_decline_{userId}", style: ButtonStyle.Danger);
+					}
+				}
+			}
+			else if (details.Suspects.Count > 0)
+			{
+				List<string> declinedLines = new List<string>();
+				foreach (IIssueSuspect suspect in details.Suspects)
+				{
+					if (!details.Issue.Promoted)
+					{
+						declinedLines.Add($"Possibly {await FormatNameAsync(suspect.AuthorId)} (CL {suspect.Change})");
+					}
+					else if (suspect.DeclinedAt == null)
+					{
+						declinedLines.Add($":heavy_minus_sign: Ignored by {await FormatNameAsync(suspect.AuthorId)} (CL {suspect.Change})");
+					}
+					else
+					{
+						declinedLines.Add($":downvote: Declined by {await FormatNameAsync(suspect.AuthorId)} at {FormatSlackTime(suspect.DeclinedAt.Value)} (CL {suspect.Change})");
+					}
+				}
+				attachment.AddSection(String.Join("\n", declinedLines));
+			}
+
+			if (IsRecipientAllowed(recipient, "issue update"))
+			{
+				await SendOrUpdateMessageAsync(recipient, GetIssueEventId(issue, recipient), userId, attachment);
+			}
+		}
+
+		static string GetIssueEventId(IIssue issue)
+		{
+			return $"issue_{issue.Id}";
+		}
+
+		static string GetIssueEventId(IIssue issue, string recipient)
+		{
+			return $"issue_{issue.Id}_for_{recipient}";
+		}
+
+		async Task<string> FormatNameAsync(UserId userId)
+		{
+			IUser? user = await _userCollection.GetUserAsync(userId);
+			if (user == null)
+			{
+				return $"User {userId}";
+			}
+			return user.Name;
+		}
+
+		static string FormatChange(int change) => $"<ugs://change?number={change}|CL {change}>";
+
+		string FormatJobStep(IIssueStep step, string name) => $"<{GetJobUrl(step.JobId)}|{step.JobName}> / <{GetStepUrl(step.JobId, step.StepId)}|{name}>";
+
+		string FormatExternalIssue(string key)
+		{
+			string? url = _externalIssueService.GetIssueUrl(key);
+			if (url == null)
+			{
+				return key;
+			}
+			else
+			{
+				return $"<{_externalIssueService.GetIssueUrl(key)}|{key}>";
+			}
+		}
+
+		static string FormatUserOrGroupMention(string userOrGroup)
+		{
+			if (userOrGroup.StartsWith("S", StringComparison.OrdinalIgnoreCase))
+			{
+				return $"<!subteam^{userOrGroup}>";
+			}
+			else
+			{
+				return $"<@{userOrGroup}>";
+			}
+		}
+
+		async Task<string> FormatMentionAsync(UserId userId, bool allowMentions)
+		{
+			IUser? user = await _userCollection.GetUserAsync(userId);
+			if (user == null)
+			{
+				return $"User {userId}";
+			}
+
+			string? slackUserId = await GetSlackUserId(user);
+			if (slackUserId == null)
+			{
+				return user.Login;
+			}
+
+			if (_allowUsers == null || !_allowUsers.Contains(slackUserId))
+			{
+				if (!_environment.IsProduction() || !allowMentions)
+				{
+					return $"{user.Name} [{slackUserId}]";
+				}
+			}
+
+			return $"<@{slackUserId}>";
+		}
+
+		static string FormatSlackTime(DateTimeOffset time)
+		{
+			return $"<!date^{time.ToUnixTimeSeconds()}^{{time}}|{time}>";
+		}
+
+		class ReportBlock
+		{
+			[JsonPropertyName("h")]
+			public bool TemplateHeader { get; set; }
+
+			[JsonPropertyName("s")]
+			public StreamId StreamId { get; set; }
+
+			[JsonPropertyName("t")]
+			public TemplateId TemplateId { get; set; }
+
+			[JsonPropertyName("i")]
+			public List<int> IssueIds { get; set; } = new List<int>();
+		}
+
+		class ReportState
+		{
+			[JsonPropertyName("tm")]
+			public DateTime Time { get; set; }
+
+			[JsonPropertyName("msg")]
+			public List<ReportBlock> Blocks { get; set; } = new List<ReportBlock>();
+		}
+
+		static string GetReportEventId(StreamId streamId, WorkflowId workflowId) => $"issue-report:{streamId}:{workflowId}";
+
+		static string GetReportBlockEventId(string ts, int idx) => $"issue-report-block:{ts}:{idx}";
+
+		/// <inheritdoc/>
+		public async Task SendIssueReportAsync(IssueReportGroup group)
+		{
+			GlobalConfig globalConfig = _globalConfig.CurrentValue;
+
+			foreach (IssueReport report in group.Reports.OrderBy(x => x.WorkflowId.Id.Text).ThenBy(x => x.StreamId.Id.Text))
+			{
+				await SendIssueReportForStreamAsync(globalConfig, group.Channel, group.Time, report);
+			}
+
+			SlackMessage message = new SlackMessage();
+			message.AddDivider();
+
+			await SendMessageAsync(group.Channel, message);
+		}
+
+		async Task SendIssueReportForStreamAsync(GlobalConfig globalConfig, string channel, DateTime time, IssueReport report)
+		{
+			const int MaxIssuesPerMessage = 8;
+
+			StreamConfig? streamConfig;
+			if (!globalConfig.TryGetStream(report.StreamId, out streamConfig))
+			{
+				return;
+			}
+
+			ReportState state = new ReportState();
+			state.Time = time;
+
+			List<List<(IIssue, IIssueSpan?, bool)>> issuesByBlock = new List<List<(IIssue, IIssueSpan?, bool)>>();
+			if (report.GroupByTemplate)
+			{
+				Dictionary<int, IIssue> issueIdToInfo = new Dictionary<int, IIssue>();
+				foreach (IIssue issue in report.Issues)
+				{
+					issueIdToInfo[issue.Id] = issue;
+				}
+
+				foreach (IGrouping<TemplateId, IIssueSpan> group in report.IssueSpans.GroupBy(x => x.TemplateRefId).OrderBy(x => x.Key.ToString()))
+				{
+					TemplateRefConfig? templateConfig;
+					if (!streamConfig.TryGetTemplate(group.Key, out templateConfig))
+					{
+						continue;
+					}
+					if (!IsIssueOpenForWorkflow(report.StreamId, group.Key, report.WorkflowId, group))
+					{
+						continue;
+					}
+
+					Dictionary<IIssue, IIssueSpan> pairs = new Dictionary<IIssue, IIssueSpan>();
+					foreach (IIssueSpan span in group)
+					{
+						if (issueIdToInfo.TryGetValue(span.IssueId, out IIssue? issue))
+						{
+							pairs[issue] = span;
+						}
+					}
+
+					if (pairs.Count > 0)
+					{
+						bool templateHeader = true;
+						foreach (IReadOnlyList<KeyValuePair<IIssue, IIssueSpan>> batch in pairs.OrderBy(x => x.Key.Id).Batch(MaxIssuesPerMessage))
+						{
+							ReportBlock block = new ReportBlock();
+							block.TemplateHeader = templateHeader;
+							block.StreamId = report.StreamId;
+							block.TemplateId = group.Key;
+							block.IssueIds.AddRange(batch.Select(x => x.Key.Id));
+							state.Blocks.Add(block);
+
+							issuesByBlock.Add(batch.Select(x => (x.Key, (IIssueSpan?)x.Value, true)).ToList());
+							templateHeader = false;
+						}
+					}
+				}
+			}
+			else
+			{
+				foreach (IReadOnlyList<IIssue> batch in report.Issues.OrderByDescending(x => x.Id).Batch(MaxIssuesPerMessage))
+				{
+					ReportBlock block = new ReportBlock();
+					block.StreamId = report.StreamId;
+					block.IssueIds.AddRange(batch.Select(x => x.Id));
+					state.Blocks.Add(block);
+
+					issuesByBlock.Add(batch.Select(x => (x, report.IssueSpans.FirstOrDefault(y => y.IssueId == x.Id), true)).ToList());
+				}
+			}
+
+			SlackMessage headerMessage = new SlackMessage();
+			headerMessage.AddHeader($"Summary for {streamConfig.Name}");
+
+			SlackMessageId? messageId = await SendMessageAsync(channel, headerMessage);
+			if (messageId != null)
+			{
+				string reportEventId = GetReportEventId(streamConfig.Id, report.WorkflowId);
+				string json = JsonSerializer.Serialize(state, _jsonSerializerOptions);
+				await AddOrUpdateMessageStateAsync(channel, reportEventId, null, json, messageId);
+
+				if (state.Blocks.Count == 0)
+				{
+					string header = ":tick: No issues open.";
+					await SendMessageAsync(channel, header);
+				}
+
+				for (int idx = 0; idx < state.Blocks.Count; idx++)
+				{
+					string blockEventId = GetReportBlockEventId(messageId.Ts, idx);
+					await UpdateReportBlockAsync(channel, blockEventId, time, streamConfig, state.Blocks[idx].TemplateId, issuesByBlock[idx], report.TriageChannel, state.Blocks[idx].TemplateHeader);
+				}
+
+				if (report.WorkflowStats.NumSteps > 0)
+				{
+					double totalPct = (report.WorkflowStats.NumPassingSteps * 100.0) / report.WorkflowStats.NumSteps;
+					string header = $"*{totalPct:0.0}%* of build steps ({report.WorkflowStats.NumPassingSteps:n0}/{report.WorkflowStats.NumSteps:n0}) succeeded since last status update.";
+					await SendMessageAsync(channel, header);
+				}
+			}
+		}
+
+		static bool IsIssueOpenForWorkflow(StreamId streamId, TemplateId templateId, WorkflowId workflowId, IEnumerable<IIssueSpan> spans)
+		{
+			foreach (IIssueSpan span in spans)
+			{
+				if (span.NextSuccess == null && span.LastFailure.Annotations.WorkflowId == workflowId)
+				{
+					if ((streamId.IsEmpty || span.StreamId == streamId) && (templateId.IsEmpty || span.TemplateRefId == templateId))
+					{
+						return true;
+					}
+				}
+			}
+			return false;
+		}
+
+		async Task UpdateReportsAsync(GlobalConfig globalConfig, IIssue issue, IReadOnlyList<IIssueSpan> spans)
+		{
+			_logger.LogInformation("Checking for report updates to issue {IssueId}", issue.Id);
+
+			HashSet<string> reportEventIds = new HashSet<string>(StringComparer.Ordinal);
+			foreach (IIssueSpan span in spans)
+			{
+				WorkflowId? workflowId = span.LastFailure.Annotations.WorkflowId;
+				if (workflowId == null)
+				{
+					continue;
+				}
+
+				string reportEventId = GetReportEventId(span.StreamId, workflowId.Value);
+				if (!reportEventIds.Add(reportEventId))
+				{
+					continue;
+				}
+
+				StreamConfig? streamConfig;
+				if (!globalConfig.TryGetStream(span.StreamId, out streamConfig))
+				{
+					continue;
+				}
+
+				WorkflowConfig? workflowConfig;
+				if (!streamConfig.TryGetWorkflow(workflowId.Value, out workflowConfig) || workflowConfig.ReportChannel == null)
+				{
+					continue;
+				}
+
+				MessageStateDocument? messageState = await GetMessageStateAsync(workflowConfig.ReportChannel, reportEventId);
+				if (messageState == null)
+				{
+					continue;
+				}
+
+				ReportState? state;
+				try
+				{
+					state = JsonSerializer.Deserialize<ReportState>(messageState.Digest, _jsonSerializerOptions);
+				}
+				catch
+				{
+					continue;
+				}
+
+				if (state == null)
+				{
+					continue;
+				}
+
+				_logger.LogInformation("Updating report {EventId}", reportEventId);
+
+				for (int idx = 0; idx < state.Blocks.Count; idx++)
+				{
+					ReportBlock block = state.Blocks[idx];
+					if (block.IssueIds.Contains(issue.Id))
+					{
+						string blockEventId = GetReportBlockEventId(messageState.Ts, idx);
+						_logger.LogInformation("Updating report block {EventId}", blockEventId);
+
+						List<(IIssue, IIssueSpan?, bool)> issues = new List<(IIssue, IIssueSpan?, bool)>();
+						foreach (int issueId in block.IssueIds)
+						{
+							IIssueDetails? details = await _issueService.GetIssueDetailsAsync(issueId);
+							if (details != null)
+							{
+								IIssueSpan? otherSpan = details.Spans.FirstOrDefault(x => x.TemplateRefId == block.TemplateId);
+								if (otherSpan == null && details.Spans.Count > 0)
+								{
+									otherSpan = details.Spans[0];
+								}
+
+								bool open = IsIssueOpenForWorkflow(streamConfig.Id, block.TemplateId, workflowId.Value, details.Spans);
+								issues.Add((details.Issue, otherSpan, open));
+							}
+						}
+
+						await UpdateReportBlockAsync(workflowConfig.ReportChannel, blockEventId, state.Time, streamConfig, block.TemplateId, issues, workflowConfig.TriageChannel, block.TemplateHeader);
+					}
+				}
+			}
+		}
+
+		async Task UpdateReportBlockAsync(string channel, string eventId, DateTime reportTime, StreamConfig streamConfig, TemplateId templateId, List<(IIssue, IIssueSpan?, bool)> issues, string? triageChannel, bool templateHeader)
+		{
+			StringBuilder body = new StringBuilder();
+
+			if (templateHeader && !templateId.IsEmpty)
+			{
+				TemplateRefConfig? templateConfig;
+				if (streamConfig.TryGetTemplate(templateId, out templateConfig))
+				{
+					JobsTabConfig? tab = streamConfig.Tabs.OfType<JobsTabConfig>().FirstOrDefault(x => x.Templates != null && x.Templates.Contains(templateId));
+					if (tab != null)
+					{
+						Uri templateUrl = new Uri(_settings.DashboardUrl, $"stream/{streamConfig.Id}?tab={tab.Title}&template={templateId}");
+						body.Append($"*<{templateUrl}|{templateConfig.Name}>*:");
+					}
+				}
+			}
+
+			foreach ((IIssue issue, IIssueSpan? span, bool open) in issues)
+			{
+				if (body.Length > 0)
+				{
+					body.Append('\n');
+				}
+
+				string text = await FormatIssueAsync(issue, span, triageChannel, reportTime, open);
+				body.Append(text);
+			}
+
+			SlackMessage message = body.ToString();
+			message.UnfurlLinks = false;
+			message.UnfurlMedia = false;
+			await SendOrUpdateMessageAsync(channel, eventId, null, message);
+		}
+
+		string GetSeverityPrefix(IssueSeverity severity)
+		{
+			return (severity == IssueSeverity.Warning) ? _settings.SlackWarningPrefix : _settings.SlackErrorPrefix;
+		}
+
+		async ValueTask<string> FormatIssueAsync(IIssue issue, IIssueSpan? span, string? triageChannel, DateTime reportTime, bool open)
+		{
+			Uri issueUrl = _settings.DashboardUrl;
+			if (span != null)
+			{
+				IIssueStep lastFailure = span.LastFailure;
+				issueUrl = new Uri(issueUrl, $"job/{lastFailure.JobId}?step={lastFailure.StepId}&issue={issue.Id}");
+			}
+
+			IUser? owner = null;
+			if (issue.OwnerId != null)
+			{
+				owner = await _userCollection.GetCachedUserAsync(issue.OwnerId.Value);
+			}
+
+			string status = "*Unassigned*";
+			if (issue.FixChange != null)
+			{
+				if (issue.Streams.Any(x => x.FixFailed ?? false))
+				{
+					status = "*Fix failed*";
+				}
+				else if (span != null && issue.Streams.Any(x => x.StreamId == span.StreamId && (x.MergeOrigin ?? false) && !(x.ContainsFix ?? false)))
+				{
+					status = "*Fix needs merging*";
+				}
+				else
+				{
+					status = $"Fixed in CL {issue.FixChange.Value}";
+				}
+			}
+			else if (issue.OwnerId != null)
+			{
+				if (owner == null)
+				{
+					status = $"Assigned to user {issue.OwnerId.Value}";
+				}
+				else
+				{
+					status = $"Assigned to {owner.Name}";
+				}
+				if (issue.AcknowledgedAt == null)
+				{
+					status = $"{status} (pending)";
+				}
+			}
+
+			if (issue.QuarantinedByUserId != null)
+			{
+				status = $"{status} - *Quarantined*";
+			}
+
+			string prefix = GetSeverityPrefix(issue.Severity);
+			StringBuilder body = new StringBuilder($"{prefix}Issue *<{issueUrl}|{issue.Id}>");
+
+			if (triageChannel != null)
+			{
+				MessageStateDocument? state = await GetMessageStateAsync(triageChannel, GetTriageThreadEventId(issue.Id));
+				if (state != null && state.Permalink != null)
+				{
+					body.Append($" (<{state.Permalink}|Thread>)");
+				}
+			}
+
+			string? summary = issue.UserSummary ?? issue.Summary;
+			body.Append($"*: {summary} [{FormatReadableTimeSpan(reportTime - issue.CreatedAt)}]");
+
+			if (!String.IsNullOrEmpty(issue.ExternalIssueKey))
+			{
+				body.Append($" ({FormatExternalIssue(issue.ExternalIssueKey)})");
+			}
+			if (open)
+			{
+				body.Append($" - {status}");
+			}
+
+			string text = body.ToString();
+			if (!open)
+			{
+				text = $"~{text}~";
+			}
+			return text;
+		}
+
+		static string FormatReadableTimeSpan(TimeSpan timeSpan)
+		{
+			double totalDays = timeSpan.TotalDays;
+			if (totalDays > 1.0)
+			{
+				return $"{totalDays:n0}d";
+			}
+
+			double totalHours = timeSpan.TotalHours;
+			if(totalHours > 1.0)
+			{
+				return $"{totalHours:n0}h";
+			}
+
+			return $"{timeSpan.TotalMinutes:n0}m";
+		}
+
+		#endregion
+
+		#region Stream updates
+
+		/// <inheritdoc/>
+		public async Task NotifyConfigUpdateAsync(Exception? ex)
+		{
+			if (String.IsNullOrEmpty(_settings.ConfigNotificationChannel))
+			{
+				return;
+			}
+
+			const string EventId = "config-update";
+
+			if (ex != null)
+			{
+				_logger.LogInformation(ex, "Sending config update failure notification: {Message}", ex.Message);
+
+				List<string> details = new List<string>();
+				if (ex is ConfigException configEx)
+				{
+					ConfigContext context = configEx.Context;
+					if (context.IncludeStack.TryPeek(out IConfigFile? blame))
+					{
+						string line = String.Empty;
+						if (configEx.InnerException is JsonException jsonEx && jsonEx.LineNumber != null)
+						{
+							line = $"({jsonEx.LineNumber})";
+						}
+
+						string file = FormatConfigFileUri(blame.Uri);
+						details.Add($"Error parsing `{file}{line}`:");
+						details.Add($"```{ex.Message}```");
+						details.Add("Include stack:\n```" + String.Join("\n", context.IncludeStack.Select(x => FormatConfigFileRevision(x))) + "```");
+
+						if (blame.Uri.Scheme == PerforceConfigSource.Scheme)
+						{
+							string blameMessage = $"Possibly due to CL {blame.Revision}";
+							if (blame.Author != null)
+							{
+								string userId = await FormatMentionAsync(blame.Author.Id, true);
+								blameMessage += $" ({userId})";
+							}
+							details.Add(blameMessage.ToString());
+						}
+					}
+				}
+
+				if (details.Count == 0)
+				{
+					details.Add(QuoteText(ex.Message));
+				}
+
+				string message = String.Join("\n", details);
+				string digest = GetMessageDigest(message);
+
+				MessageStateDocument? state = await GetMessageStateAsync(_settings.ConfigNotificationChannel, EventId);
+				if (state == null || state.Digest != digest)
+				{
+					SlackMessage header = new SlackMessage();
+					header.AddHeader($"Config Update Error");
+					await SendMessageAsync(_settings.ConfigNotificationChannel, header);
+
+					SlackMessageId? messageId = await SendMessageAsync(_settings.ConfigNotificationChannel, message);
+					await AddOrUpdateMessageStateAsync(_settings.ConfigNotificationChannel, EventId, null, digest, messageId);
+				}
+			}
+			else
+			{
+				if (await DeleteMessageStateAsync(_settings.ConfigNotificationChannel, EventId))
+				{
+					SlackMessage message = new SlackMessage();
+					message.AddSection($"*Config Update Succeeded*");
+					await SendMessageAsync(_settings.ConfigNotificationChannel, message);
+					await DeleteMessageStateAsync(_settings.ConfigNotificationChannel, EventId);
+				}
+			}
+		}
+
+		static string FormatConfigFileUri(Uri uri)
+		{
+			if (uri.Scheme == FileConfigSource.Scheme)
+			{
+				return Uri.UnescapeDataString(uri.AbsolutePath).Replace('/', Path.DirectorySeparatorChar);
+			}
+			else if (uri.Scheme == PerforceConfigSource.Scheme)
+			{
+				return uri.AbsolutePath;
+			}
+			else
+			{
+				return uri.ToString();
+			}
+		}
+
+		static string FormatConfigFileRevision(IConfigFile file)
+		{
+			Uri uri = file.Uri;
+			if (uri.Scheme == PerforceConfigSource.Scheme)
+			{
+				return $"{FormatConfigFileUri(uri)}@{file.Revision}";
+			}
+			else
+			{
+				return FormatConfigFileUri(uri);
+			}
+		}
+
+		/// <inheritdoc/>
+		public async Task NotifyConfigUpdateFailureAsync(string errorMessage, string fileName, int? change = null, IUser? author = null, string? description = null)
+		{
+			_logger.LogInformation("Sending config update failure notification for {FileName} (change: {Change}, author: {UserId})", fileName, change ?? -1, author?.Id ?? UserId.Empty);
+
+			string? slackUserId = null;
+			if (author != null)
+			{
+				slackUserId = await GetSlackUserId(author);
+				if (slackUserId == null)
+				{
+					_logger.LogWarning("Unable to identify Slack user id for {UserId}", author.Id);
+				}
+				else
+				{
+					_logger.LogInformation("Mappsed user {UserId} to Slack user {SlackUserId}", author.Id, slackUserId);
+				}
+			}
+
+			if (slackUserId != null)
+			{
+				await SendConfigUpdateFailureMessageAsync(slackUserId, errorMessage, fileName, change, slackUserId, description);
+			}
+			if (_settings.UpdateStreamsNotificationChannel != null)
+			{
+				await SendConfigUpdateFailureMessageAsync($"#{_settings.UpdateStreamsNotificationChannel}", errorMessage, fileName, change, slackUserId, description);
+			}
+		}
+
+		private async Task SendConfigUpdateFailureMessageAsync(string recipient, string errorMessage, string fileName, int? change = null, string? author = null, string? description = null)
+		{
+			string outcomeColor = ErrorColor;
+			SlackAttachment attachment = new SlackAttachment();
+			attachment.Color = outcomeColor;
+			attachment.FallbackText = $"Update Failure: {fileName}";
+
+			attachment.AddHeader($"Config Update Failure :rip:", true);
+
+			attachment.AddSection($"Horde was unable to update {fileName}");
+			attachment.AddSection(QuoteText(errorMessage));
+			if (change != null)
+			{
+				if (author != null)
+				{
+					attachment.AddSection($"Possibly due to CL: {change.Value} by <@{author}>");
+				}
+				else
+				{
+					attachment.AddSection($"Possibly due to CL: {change.Value} - (Could not determine author from P4 user)");
+				}
+				if (description != null)
+				{
+					attachment.AddSection(QuoteText(description));
+				}
+			}
+
+			await SendMessageAsync(recipient, attachment);
+		}
+
+		#endregion
+
+		#region Device notifications
+
+		/// <inheritdoc/>
+		public async Task NotifyDeviceServiceAsync(string message, IDevice? device = null, IDevicePool? pool = null, StreamConfig? streamConfig = null, IJob? job = null, IJobStep? step = null, INode? node = null, IUser? user = null)
+		{
+			string? recipient = null;
+
+			if (user != null)
+			{
+				string? slackRecipient = await GetSlackUserId(user);
+
+				if (slackRecipient == null)
+				{
+					_logger.LogError("NotifyDeviceServiceAsync - Unable to send user slack notification, user {UserId} slack user id not found", user.Id);
+					return;
+				}
+
+				recipient = slackRecipient;
+			}
+			
+			if (recipient != null)
+			{
+				_logger.LogDebug("Sending device service notification to {Recipient}", recipient);
+				await SendDeviceServiceMessage(recipient, message, device, pool, streamConfig, job, step, node, user);
+			}
+		}
+
+		/// <summary>
+		/// Creates a Slack message for a device service notification
+		/// </summary>
+		/// <param name="recipient"></param>
+		/// <param name="message"></param>
+		/// <param name="device"></param>
+		/// <param name="pool"></param>
+		/// <param name="streamConfig"></param>
+		/// <param name="job">The job that contains the step that completed.</param>
+		/// <param name="step">The job step that completed.</param>
+		/// <param name="node">The node for the job step.</param>
+		/// <param name="user">The user to notify.</param>
+		private Task SendDeviceServiceMessage(string recipient, string message, IDevice? device = null, IDevicePool? pool = null, StreamConfig? streamConfig = null, IJob? job = null, IJobStep? step = null, INode? node = null, IUser? user = null)
+		{
+
+			if (user != null)
+			{
+				return SendMessageAsync(recipient, message);
+			}
+
+			// truncate message to avoid slack error on message length
+			if (message.Length > 150)
+			{
+				message = message.Substring(0, 146) + "...";
+			}
+
+			SlackAttachment attachment = new SlackAttachment();
+							
+			attachment.FallbackText = $"{message}";
+
+			if (device != null && pool != null)
+			{
+				attachment.FallbackText += $" - Device: {device.Name} Pool: {pool.Name}";
+			}
+				
+			attachment.AddHeader(message, false);
+
+			if (streamConfig != null && job != null && step != null && node != null)
+			{
+				Uri jobStepLink = new Uri($"{_settings.DashboardUrl}job/{job.Id}?step={step.Id}");
+				Uri jobStepLogLink = new Uri($"{_settings.DashboardUrl}log/{step.LogId}");
+
+				attachment.FallbackText += $" - {streamConfig.Name} - {GetJobChangeText(job)} - {job.Name} - {node.Name}";
+				attachment.AddSection($"*<{jobStepLink}|{streamConfig.Name} - {GetJobChangeText(job)} - {job.Name} - {node.Name}>*");
+				attachment.AddSection($"<{jobStepLogLink}|View Job Step Log>");
+			}
+			else
+			{
+				attachment.FallbackText += " - No job information (Gauntlet might need to be updated in stream)";
+				attachment.AddSection("*No job information (Gauntlet might need to be updated in stream)*");
+			}
+
+			return SendMessageAsync(recipient, attachment);
+		}
+
+		#endregion
+
+		static string QuoteText(string text, int maxLength = MaxLineLength)
+		{
+			maxLength -= 6;
+			if (text.Length > maxLength)
+			{
+				int length = text.LastIndexOf('\n', maxLength);
+				if (length == 0)
+				{
+					text = text.Substring(0, maxLength - 7) + "...\n...";
+				}
+				else
+				{
+					text = text.Substring(0, length + 1) + "...";
+				}
+			}
+			return $"```{text}```";
+		}
+
+		const int MaxJobStepEvents = 5;
+
+		static string GetJobChangeText(IJob job)
+		{
+			if (job.PreflightChange == 0)
+			{
+				return $"CL {job.Change}";
+			}
+			else
+			{
+				return $"Preflight CL {job.PreflightChange} against CL {job.Change}";
+			}
+		}
+
+		static bool ShouldUpdateUser(SlackUserDocument? document)
+		{
+			if(document == null || document.Version < SlackUserDocument.CurrentVersion)
+			{
+				return true;
+			}
+
+			TimeSpan expiryTime;
+			if (document.SlackUserId == null)
+			{
+				expiryTime = TimeSpan.FromMinutes(10.0);
+			}
+			else
+			{
+				expiryTime = TimeSpan.FromDays(1.0);
+			}
+			return document.Time + expiryTime < DateTime.UtcNow;
+		}
+
+		private async Task<string?> GetSlackUserId(IUser user)
+		{
+			return (await GetSlackUser(user))?.SlackUserId;
+		}
+
+		private async Task<SlackUserDocument?> GetSlackUser(IUser user)
+		{
+			string? email = user.Email;
+			if (email == null)
+			{
+				_logger.LogWarning("Unable to find Slack user id for {UserId} ({Name}): No email address in user profile", user.Id, user.Name);
+				return null;
+			}
+
+			SlackUserDocument? userDocument;
+			if (!_userCache.TryGetValue(email, out userDocument))
+			{
+				userDocument = await _slackUsers.Find(x => x.Id == user.Id).FirstOrDefaultAsync();
+				if (userDocument == null || ShouldUpdateUser(userDocument))
+				{
+					SlackUser? userInfo = await _slackClient.FindUserByEmailAsync(email);
+					if (userDocument == null || userInfo != null)
+					{
+						userDocument = new SlackUserDocument(user.Id, userInfo);
+						await _slackUsers.ReplaceOneAsync(x => x.Id == user.Id, userDocument, new ReplaceOptions { IsUpsert = true });
+					}
+				}
+				using (ICacheEntry entry = _userCache.CreateEntry(email))
+				{
+					entry.SlidingExpiration = TimeSpan.FromMinutes(10.0);
+					entry.Value = userDocument;
+				}
+			}
+
+			return userDocument;
+		}
+
+		private async Task<SlackMessageId?> SendMessageAsync(string recipient, SlackMessage message)
+		{
+			if (!IsRecipientAllowed(recipient, message.Text))
+			{
+				return null;
+			}
+			else
+			{
+				return await _slackClient.PostMessageAsync(recipient, message);
+			}
+		}
+
+		bool IsRecipientAllowed(string recipient, string? description)
+		{
+			if (_allowUsers != null && !_allowUsers.Contains(recipient) && !recipient.StartsWith("#", StringComparison.Ordinal))
+			{
+				_logger.LogDebug("Suppressing message to {Recipient}: {Description}", recipient, description);
+				return false;
+			}
+			return true;
+		}
+
+		private async Task<(MessageStateDocument, bool)> SendOrUpdateMessageAsync(string recipient, string eventId, UserId? userId, SlackMessage message)
+		{
+			return await SendOrUpdateMessageToThreadAsync(recipient, eventId, userId, null, message);
+		}
+
+		private static string GetMessageDigest(SlackMessage message)
+		{
+			return ContentHash.MD5(JsonSerializer.Serialize(message, new JsonSerializerOptions { DefaultIgnoreCondition = JsonIgnoreCondition.WhenWritingNull })).ToString();
+		}
+
+		private async Task<(MessageStateDocument, bool)> SendOrUpdateMessageToThreadAsync(string recipient, string eventId, UserId? userId, SlackMessageId? threadId, SlackMessage message)
+		{
+			string requestDigest = GetMessageDigest(message);
+
+			MessageStateDocument? prevState = await GetMessageStateAsync(recipient, eventId);
+			if (prevState != null && prevState.Digest == requestDigest)
+			{
+				return (prevState, false);
+			}
+
+			(MessageStateDocument state, bool isNew) = await AddOrUpdateMessageStateAsync(recipient, eventId, userId, requestDigest, null);
+			if (isNew)
+			{
+				_logger.LogInformation("Sending new slack message to {SlackUser} (state: {StateId}, threadMessageId: {ThreadTs})", recipient, state.Id, threadId?.ToString() ?? "n/a");
+
+				SlackMessageId id;
+				if (threadId == null)
+				{
+					id = await _slackClient.PostMessageAsync(recipient, message);
+				}
+				else
+				{
+					id = await _slackClient.PostMessageToThreadAsync(threadId, message);
+				}
+
+				state.Channel = id.Channel;
+				state.Ts = id.Ts;
+
+				await UpdateMessageStateAsync(state.Id, id);
+			}
+			else if (!String.IsNullOrEmpty(state.Ts))
+			{
+				_logger.LogInformation("Updating existing slack message {StateId} for user {SlackUser} (messageId: {MessageId}, threadId: {ThreadTs})", state.Id, recipient, state.MessageId, threadId?.ToString() ?? "n/a");
+
+				if (threadId == null)
+				{
+					await _slackClient.UpdateMessageAsync(state.MessageId, message);
+				}
+				else
+				{
+					await _slackClient.UpdateMessageAsync(new SlackMessageId(threadId.Channel, threadId.Ts, state.Ts), message);
+				}
+			}
+			return (state, isNew);
+		}
+
+		async ValueTask EscalateAsync(CancellationToken cancellationToken)
+		{
+			GlobalConfig globalConfig = _globalConfig.CurrentValue;
+
+			DateTime utcNow = DateTime.UtcNow;
+			double time = (utcNow - DateTime.UnixEpoch).TotalSeconds;
+
+			int[] issueIds = await _redisService.GetDatabase().SortedSetRangeByScoreAsync(_escalateIssues, 0, time);
+			if (issueIds.Length > 0)
+			{
+				_logger.LogInformation("Escalating issues for {Time} ({TimeSecs})", utcNow, time);
+				foreach (int issueId in issueIds)
+				{
+					cancellationToken.ThrowIfCancellationRequested();
+					try
+					{
+						double? nextTime = await EscalateSingleIssueAsync(globalConfig, issueId, utcNow);
+						if (nextTime == null)
+						{
+							_logger.LogInformation("Cancelling escalation for issue {IssueId}", issueId);
+							await _redisService.GetDatabase().SortedSetRemoveAsync(_escalateIssues, issueId);
+						}
+						else
+						{
+							_logger.LogInformation("Next escalation for issue {IssueId} is at timestamp {Time}", issueId, nextTime.Value);
+							await _redisService.GetDatabase().SortedSetAddAsync(_escalateIssues, issueId, nextTime.Value);
+						}
+					}
+					catch (SlackException ex)
+					{
+						_logger.LogError(ex, "Slack exception while escalating issue {IssueId}; cancelling.", issueId);
+						await _redisService.GetDatabase().SortedSetRemoveAsync(_escalateIssues, issueId);
+					}
+				}
+			}
+		}
+
+		async Task<double?> EscalateSingleIssueAsync(GlobalConfig globalConfig, int issueId, DateTime utcNow)
+		{
+			IIssue? issue = await _issueService.Collection.GetIssueAsync(issueId);
+			if (issue == null)
+			{
+				return null;
+			}
+
+			List<IIssueSpan> spans = await _issueService.Collection.FindSpansAsync(issueId);
+			if (spans.Count == 0)
+			{
+				return null;
+			}
+
+			IIssueSpan span = spans[0];
+
+			WorkflowId? workflowId = span.LastFailure.Annotations.WorkflowId;
+			if (workflowId == null)
+			{
+				return null;
+			}
+
+			StreamConfig? streamConfig;
+			if (!globalConfig.TryGetStream(span.StreamId, out streamConfig))
+			{
+				return null;
+			}
+			if (!streamConfig.TryGetWorkflow(workflowId.Value, out WorkflowConfig? workflow))
+			{
+				return null;
+			}
+			if (workflow.TriageChannel == null || workflow.EscalateAlias == null || workflow.EscalateTimes.Count == 0)
+			{
+				return null;
+			}
+
+			if (!IsIssueOpenForWorkflow(streamConfig.Id, span.TemplateRefId, workflow.Id, spans))
+			{
+				return null;
+			}
+			if (issue.FixChange != null && GetFixFailedSpan(issue, spans) == null)
+			{
+				return null;
+			}
+
+			if (issue.QuarantineTimeUtc == null)
+			{
+				MessageStateDocument? state = await GetMessageStateAsync(workflow.TriageChannel, GetTriageThreadEventId(issueId));
+				if (state == null)
+				{
+					return null;
+				}
+
+				TimeSpan openTime = utcNow - issue.CreatedAt;
+
+				string openTimeStr;
+				if (openTime < TimeSpan.FromHours(1.0))
+				{
+					openTimeStr = $"{(int)openTime.TotalMinutes}m";
+				}
+				else if (openTime < TimeSpan.FromDays(1.0))
+				{
+					openTimeStr = $"{(int)openTime.TotalHours}h";
+				}
+				else
+				{
+					openTimeStr = $"{(int)openTime.TotalDays}d";
+				}
+
+				Uri issueUrl = GetIssueUrl(issue, span.FirstFailure);
+				await _slackClient.PostMessageToThreadAsync(state.MessageId, $"{FormatUserOrGroupMention(workflow.EscalateAlias)} - Issue <{issueUrl}|{issue.Id}> has not been resolved after {openTimeStr}.");
+			}
+
+			DateTime nextEscalationTime = span.FirstFailure.StepTime;
+			for (int idx = 0;;idx++)
+			{
+				if (idx >= workflow.EscalateTimes.Count)
+				{
+					nextEscalationTime = utcNow + TimeSpan.FromMinutes(workflow.EscalateTimes[^1]);
+					_logger.LogInformation("Next escalation time for issue {IssueId} is {NextEscalationTime} (now + {Time}m [{Idx}/{Count}])", issue.Id, nextEscalationTime, workflow.EscalateTimes[^1], workflow.EscalateTimes.Count, workflow.EscalateTimes.Count);
+					break;
+				}
+
+				DateTime prevEscalationTime = nextEscalationTime;
+				nextEscalationTime += TimeSpan.FromMinutes(workflow.EscalateTimes[idx]);
+
+				if (nextEscalationTime > utcNow)
+				{
+					_logger.LogInformation("Next escalation time for issue {IssueId} is {NextEscalationTime} ({PrevEscalationTime} + {Time}m [{Idx}/{Count}])", issue.Id, nextEscalationTime, prevEscalationTime, workflow.EscalateTimes[idx], idx + 1, workflow.EscalateTimes.Count);
+					break;
+				}
+			}
+
+			return (nextEscalationTime - DateTime.UnixEpoch).TotalSeconds;
+		}
+
+		/// <inheritdoc/>
+		protected override async Task ExecuteAsync(CancellationToken stoppingToken)
+		{
+			if (String.IsNullOrEmpty(_settings.SlackSocketToken))
+			{
+				_logger.LogInformation("No Slack socket token configured; will not be able to respond to interactive messages.");
+				return;
+			}
+
+			while (!stoppingToken.IsCancellationRequested)
+			{
+				try
+				{
+					Uri? webSocketUrl = await GetWebSocketUrlAsync(stoppingToken);
+					if (webSocketUrl == null)
+					{
+						_logger.LogWarning("Unable to get Slack websocket URL. Pausing before retry.");
+						await Task.Delay(TimeSpan.FromSeconds(5.0), stoppingToken);
+						continue;
+					}
+					await HandleSocketAsync(webSocketUrl, stoppingToken);
+				}
+				catch (OperationCanceledException)
+				{
+					break;
+				}
+				catch (Exception ex)
+				{
+					_logger.LogError(ex, "Exception while updating Slack socket");
+					await Task.Delay(TimeSpan.FromSeconds(5.0), stoppingToken);
+				}
+			}
+		}
+
+		/// <summary>
+		/// Get the url for opening a websocket to Slack
+		/// </summary>
+		/// <param name="stoppingToken"></param>
+		/// <returns></returns>
+		private async Task<Uri?> GetWebSocketUrlAsync(CancellationToken stoppingToken)
+		{
+			using HttpClient client = new HttpClient();
+			client.DefaultRequestHeaders.Add("Authorization", $"Bearer {_settings.SlackSocketToken}");
+
+			using FormUrlEncodedContent content = new FormUrlEncodedContent(Array.Empty<KeyValuePair<string?, string?>>());
+			HttpResponseMessage response = await client.PostAsync(new Uri("https://slack.com/api/apps.connections.open"), content, stoppingToken);
+
+			byte[] responseData = await response.Content.ReadAsByteArrayAsync(stoppingToken);
+
+			SocketResponse socketResponse = JsonSerializer.Deserialize<SocketResponse>(responseData)!;
+			if (!socketResponse.Ok)
+			{
+				_logger.LogWarning("Unable to get websocket url: {Response}", Encoding.UTF8.GetString(responseData));
+				return null;
+			}
+
+			_logger.LogInformation("Using Slack websocket url: {Url}", socketResponse.Url);
+			return socketResponse.Url;
+		}
+
+		/// <summary>
+		/// Handle the lifetime of a websocket connection
+		/// </summary>
+		/// <param name="socketUrl"></param>
+		/// <param name="stoppingToken"></param>
+		/// <returns></returns>
+		private async Task HandleSocketAsync(Uri socketUrl, CancellationToken stoppingToken)
+		{
+			_logger.LogInformation("Opening Slack socket {Url}", socketUrl);
+
+			using ClientWebSocket socket = new ClientWebSocket();
+			await socket.ConnectAsync(socketUrl, stoppingToken);
+
+			byte[] buffer = new byte[2048];
+
+			bool disconnect = false;
+			while(!disconnect)
+			{
+				stoppingToken.ThrowIfCancellationRequested();
+
+				// Read the next message
+				int length = 0;
+				while(!disconnect)
+				{
+					if (length == buffer.Length)
+					{
+						Array.Resize(ref buffer, buffer.Length + 2048);
+					}
+
+					WebSocketReceiveResult result = await socket.ReceiveAsync(new ArraySegment<byte>(buffer, length, buffer.Length - length), stoppingToken);
+					if (result.MessageType == WebSocketMessageType.Close)
+					{
+						return;
+					}
+					length += result.Count;
+
+					if (result.EndOfMessage)
+					{
+						break;
+					}
+				}
+
+				// Get the message data
+				_logger.LogInformation("Slack event: {Message}", Encoding.UTF8.GetString(buffer, 0, length));
+				EventMessage eventMessage = JsonSerializer.Deserialize<EventMessage>(buffer.AsSpan(0, length))!;
+
+				// Prepare the response
+				EventResponse response = new EventResponse();
+				response.EnvelopeId = eventMessage.EnvelopeId;
+
+				// Handle the message type
+				if (eventMessage.Type != null)
+				{
+					string type = eventMessage.Type;
+					if (type.Equals("hello", StringComparison.Ordinal))
+					{
+						// nop
+					}
+					else if (type.Equals("disconnect", StringComparison.Ordinal))
+					{
+						disconnect = true;
+					}
+					else if (type.Equals("interactive", StringComparison.Ordinal) && eventMessage.Payload != null)
+					{
+						try
+						{
+							response.Payload = await HandleInteractionMessage(eventMessage.Payload, stoppingToken);
+						}
+						catch (Exception ex)
+						{
+							_logger.LogError(ex, "Exception handling Slack interaction message: {Message}", ex.Message);
+						}
+					}
+					else
+					{
+						_logger.LogDebug("Unhandled event type ({Type})", type);
+					}
+				}
+
+				// Send the response
+				await socket.SendAsync(JsonSerializer.SerializeToUtf8Bytes(response), WebSocketMessageType.Text, true, stoppingToken);
+			}
+		}
+		/// <summary>
+		/// Handle a button being clicked
+		/// </summary>
+		private async Task<object?> HandleInteractionMessage(EventPayload payload, CancellationToken cancellationToken)
+		{
+			_ = cancellationToken;
+
+			if (payload.User != null && payload.User.Id != null)
+			{
+				if (String.Equals(payload.Type, "block_actions", StringComparison.Ordinal))
+				{
+					foreach (ActionInfo action in payload.Actions)
+					{
+						if (action.Value != null)
+						{
+							Match? match;
+							if (TryMatch(action.Value, @"^issue_(\d+)_([a-zA-Z]+)_([a-fA-F0-9]{24})$", out match))
+							{
+								int issueId = Int32.Parse(match.Groups[1].Value, NumberStyles.Integer, CultureInfo.InvariantCulture);
+								string verb = match.Groups[2].Value;
+								UserId userId = new UserId(ObjectId.Parse(match.Groups[3].Value));
+								await HandleIssueDmResponseAsync(issueId, verb, userId, payload.User.Id);
+							}
+							else if (TryMatch(action.Value, @"^issue_(\d+)_([a-zA-Z]+)$", out match) && payload.TriggerId != null)
+							{
+								int issueId = Int32.Parse(match.Groups[1].Value, NumberStyles.Integer, CultureInfo.InvariantCulture);
+								string verb = match.Groups[2].Value;
+								await HandleIssueChannelResponseAsync(issueId, verb, payload.User.Id, payload.TriggerId);
+							}
+							else
+							{
+								_logger.LogWarning("Unrecognized action value: {Value}", action.Value);
+							}
+						}
+					}
+				}
+				else if (String.Equals(payload.Type, "view_submission", StringComparison.Ordinal))
+				{
+					if (payload.View != null && payload.View.State != null && payload.View.CallbackId != null)
+					{
+						Match? match;
+						if (TryMatch(payload.View.CallbackId, @"^issue_(\d+)_markfixed_([a-fA-F0-9]{24})$", out match))
+						{
+							int issueId = Int32.Parse(match.Groups[1].Value, NumberStyles.Integer, CultureInfo.InvariantCulture);
+							UserId userId = UserId.Parse(match.Groups[2].Value);
+
+							string? fixChangeStr;
+							if (payload.View.State.TryGetValue("fix_cl", "fix_cl_action", out fixChangeStr))
+							{
+								int fixChange;
+								if (!Int32.TryParse(fixChangeStr, out fixChange) || fixChange < 0)
+								{
+									Dictionary<string, string> errors = new Dictionary<string, string>();
+									errors.Add("fix_cl", $"'{fixChangeStr}' is not a valid fix changelist.");
+									return new { response_action = "errors", errors };
+								}
+
+								await _issueService.UpdateIssueAsync(issueId, fixChange: fixChange, resolvedById: userId, initiatedById: userId);
+								_logger.LogInformation("Marked issue {IssueId} fixed by user {UserId} in {Change}", issueId, userId, fixChange);
+							}
+						}
+						else if (TryMatch(payload.View.CallbackId, @"^issue_(\d+)_ack_([a-fA-F0-9]{24})$", out match))
+						{
+							int issueId = Int32.Parse(match.Groups[1].Value, NumberStyles.Integer, CultureInfo.InvariantCulture);
+							UserId userId = UserId.Parse(match.Groups[2].Value);
+							await _issueService.UpdateIssueAsync(issueId, acknowledged: true, ownerId: userId, initiatedById: userId);
+						}
+					}
+				}
+			}
+			return null;
+		}
+
+		static bool TryMatch(string input, string pattern, [NotNullWhen(true)] out Match? match)
+		{
+			Match result = Regex.Match(input, pattern);
+			if(result.Success)
+			{
+				match = result;
+				return true;
+			}
+			else
+			{
+				match = null;
+				return false;
+			}
+		}
+
+		async Task HandleIssueDmResponseAsync(int issueId, string verb, UserId userId, string userName)
+		{
+			_logger.LogInformation("Issue {IssueId}: {Action} from {SlackUser} ({UserId})", issueId, verb, userName, userId);
+
+			if (String.Equals(verb, "ack", StringComparison.Ordinal))
+			{
+				await _issueService.UpdateIssueAsync(issueId, acknowledged: true, initiatedById: userId);
+			}
+			else if (String.Equals(verb, "accept", StringComparison.Ordinal))
+			{
+				await _issueService.UpdateIssueAsync(issueId, ownerId: userId, nominatedById: userId, acknowledged: true, initiatedById: userId);
+			}
+			else if (String.Equals(verb, "decline", StringComparison.Ordinal))
+			{
+				await _issueService.UpdateIssueAsync(issueId, declinedById: userId, initiatedById: userId);
+			}
+
+			IIssue? newIssue = await _issueService.Collection.GetIssueAsync(issueId);
+			if (newIssue != null)
+			{
+				IUser? user = await _userCollection.GetUserAsync(userId);
+				if (user != null)
+				{
+					string? recipient = await GetSlackUserId(user);
+					if (recipient != null)
+					{
+						IIssueDetails details = await _issueService.GetIssueDetailsAsync(newIssue);
+						await SendIssueMessageAsync(_globalConfig.CurrentValue, recipient, newIssue, details, userId, DefaultAllowMentions);
+					}
+				}
+			}
+		}
+
+		async Task HandleIssueChannelResponseAsync(int issueId, string verb, string slackUserId, string triggerId)
+		{
+			_logger.LogInformation("Issue {IssueId}: {Action} from {SlackUser}", issueId, verb, slackUserId);
+
+			SlackUser? slackUser = await _slackClient.GetUserAsync(slackUserId);
+			if (slackUser == null || slackUser.Profile == null || slackUser.Profile.Email == null)
+			{
+				_logger.LogWarning("Unable to find Slack user profile for {UserId}", slackUserId);
+				return;
+			}
+
+			IUser? user = await _userCollection.FindUserByEmailAsync(slackUser.Profile.Email);
+			if (user == null)
+			{
+				_logger.LogWarning("Unable to find Horde user profile for {Email}", slackUser.Profile.Email);
+				return;
+			}
+
+			if (String.Equals(verb, "ack", StringComparison.Ordinal))
+			{
+				IIssue? issue = await _issueService.Collection.GetIssueAsync(issueId);
+				if (issue == null)
+				{
+					_logger.LogWarning("Unable to find issue {IssueId}", issueId);
+					return;
+				}
+
+				if (issue.OwnerId != null && issue.OwnerId != user.Id)
+				{
+					IUser? owner = await _userCollection.GetUserAsync(issue.OwnerId.Value);
+					if (owner != null)
+					{
+						SlackView view = new SlackView($"Issue {issueId}");
+						view.CallbackId = $"issue_{issueId}_ack_{user.Id}";
+						view.AddSection($"This issue is current assigned to {owner.Name}. Assign it to yourself?");
+						view.Close = "Cancel";
+						view.Submit = "Assign to Me";
+
+						await _slackClient.OpenViewAsync(triggerId, view);
+						return;
+					}
+				}
+
+				await _issueService.UpdateIssueAsync(issueId, acknowledged: true, ownerId: user.Id, nominatedById: user.Id, initiatedById: user.Id);
+			}
+			else if (String.Equals(verb, "decline", StringComparison.Ordinal))
+			{
+				IIssue? issue = await _issueService.Collection.GetIssueAsync(issueId);
+				if (issue == null)
+				{
+					_logger.LogWarning("Unable to find issue {IssueId}", issueId);
+					return;
+				}
+
+				if (issue.OwnerId != user.Id)
+				{
+					List<IIssueSuspect> suspects = await _issueService.Collection.FindSuspectsAsync(issue);
+					if (!suspects.Any(x => x.AuthorId == user.Id))
+					{
+						SlackView view = new SlackView($"Issue {issueId}");
+						view.AddSection("You are not currently listed as a suspect for this issue.");
+						view.Close = "Cancel";
+
+						await _slackClient.OpenViewAsync(triggerId, view);
+						return;
+					}
+				}
+
+				await _issueService.UpdateIssueAsync(issueId, declinedById: user.Id, initiatedById: user.Id);
+			}
+			else if (String.Equals(verb, "markfixed", StringComparison.Ordinal))
+			{
+				SlackView view = new SlackView($"Issue {issueId}");
+				view.CallbackId = $"issue_{issueId}_markfixed_{user.Id}";
+				view.AddInput("Fix Changelist", new PlainTextInputElement("fix_cl_action", placeholder: "Number")).BlockId = "fix_cl";
+				view.Close = "Cancel";
+				view.Submit = "Mark Fixed";
+
+				await _slackClient.OpenViewAsync(triggerId, view);
+				return;
+			}
+		}
+	}
+}
diff --git a/Engine/Source/Programs/Horde/Horde.Server/Perforce/ChangeMessages.cs b/Engine/Source/Programs/Horde/Horde.Server/Perforce/ChangeMessages.cs
new file mode 100644
--- /dev/null
+++ b/Engine/Source/Programs/Horde/Horde.Server/Perforce/ChangeMessages.cs
@@ -0,0 +1,9 @@
+// Copyright Epic Games, Inc. All Rights Reserved.
+
+using System.Collections.Generic;
+using Horde.Server.Users;
+
+namespace Horde.Server.Perforce
+{
+
+}
