Horde: Rename the Horde.Build project to Horde.Server.

diff --git a/Engine/Source/Programs/Horde/Horde.Server/Perforce/PerforceLoadBalancer.cs b/Engine/Source/Programs/Horde/Horde.Server/Perforce/PerforceLoadBalancer.cs
new file mode 100644
--- /dev/null
+++ b/Engine/Source/Programs/Horde/Horde.Server/Perforce/PerforceLoadBalancer.cs
@@ -0,0 +1,611 @@
+// Copyright Epic Games, Inc. All Rights Reserved.
+
+using System;
+using System.Collections.Generic;
+using System.Globalization;
+using System.Linq;
+using System.Net;
+using System.Net.Http;
+using System.Text.Json;
+using System.Threading;
+using System.Threading.Tasks;
+using EpicGames.Horde.Common;
+using Google.Protobuf.WellKnownTypes;
+using Horde.Server.Agents;
+using Horde.Server.Agents.Leases;
+using Horde.Server.Server;
+using Horde.Server.Utilities;
+using HordeCommon;
+using HordeCommon.Rpc.Tasks;
+using Microsoft.Extensions.Hosting;
+using Microsoft.Extensions.Logging;
+using Microsoft.Extensions.Options;
+using MongoDB.Bson.Serialization.Attributes;
+using MongoDB.Driver;
+
+namespace Horde.Server.Perforce
+{
+	/// <summary>
+	/// Health of a particular server
+	/// </summary>
+	public enum PerforceServerStatus
+	{
+		/// <summary>
+		/// Server could not be reached
+		/// </summary>
+		Unknown,
+
+		/// <summary>
+		/// Bad. Do not use.
+		/// </summary>
+		Unhealthy,
+
+		/// <summary>
+		/// Degraded, but functioning
+		/// </summary>
+		Degraded,
+
+		/// <summary>
+		/// Good
+		/// </summary>
+		Healthy,
+	}
+
+	/// <summary>
+	/// Perforce server infomation
+	/// </summary>
+	public interface IPerforceServer
+	{
+		/// <summary>
+		/// The server and port
+		/// </summary>
+		public string ServerAndPort { get; }
+
+		/// <summary>
+		/// The server and port before doing DNS resolution
+		/// </summary>
+		public string BaseServerAndPort { get; }
+
+		/// <summary>
+		/// The cluster this server belongs to
+		/// </summary>
+		public string Cluster { get; }
+
+		/// <summary>
+		/// Current status
+		/// </summary>
+		public PerforceServerStatus Status { get; }
+
+		/// <summary>
+		/// Number of leases using this server
+		/// </summary>
+		public int NumLeases { get; }
+
+		/// <summary>
+		/// Error message related to this server
+		/// </summary>
+		public string? Detail { get; }
+
+		/// <summary>
+		/// Last update time for this server
+		/// </summary>
+		public DateTime? LastUpdateTime { get; }
+	}
+
+	/// <summary>
+	/// Load balancer for Perforce edge servers
+	/// </summary>
+	public sealed class PerforceLoadBalancer : IHostedService, IDisposable
+	{
+		/// <summary>
+		/// Information about a resolved Perforce server
+		/// </summary>
+		class PerforceServerEntry : IPerforceServer
+		{
+			public string ServerAndPort { get; set; } = String.Empty;
+			public string BaseServerAndPort { get; set; } = String.Empty;
+			public string? HealthCheckUrl { get; set; }
+			public string Cluster { get; set; } = String.Empty;
+			public PerforceServerStatus Status { get; set; }
+			public string? Detail { get; set; }
+			public int NumLeases { get; set; }
+			public DateTime? LastUpdateTime { get; set; }
+
+			[BsonConstructor]
+			private PerforceServerEntry()
+			{
+			}
+
+			public PerforceServerEntry(string serverAndPort, string baseServerAndPort, string? healthCheckUrl, string cluster, PerforceServerStatus status, string? detail, DateTime? lastUpdateTime)
+			{
+				ServerAndPort = serverAndPort;
+				BaseServerAndPort = baseServerAndPort;
+				HealthCheckUrl = healthCheckUrl;
+				Cluster = cluster;
+				Status = status;
+				Detail = detail;
+				LastUpdateTime = lastUpdateTime;
+			}
+		}
+
+		/// <summary>
+		/// 
+		/// </summary>
+		[SingletonDocument("perforce-server-list", "6046aec374a9283100967ee7")]
+		class PerforceServerList : SingletonBase
+		{
+			public List<PerforceServerEntry> Servers { get; set; } = new List<PerforceServerEntry>();
+		}
+
+		readonly MongoService _mongoService;
+		readonly GlobalsService _globalsService;
+		readonly ILeaseCollection _leaseCollection;
+		readonly SingletonDocument<PerforceServerList> _serverListSingleton;
+		readonly Random _random = new Random();
+		readonly IOptionsMonitor<GlobalConfig> _globalConfig;
+		readonly ILogger _logger;
+		readonly ITicker _ticker;
+
+		/// <summary>
+		/// Constructor
+		/// </summary>
+		public PerforceLoadBalancer(MongoService mongoService, GlobalsService globalsService, ILeaseCollection leaseCollection, IClock clock, IOptionsMonitor<GlobalConfig> globalConfig, ILogger<PerforceLoadBalancer> logger)
+		{
+			_mongoService = mongoService;
+			_globalsService = globalsService;
+			_leaseCollection = leaseCollection;
+			_serverListSingleton = new SingletonDocument<PerforceServerList>(mongoService);
+			_globalConfig = globalConfig;
+			_logger = logger;
+			if (mongoService.ReadOnlyMode)
+			{
+				_ticker = new NullTicker();
+			}
+			else
+			{
+				_ticker = clock.AddSharedTicker<PerforceLoadBalancer>(TimeSpan.FromMinutes(1.0), TickInternalAsync, logger);
+			}
+		}
+
+		/// <inheritdoc/>
+		public Task StartAsync(CancellationToken cancellationToken) => _ticker.StartAsync();
+
+		/// <inheritdoc/>
+		public Task StopAsync(CancellationToken cancellationToken) => _ticker.StopAsync();
+
+		/// <inheritdoc/>
+		public void Dispose() => _ticker.Dispose();
+
+		/// <summary>
+		/// Get the current server list
+		/// </summary>
+		/// <returns></returns>
+		async Task<PerforceServerList> GetServerListAsync()
+		{
+			PerforceServerList serverList = await _serverListSingleton.GetAsync();
+
+			DateTime minLastUpdateTime = DateTime.UtcNow - TimeSpan.FromMinutes(2.5);
+			foreach (PerforceServerEntry server in serverList.Servers)
+			{
+				if (server.Status == PerforceServerStatus.Healthy && server.LastUpdateTime != null && server.LastUpdateTime.Value < minLastUpdateTime)
+				{
+					server.Status = PerforceServerStatus.Degraded;
+					server.Detail = "Server has not responded to health check";
+				}
+			}
+
+			return serverList;
+		}
+
+		/// <summary>
+		/// Get the perforce servers
+		/// </summary>
+		/// <returns></returns>
+		public async Task<List<IPerforceServer>> GetServersAsync()
+		{
+			PerforceServerList serverList = await GetServerListAsync();
+			return serverList.Servers.ConvertAll<IPerforceServer>(x => x);
+		}
+
+		/// <summary>
+		/// Allocates a server for use by a lease
+		/// </summary>
+		/// <returns>The server to use. Null if there is no healthy server available.</returns>
+		public async Task<IPerforceServer?> GetServer(string cluster)
+		{
+			PerforceServerList serverList = await GetServerListAsync();
+
+			List<PerforceServerEntry> candidates = serverList.Servers.Where(x => x.Cluster == cluster && x.Status >= PerforceServerStatus.Healthy).ToList();
+			if(candidates.Count == 0)
+			{
+				int idx = _random.Next(0, candidates.Count);
+				return candidates[idx];
+			}
+			return null;
+		}
+
+		/// <summary>
+		/// Select a Perforce server to use by the Horde server
+		/// </summary>
+		/// <param name="cluster"></param>
+		/// <returns></returns>
+		public Task<IPerforceServer?> SelectServerAsync(PerforceCluster cluster)
+		{
+			List<string> properties = new List<string>{ "HordeServer=1" };
+			return SelectServerAsync(cluster, properties);
+		}
+
+		/// <summary>
+		/// Select a Perforce server to use
+		/// </summary>
+		/// <param name="cluster"></param>
+		/// <param name="agent"></param>
+		/// <returns></returns>
+		public Task<IPerforceServer?> SelectServerAsync(PerforceCluster cluster, IAgent agent)
+		{
+			return SelectServerAsync(cluster, agent.Properties);
+		}
+
+		/// <summary>
+		/// Evaluate a server condition using standard Horde syntax
+		/// </summary>
+		static bool EvaluateCondition(Condition? condition, IReadOnlyList<string> properties)
+		{
+			if (condition == null)
+			{
+				return false;
+			}
+			if (condition.IsEmpty())
+			{
+				return false;
+			}
+			if (properties == null)
+			{
+				throw new ArgumentException("Properties argument to EvaluateCondition may not be null", nameof(properties));
+			}
+			return condition.Evaluate(x => GetPropertyValues(properties, x));
+		}
+
+		/// <summary>
+		/// Get property values from a list starting with the given key
+		/// </summary>
+		static IEnumerable<string> GetPropertyValues(IReadOnlyList<string> properties, string name)
+		{
+			foreach (string property in properties)
+			{
+				if (property.Length > name.Length && property[name.Length] == '=' && property.StartsWith(name, StringComparison.OrdinalIgnoreCase))
+				{
+					yield return property.Substring(name.Length + 1);
+				}
+			}
+		}
+
+		/// <summary>
+		/// Select a Perforce server to use
+		/// </summary>
+		/// <param name="cluster"></param>
+		/// <param name="properties"></param>
+		/// <returns></returns>
+		public async Task<IPerforceServer?> SelectServerAsync(PerforceCluster cluster, IReadOnlyList<string> properties)
+		{
+			// Find all the valid servers for this agent
+			List<PerforceServer> validServers = new List<PerforceServer>();
+			if (properties != null)
+			{
+				validServers.AddRange(cluster.Servers.Where(x => (x.Properties != null && x.Properties.Count > 0 && x.Properties.All(y => properties.Contains(y))) || EvaluateCondition(x.Condition, properties)));
+			}
+			if (validServers.Count == 0)
+			{
+				validServers.AddRange(cluster.Servers.Where(x => (x.Properties == null || x.Properties.Count == 0) && (x.Condition == null || x.Condition.IsEmpty())));
+			}
+
+			HashSet<string> validServerNames = new HashSet<string>(validServers.Select(x => x.ServerAndPort), StringComparer.OrdinalIgnoreCase);
+
+			// Find all the matching servers.
+			PerforceServerList serverList = await GetServerListAsync();
+
+			List<PerforceServerEntry> candidates = serverList.Servers.Where(x => x.Cluster == cluster.Name && validServerNames.Contains(x.BaseServerAndPort)).ToList();
+			if (candidates.Count == 0)
+			{
+				foreach (PerforceServer validServer in validServers)
+				{
+					_logger.LogDebug("Fetching server info for {ServerAndPort}", validServer.ServerAndPort);
+					await UpdateServerAsync(cluster, validServer, candidates);
+				}
+				if (candidates.Count == 0)
+				{
+					_logger.LogWarning("Unable to resolve any Perforce servers from valid list");
+					return null; 
+				}
+			}
+
+			// Remove any servers that are unhealthy
+			if (candidates.Any(x => x.Status == PerforceServerStatus.Healthy))
+			{
+				candidates.RemoveAll(x => x.Status != PerforceServerStatus.Healthy);
+			}
+			else
+			{
+				candidates.RemoveAll(x => x.Status != PerforceServerStatus.Unknown);
+			}
+
+			if (candidates.Count == 0)
+			{
+				_logger.LogWarning("Unable to find any healthy Perforce server in cluster {ClusterName}", cluster.Name);
+				return null;
+			}
+
+			// Select which server to use with a weighted average of the number of active leases
+			int index = 0;
+			if (candidates.Count > 1)
+			{
+				const int BaseWeight = 20;
+
+				int totalLeases = candidates.Sum(x => x.NumLeases);
+				int totalWeight = candidates.Sum(x => (totalLeases + BaseWeight) - x.NumLeases);
+
+				int weight = _random.Next(totalWeight);
+				for (; index + 1 < candidates.Count; index++)
+				{
+					weight -= (totalLeases + BaseWeight) - candidates[index].NumLeases;
+					if(weight < 0)
+					{
+						break;
+					}
+				}
+			}
+
+			return candidates[index];
+		}
+
+		/// <inheritdoc/>
+		async ValueTask TickInternalAsync(CancellationToken cancellationToken)
+		{
+			GlobalConfig globalConfig = _globalConfig.CurrentValue;
+
+			// Set of new server entries
+			List<PerforceServerEntry> newServers = new List<PerforceServerEntry>();
+
+			// Update the state of all the valid servers
+			foreach (PerforceCluster cluster in globalConfig.PerforceClusters)
+			{
+				foreach (PerforceServer server in cluster.Servers)
+				{
+					try
+					{
+						await UpdateServerAsync(cluster, server, newServers);
+					}
+					catch (Exception ex)
+					{
+						_logger.LogError(ex, "Exception while updating Perforce server status");
+					}
+				}
+			}
+
+			// Update the number of leases for each entry
+			List<PerforceServerEntry> newEntries = newServers.OrderBy(x => x.Cluster).ThenBy(x => x.BaseServerAndPort).ThenBy(x => x.ServerAndPort).ToList();
+			await UpdateLeaseCounts(newEntries);
+			PerforceServerList list = await _serverListSingleton.UpdateAsync(list => MergeServerList(list, newEntries));
+
+			// Now update the health of each entry
+			List<Task> tasks = new List<Task>();
+			foreach (PerforceServerEntry entry in list.Servers)
+			{
+				tasks.Add(Task.Run(() => UpdateHealthAsync(entry), CancellationToken.None));
+			}
+			await Task.WhenAll(tasks);
+		}
+
+		static void MergeServerList(PerforceServerList serverList, List<PerforceServerEntry> newEntries)
+		{
+			Dictionary<string, PerforceServerEntry> existingEntries = new Dictionary<string, PerforceServerEntry>(StringComparer.OrdinalIgnoreCase);
+			foreach (PerforceServerEntry existingEntry in serverList.Servers)
+			{
+				existingEntries[existingEntry.ServerAndPort] = existingEntry;
+			}
+			foreach (PerforceServerEntry newEntry in newEntries)
+			{
+				PerforceServerEntry? existingEntry;
+				if (existingEntries.TryGetValue(newEntry.ServerAndPort, out existingEntry))
+				{
+					newEntry.Status = existingEntry.Status;
+					newEntry.Detail = existingEntry.Detail;
+					newEntry.LastUpdateTime = existingEntry.LastUpdateTime;
+				}
+			}
+			serverList.Servers = newEntries;
+		}
+
+		async Task UpdateLeaseCounts(IEnumerable<PerforceServerEntry> newServerEntries)
+		{
+			Dictionary<string, Dictionary<string, PerforceServerEntry>> newServerLookup = new Dictionary<string, Dictionary<string, PerforceServerEntry>>();
+			foreach (PerforceServerEntry newServerEntry in newServerEntries)
+			{
+				if (newServerEntry.Cluster != null)
+				{
+					Dictionary<string, PerforceServerEntry>? clusterServers;
+					if (!newServerLookup.TryGetValue(newServerEntry.Cluster, out clusterServers))
+					{
+						clusterServers = new Dictionary<string, PerforceServerEntry>();
+						newServerLookup.Add(newServerEntry.Cluster, clusterServers);
+					}
+					clusterServers[newServerEntry.ServerAndPort] = newServerEntry;
+				}
+			}
+
+			List<ILease> leases = await _leaseCollection.FindActiveLeasesAsync();
+			foreach (ILease lease in leases)
+			{
+				Any any = Any.Parser.ParseFrom(lease.Payload.ToArray());
+
+				HashSet<(string, string)> servers = new HashSet<(string, string)>();
+				if (any.TryUnpack(out ConformTask conformTask))
+				{
+					foreach (HordeCommon.Rpc.Messages.AgentWorkspace workspace in conformTask.Workspaces)
+					{
+						servers.Add((workspace.Cluster, workspace.ServerAndPort));
+					}
+				}
+				if (any.TryUnpack(out ExecuteJobTask executeJobTask))
+				{
+					if (executeJobTask.AutoSdkWorkspace != null)
+					{
+						servers.Add((executeJobTask.AutoSdkWorkspace.Cluster, executeJobTask.AutoSdkWorkspace.ServerAndPort));
+					}
+					if (executeJobTask.Workspace != null)
+					{
+						servers.Add((executeJobTask.Workspace.Cluster, executeJobTask.Workspace.ServerAndPort));
+					}
+				}
+
+				foreach ((string cluster, string serverAndPort) in servers)
+				{
+					IncrementLeaseCount(cluster, serverAndPort, newServerLookup);
+				}
+			}
+		}
+
+		static void IncrementLeaseCount(string cluster, string serverAndPort, Dictionary<string, Dictionary<string, PerforceServerEntry>> newServers)
+		{
+			Dictionary<string, PerforceServerEntry>? clusterServers;
+			if (cluster != null && newServers.TryGetValue(cluster, out clusterServers))
+			{
+				PerforceServerEntry? entry;
+				if (clusterServers.TryGetValue(serverAndPort, out entry))
+				{
+					entry.NumLeases++;
+				}
+			}
+		}
+
+		async Task UpdateServerAsync(PerforceCluster cluster, PerforceServer server, List<PerforceServerEntry> newServers)
+		{
+			string initialHostName = server.ServerAndPort;
+			int port = 1666;
+
+			int portIdx = initialHostName.LastIndexOf(':');
+			if (portIdx != -1)
+			{
+				port = Int32.Parse(initialHostName.Substring(portIdx + 1), NumberStyles.Integer, CultureInfo.InvariantCulture);
+				initialHostName = initialHostName.Substring(0, portIdx);
+			}
+
+			List<string> hostNames = new List<string>();
+			if (!server.ResolveDns || IPAddress.TryParse(initialHostName, out IPAddress? _))
+			{
+				hostNames.Add(initialHostName);
+			}
+			else
+			{
+				await ResolveServersAsync(initialHostName, hostNames);
+			}
+
+			foreach (string hostName in hostNames)
+			{
+				string? healthCheckUrl = null;
+				if (server.HealthCheck)
+				{
+					healthCheckUrl = $"http://{hostName}:5000/healthcheck";
+				}
+				newServers.Add(new PerforceServerEntry($"{hostName}:{port}", server.ServerAndPort, healthCheckUrl, cluster.Name, PerforceServerStatus.Unknown, "", DateTime.UtcNow));
+			}
+		}
+
+		async Task ResolveServersAsync(string hostName, List<string> hostNames)
+		{
+			// Find all the addresses of the hosts
+			IPHostEntry entry = await Dns.GetHostEntryAsync(hostName);
+			foreach (IPAddress address in entry.AddressList)
+			{
+				try
+				{
+					hostNames.Add(address.ToString());
+				}
+				catch (Exception ex)
+				{
+					_logger.LogDebug(ex, "Unable to resolve host name for Perforce server {Address}", address);
+				}
+			}
+		}
+
+		/// <summary>
+		/// Updates 
+		/// </summary>
+		/// <param name="entry"></param>
+		/// <returns></returns>
+		async Task UpdateHealthAsync(PerforceServerEntry entry)
+		{
+			DateTime? updateTime = null;
+			string detail = "Health check disabled";
+
+			// Get the health of the server
+			PerforceServerStatus health = PerforceServerStatus.Healthy;
+			if (entry.HealthCheckUrl != null)
+			{
+				updateTime = DateTime.UtcNow;
+				Uri healthCheckUrl = new Uri(entry.HealthCheckUrl);
+				try
+				{
+					(health, detail) = await GetServerHealthAsync(healthCheckUrl);
+				}
+				catch
+				{
+					(health, detail) = (PerforceServerStatus.Unhealthy, $"Failed to query status at {healthCheckUrl}");
+				}
+			}
+
+			// Update the server record
+			if (health != entry.Status || detail != entry.Detail || updateTime != entry.LastUpdateTime)
+			{
+				await _serverListSingleton.UpdateAsync(x => UpdateHealth(x, entry.ServerAndPort, health, detail, updateTime));
+			}
+		}
+
+		static void UpdateHealth(PerforceServerList serverList, string serverAndPort, PerforceServerStatus status, string detail, DateTime? updateTime)
+		{
+			PerforceServerEntry? entry = serverList.Servers.FirstOrDefault(x => x.ServerAndPort.Equals(serverAndPort, StringComparison.Ordinal));
+			if (entry != null)
+			{
+				if (entry.LastUpdateTime == null || updateTime == null || entry.LastUpdateTime.Value < updateTime.Value)
+				{
+					entry.Status = status;
+					entry.Detail = detail;
+					entry.LastUpdateTime = updateTime;
+				}
+			}
+		}
+
+		static async Task<(PerforceServerStatus, string)> GetServerHealthAsync(Uri healthCheckUrl)
+		{
+			using HttpClient client = new HttpClient();
+			HttpResponseMessage response = await client.GetAsync(healthCheckUrl);
+
+			byte[] data = await response.Content.ReadAsByteArrayAsync();
+			JsonDocument document = JsonDocument.Parse(data);
+
+			foreach (JsonElement element in document.RootElement.GetProperty("results").EnumerateArray())
+			{
+				if (element.TryGetProperty("checker", out JsonElement checker) && checker.ValueEquals("edge_traffic_lights"))
+				{
+					if (element.TryGetProperty("output", out JsonElement output))
+					{
+						string status = output.GetString() ?? String.Empty;
+						switch(status)
+						{
+							case "green":
+								return (PerforceServerStatus.Healthy, "Server is healthy");
+							case "yellow":
+								return (PerforceServerStatus.Degraded, "Degraded service");
+							case "red":
+								return (PerforceServerStatus.Unhealthy, "Server is being drained");
+							default:
+								return (PerforceServerStatus.Unknown, $"Expected state for health check ({status})");
+						}
+					}
+				}
+			}
+
+			return (PerforceServerStatus.Unknown, "Unable to parse health check output");
+		}
+	}
+}
diff --git a/Engine/Source/Programs/Horde/Horde.Server/Perforce/ReplicationService.cs b/Engine/Source/Programs/Horde/Horde.Server/Perforce/ReplicationService.cs
new file mode 100644
--- /dev/null
+++ b/Engine/Source/Programs/Horde/Horde.Server/Perforce/ReplicationService.cs
@@ -0,0 +1,167 @@
+// Copyright Epic Games, Inc. All Rights Reserved.
+
+using System;
+using System.Collections.Generic;
+using System.IO;
+using System.Linq;
+using System.Threading;
+using System.Threading.Tasks;
+using EpicGames.Core;
+using EpicGames.Horde.Storage;
+using EpicGames.Horde.Storage.Nodes;
+using Horde.Server.Server;
+using Horde.Server.Storage;
+using Horde.Server.Streams;
+using Horde.Server.Utilities;
+using HordeCommon;
+using Microsoft.Extensions.Caching.Memory;
+using Microsoft.Extensions.Hosting;
+using Microsoft.Extensions.Logging;
+using Microsoft.Extensions.Options;
+
+namespace Horde.Server.Perforce
+{
+	/// <summary>
+	/// Exception triggered during content replication
+	/// </summary>
+	public sealed class ReplicationException : Exception
+	{
+		internal ReplicationException(string message) : base(message)
+		{
+		}
+	}
+
+	/// <summary>
+	/// Service which replicates content from Perforce
+	/// </summary>
+	sealed class ReplicationService : IHostedService
+	{
+		readonly IStreamCollection _streamCollection;
+		readonly IPerforceService _perforceService;
+		readonly PerforceReplicator _replicator;
+		readonly StorageService _storageService;
+		readonly IMemoryCache _memoryCache;
+		readonly ITicker _ticker;
+		readonly IOptionsMonitor<GlobalConfig> _globalConfig;
+		readonly ILogger _logger;
+
+		/// <summary>
+		/// Constructor
+		/// </summary>
+		public ReplicationService(IStreamCollection streamCollection, IPerforceService perforceService, PerforceReplicator replicator, StorageService storageService, IMemoryCache memoryCache, IClock clock, IOptionsMonitor<GlobalConfig> globalConfig, ILogger<ReplicationService> logger)
+		{
+			_streamCollection = streamCollection;
+			_perforceService = perforceService;
+			_replicator = replicator;
+			_storageService = storageService;
+			_memoryCache = memoryCache;
+			_ticker = clock.AddSharedTicker<ReplicationService>(TimeSpan.FromSeconds(20.0), TickSharedAsync, logger);
+			_globalConfig = globalConfig;
+			_logger = logger;
+		}
+
+		/// <inheritdoc/>
+		public async Task StartAsync(CancellationToken cancellationToken)
+		{
+			await _ticker.StartAsync();
+		}
+
+		/// <inheritdoc/>
+		public async Task StopAsync(CancellationToken cancellationToken)
+		{
+			await _ticker.StopAsync();
+		}
+
+		async ValueTask TickSharedAsync(CancellationToken cancellationToken)
+		{
+			Dictionary<StreamId, BackgroundTask> streamIdToTask = new Dictionary<StreamId, BackgroundTask>();
+			try
+			{
+				for (; ; )
+				{
+					IReadOnlyList<StreamConfig> streamConfigs = _globalConfig.CurrentValue.Streams;
+
+					HashSet<StreamId> removeStreams = new HashSet<StreamId>(streamIdToTask.Keys);
+					foreach (StreamConfig streamConfig in streamConfigs)
+					{
+						if (streamConfig.ReplicationMode == ContentReplicationMode.Full)
+						{
+							removeStreams.Remove(streamConfig.Id);
+							if (!streamIdToTask.ContainsKey(streamConfig.Id))
+							{
+								streamIdToTask.Add(streamConfig.Id, BackgroundTask.StartNew(ctx => RunReplicationGuardedAsync(streamConfig, ctx)));
+								_logger.LogInformation("Started replication of {StreamId}", streamConfig.Id);
+							}
+						}
+					}
+
+					foreach (StreamId removeStreamId in removeStreams)
+					{
+						if (streamIdToTask.Remove(removeStreamId, out BackgroundTask? task))
+						{
+							await task.DisposeAsync();
+							_logger.LogInformation("Stopped replication of {StreamId}", removeStreamId);
+						}
+					}
+
+					await Task.Delay(TimeSpan.FromMinutes(1.0), cancellationToken);
+				}
+			}
+			finally
+			{
+				await Parallel.ForEachAsync(streamIdToTask.Values, (x, ctx) => x.DisposeAsync());
+			}
+		}
+
+		async Task RunReplicationGuardedAsync(StreamConfig streamConfig, CancellationToken cancellationToken)
+		{
+			for (; ; )
+			{
+				try
+				{
+					await RunReplicationAsync(streamConfig, cancellationToken);
+				}
+				catch (OperationCanceledException) when (cancellationToken.IsCancellationRequested)
+				{
+					break;
+				}
+				catch (Exception ex)
+				{
+					_logger.LogError(ex, "Exception running replication for {StreamId}: {Message}", streamConfig.Id, ex.Message);
+				}
+
+				await Task.Delay(TimeSpan.FromSeconds(20.0), cancellationToken);
+			}
+		}
+
+		async Task RunReplicationAsync(StreamConfig streamConfig, CancellationToken cancellationToken)
+		{
+			RefName refName = new RefName(streamConfig.Id.ToString());
+
+			IStorageClientImpl store = await _storageService.GetClientAsync(Namespace.Perforce, cancellationToken);
+			TreeReader reader = new TreeReader(store, _memoryCache, _logger);
+
+			CommitNode? lastCommitNode = await reader.TryReadNodeAsync<CommitNode>(refName, cancellationToken: cancellationToken);
+			ICommitCollection commits = _perforceService.GetCommits(streamConfig);
+
+			PerforceReplicationOptions options = new PerforceReplicationOptions();
+
+			ICommit commit;
+			if (lastCommitNode == null)
+			{
+				commit = await commits.GetLatestAsync(cancellationToken);
+			}
+			else
+			{
+				commit = await commits.SubscribeAsync(lastCommitNode.Number, cancellationToken: cancellationToken).FirstAsync(cancellationToken);
+			}
+
+			for (; ; )
+			{
+				_logger.LogInformation("Replicating {StreamId} change {Change}", streamConfig.Id, commit.Number);
+				await _replicator.WriteAsync(streamConfig, commit.Number, options, cancellationToken);
+				commit = await commits.SubscribeAsync(commit.Number, cancellationToken: cancellationToken).FirstAsync(cancellationToken);
+			}
+		}
+	}
+}
diff --git a/Engine/Source/Programs/Horde/Horde.Server/Program.cs b/Engine/Source/Programs/Horde/Horde.Server/Program.cs
new file mode 100644
--- /dev/null
+++ b/Engine/Source/Programs/Horde/Horde.Server/Program.cs
@@ -0,0 +1,265 @@
+// Copyright Epic Games, Inc. All Rights Reserved.
+
+using System;
+using System.Collections.Generic;
+using System.Diagnostics;
+using System.IO;
+using System.Reflection;
+using System.Runtime.InteropServices;
+using System.Threading.Tasks;
+using EpicGames.Core;
+using Horde.Server.Commands;
+using Microsoft.Extensions.Configuration;
+using Microsoft.Extensions.DependencyInjection;
+using Microsoft.Extensions.Hosting;
+using OpenTracing;
+using OpenTracing.Propagation;
+using OpenTracing.Util;
+using Serilog;
+using Serilog.Configuration;
+using Serilog.Core;
+using Serilog.Events;
+using Serilog.Exceptions;
+using Serilog.Exceptions.Core;
+using Serilog.Exceptions.Grpc.Destructurers;
+using Serilog.Formatting.Json;
+using Serilog.Sinks.SystemConsole.Themes;
+
+namespace Horde.Server
+{
+	static class LoggerExtensions
+	{
+		public static LoggerConfiguration Console(this LoggerSinkConfiguration sinkConfig, ServerSettings settings)
+		{
+			if (settings.LogJsonToStdOut)
+			{
+				return sinkConfig.Console(new JsonFormatter(renderMessage: true));
+			}
+			else
+			{
+				ConsoleTheme theme;
+				if (RuntimeInformation.IsOSPlatform(OSPlatform.Windows) && Environment.OSVersion.Version < new Version(10, 0))
+				{
+					theme = SystemConsoleTheme.Literate;
+				}
+				else
+				{
+					theme = AnsiConsoleTheme.Code;
+				}
+				return sinkConfig.Console(outputTemplate: "[{Timestamp:HH:mm:ss} {Level:w3}] {Indent}{Message:l}{NewLine}{Exception}", theme: theme, restrictedToMinimumLevel: settings.ConsoleLogLevel);
+			}
+		}
+
+		public static LoggerConfiguration WithHordeConfig(this LoggerConfiguration configuration, ServerSettings settings)
+		{
+			if (settings.WithDatadog)
+			{
+				configuration = configuration.Enrich.With<DatadogLogEnricher>();
+			}
+			return configuration;
+		}
+	}
+
+	class DatadogLogEnricher : ILogEventEnricher
+	{
+		public void Enrich(Serilog.Events.LogEvent logEvent, ILogEventPropertyFactory propertyFactory)
+		{
+			ISpan? span = GlobalTracer.Instance?.ActiveSpan;
+			if (span != null)
+			{
+				logEvent.AddPropertyIfAbsent(propertyFactory.CreateProperty("dd.trace_id", span.Context.TraceId));
+				logEvent.AddPropertyIfAbsent(propertyFactory.CreateProperty("dd.span_id", span.Context.SpanId));
+			}
+		}
+	}
+
+	class TestTracer : ITracer
+	{
+		readonly ITracer _inner;
+
+		public TestTracer(ITracer inner)
+		{
+			_inner = inner;
+		}
+
+		public IScopeManager ScopeManager => _inner.ScopeManager;
+
+		public ISpan ActiveSpan => _inner.ActiveSpan;
+
+		public ISpanBuilder BuildSpan(string operationName)
+		{
+			Serilog.Log.Debug("Creating span {Name}", operationName);
+			return _inner.BuildSpan(operationName);
+		}
+
+		public ISpanContext Extract<TCarrier>(IFormat<TCarrier> format, TCarrier carrier)
+		{
+			return _inner.Extract<TCarrier>(format, carrier);
+		}
+
+		public void Inject<TCarrier>(ISpanContext spanContext, IFormat<TCarrier> format, TCarrier carrier)
+		{
+			_inner.Inject<TCarrier>(spanContext, format, carrier);
+		}
+	}
+
+	class Program
+	{
+		public static SemVer Version => s_version;
+
+		public static string DeploymentEnvironment { get; } = GetEnvironment();
+
+		public static DirectoryReference AppDir { get; } = GetAppDir();
+
+		public static DirectoryReference DataDir { get; } = GetDataDir();
+
+		public static FileReference UserConfigFile { get; } = FileReference.Combine(DataDir, "Horde.json");
+
+		public static Type[] ConfigSchemas = FindSchemaTypes();
+
+		static SemVer s_version;
+
+		static Type[] FindSchemaTypes()
+		{
+			List<Type> schemaTypes = new List<Type>();
+			foreach (Type type in Assembly.GetExecutingAssembly().GetTypes())
+			{
+				if (type.GetCustomAttribute<JsonSchemaAttribute>() != null)
+				{
+					schemaTypes.Add(type);
+				}
+			}
+			return schemaTypes.ToArray();
+		}
+
+		public static async Task<int> Main(string[] args)
+		{
+			FileVersionInfo versionInfo = FileVersionInfo.GetVersionInfo(Assembly.GetExecutingAssembly().Location);
+			if (String.IsNullOrEmpty(versionInfo.ProductVersion))
+			{
+				s_version = SemVer.Parse("0.0.0");
+			}
+			else
+			{
+				s_version = SemVer.Parse(versionInfo.ProductVersion);
+			}
+
+			CommandLineArguments arguments = new CommandLineArguments(args);
+
+			IConfiguration config = CreateConfig(UserConfigFile);
+
+			ServerSettings hordeSettings = new ServerSettings();
+			config.GetSection("Horde").Bind(hordeSettings);
+
+			DirectoryReference logDir = AppDir;
+			if (RuntimeInformation.IsOSPlatform(OSPlatform.Windows))
+			{
+				logDir = DirectoryReference.Combine(DataDir);
+			}
+
+			Serilog.Log.Logger = new LoggerConfiguration()
+				.WithHordeConfig(hordeSettings)
+				.Enrich.FromLogContext()
+				.Enrich.WithExceptionDetails(new DestructuringOptionsBuilder()
+					.WithDefaultDestructurers()
+					.WithDestructurers(new[] { new RpcExceptionDestructurer() }))
+				.WriteTo.Console(hordeSettings)
+				.WriteTo.File(Path.Combine(logDir.FullName, "Log.txt"), outputTemplate: "[{Timestamp:HH:mm:ss} {Level:w3}] {Indent}{Message:l}{NewLine}{Exception} [{SourceContext}]", rollingInterval: RollingInterval.Day, rollOnFileSizeLimit: true, fileSizeLimitBytes: 20 * 1024 * 1024, retainedFileCountLimit: 10)
+				.WriteTo.File(new JsonFormatter(renderMessage: true), Path.Combine(logDir.FullName, "Log.json"), rollingInterval: RollingInterval.Day, rollOnFileSizeLimit: true, fileSizeLimitBytes: 20 * 1024 * 1024, retainedFileCountLimit: 10)
+				.ReadFrom.Configuration(config)
+				.CreateLogger();
+
+			Serilog.Log.Logger.Information("Server version: {Version}", Version);
+
+			if (hordeSettings.WithDatadog)
+			{
+				GlobalTracer.Register(Datadog.Trace.OpenTracing.OpenTracingTracerFactory.WrapTracer(Datadog.Trace.Tracer.Instance));
+				Serilog.Log.Logger.Information("Enabling datadog tracing (OpenTrace)");
+			}
+
+			ServiceCollection services = new ServiceCollection();
+			services.AddCommandsFromAssembly(Assembly.GetExecutingAssembly());
+			services.AddLogging(builder => builder.AddSerilog());
+			services.AddSingleton<IConfiguration>(config);
+			services.AddSingleton<ServerSettings>(hordeSettings);
+
+#pragma warning disable ASP0000 // Do not call 'IServiceCollection.BuildServiceProvider' in 'ConfigureServices'
+			await using ServiceProvider serviceProvider = services.BuildServiceProvider();
+			return await CommandHost.RunAsync(arguments, serviceProvider, typeof(ServerCommand));
+#pragma warning restore ASP0000 // Do not call 'IServiceCollection.BuildServiceProvider' in 'ConfigureServices'
+		}
+
+		// Used by WebApplicationFactory in controller tests. Uses reflection to call this exact function signature.
+		public static IHostBuilder CreateHostBuilder(string[] args) => ServerCommand.CreateHostBuilderForTesting(args);
+
+		/// <summary>
+		/// Gets the current environment
+		/// </summary>
+		/// <returns></returns>
+		static string GetEnvironment()
+		{
+			string? environment = Environment.GetEnvironmentVariable("ASPNETCORE_ENVIRONMENT");
+			if (String.IsNullOrEmpty(environment))
+			{
+				environment = "Production";
+			}
+			return environment;
+		}
+
+		/// <summary>
+		/// Get the application directory
+		/// </summary>
+		/// <returns></returns>
+		static DirectoryReference GetAppDir()
+		{
+			return new FileReference(Assembly.GetExecutingAssembly().Location).Directory;
+		}
+
+		/// <summary>
+		/// Gets the default directory for storing application data
+		/// </summary>
+		/// <returns>The default data directory</returns>
+		static DirectoryReference GetDataDir()
+		{
+			IConfiguration config = CreateConfig(null);
+
+			string? dataDir = config.GetSection("Horde").GetValue(typeof(string), nameof(ServerSettings.DataDir)) as string;
+			if (dataDir != null)
+			{
+				return DirectoryReference.Combine(GetAppDir(), dataDir);
+			}
+
+			if (RuntimeInformation.IsOSPlatform(OSPlatform.Windows))
+			{
+				DirectoryReference? dir = DirectoryReference.GetSpecialFolder(Environment.SpecialFolder.CommonApplicationData);
+				if (dir != null)
+				{
+					return DirectoryReference.Combine(dir, "HordeServer");
+				}
+			}
+			return DirectoryReference.Combine(GetAppDir(), "Data");
+		}
+
+		/// <summary>
+		/// Constructs a configuration object for the current environment
+		/// </summary>
+		/// <param name="userConfigFile"></param>
+		/// <returns></returns>
+		static IConfiguration CreateConfig(FileReference? userConfigFile)
+		{
+			IConfigurationBuilder builder = new ConfigurationBuilder()
+				.SetBasePath(AppDir.FullName)
+				.AddJsonFile("appsettings.json", optional: false)
+				.AddJsonFile("appsettings.Build.json", optional: true) // specific settings for builds (installer/dockerfile)
+				.AddJsonFile($"appsettings.{DeploymentEnvironment}.json", optional: true) // environment variable overrides, also used in k8s setups with Helm
+				.AddJsonFile("appsettings.User.json", optional: true);
+
+			if (userConfigFile != null)
+			{
+				builder = builder.AddJsonFile(userConfigFile.FullName, optional: true, reloadOnChange: true);
+			}
+
+			return builder.AddEnvironmentVariables().Build();
+		}
+	}
+}
diff --git a/Engine/Source/Programs/Horde/Horde.Server/Projects/ProjectMessages.cs b/Engine/Source/Programs/Horde/Horde.Server/Projects/ProjectMessages.cs
new file mode 100644
--- /dev/null
+++ b/Engine/Source/Programs/Horde/Horde.Server/Projects/ProjectMessages.cs
@@ -0,0 +1,202 @@
+// Copyright Epic Games, Inc. All Rights Reserved.
+
+using System.Collections.Generic;
+using System.ComponentModel.DataAnnotations;
+using System.Linq;
+using System.Text.RegularExpressions;
+using Horde.Server.Acls;
+using Horde.Server.Streams;
+
+namespace Horde.Server.Projects
+{
+	/// <summary>
+	/// Information about a stream within a project
+	/// </summary>
+	public class GetProjectStreamResponse
+	{
+		/// <summary>
+		/// The stream id
+		/// </summary>
+		public string Id { get; set; }
+
+		/// <summary>
+		/// The stream name
+		/// </summary>
+		public string Name { get; set; }
+
+		/// <summary>
+		/// Constructor
+		/// </summary>
+		/// <param name="id">The unique stream id</param>
+		/// <param name="name">The stream name</param>
+		public GetProjectStreamResponse(string id, string name)
+		{
+			Id = id;
+			Name = name;
+		}
+	}
+
+	/// <summary>
+	/// Information about a category to display for a stream
+	/// </summary>
+	public class GetProjectCategoryResponse
+	{
+		/// <summary>
+		/// Heading for this column
+		/// </summary>
+		public string Name { get; set; }
+
+		/// <summary>
+		/// Index of the row to display this category on
+		/// </summary>
+		public int Row { get; set; }
+
+		/// <summary>
+		/// Whether to show this category on the nav menu
+		/// </summary>
+		public bool ShowOnNavMenu { get; set; }
+
+		/// <summary>
+		/// Patterns for stream names to include
+		/// </summary>
+		public List<string> IncludePatterns { get; set; } = new List<string>();
+
+		/// <summary>
+		/// Patterns for stream names to exclude
+		/// </summary>
+		public List<string> ExcludePatterns { get; set; } = new List<string>();
+
+		/// <summary>
+		/// Streams to include in this category
+		/// </summary>
+		public List<string> Streams { get; set; } = new List<string>();
+
+		/// <summary>
+		/// Constructor
+		/// </summary>
+		/// <param name="streamCategory">The category to construct from</param>
+		public GetProjectCategoryResponse(ProjectCategoryConfig streamCategory)
+		{
+			Name = streamCategory.Name;
+			Row = streamCategory.Row;
+			ShowOnNavMenu = streamCategory.ShowOnNavMenu;
+			IncludePatterns = streamCategory.IncludePatterns;
+			ExcludePatterns = streamCategory.ExcludePatterns;
+		}
+	}
+
+	/// <summary>
+	/// Response describing a project
+	/// </summary>
+	public class GetProjectResponse
+	{
+		/// <summary>
+		/// Unique id of the project
+		/// </summary>
+		public string Id { get; set; }
+
+		/// <summary>
+		/// Name of the project
+		/// </summary>
+		public string Name { get; set; }
+
+		/// <summary>
+		/// Order to display this project on the dashboard
+		/// </summary>
+		public int Order { get; set; }
+
+		/// <summary>
+		/// List of streams that are in this project
+		/// </summary>
+		public List<GetProjectStreamResponse>? Streams { get; set; }
+
+		/// <summary>
+		/// List of stream categories to display
+		/// </summary>
+		public List<GetProjectCategoryResponse>? Categories { get; set; }
+
+		/// <summary>
+		/// Constructor
+		/// </summary>
+		/// <param name="id">Unique id of the project</param>
+		/// <param name="name">Name of the project</param>
+		/// <param name="order">Order to show this project on the dashboard</param>
+		/// <param name="streams">List of streams to display</param>
+		/// <param name="categories">List of stream categories to display</param>
+		public GetProjectResponse(string id, string name, int order, List<GetProjectStreamResponse>? streams, List<GetProjectCategoryResponse>? categories)
+		{
+			Id = id;
+			Name = name;
+			Order = order;
+			Streams = streams;
+			Categories = categories;
+		}
+
+		/// <summary>
+		/// Converts this object to a public response
+		/// </summary>
+		/// <param name="projectConfig">The project instance</param>
+		/// <param name="includeStreams">Whether to include streams in the response</param>
+		/// <param name="includeCategories">Whether to include categories in the response</param>
+		/// <param name="streamConfigs">The list of streams</param>
+		/// <returns>Response instance</returns>
+		public static GetProjectResponse FromConfig(ProjectConfig projectConfig, bool includeStreams, bool includeCategories, List<StreamConfig>? streamConfigs)
+		{
+			List<GetProjectStreamResponse>? streamResponses = null;
+			if (includeStreams)
+			{
+				streamResponses = streamConfigs!.ConvertAll(x => new GetProjectStreamResponse(x.Id.ToString(), x.Name));
+			}
+
+			List<GetProjectCategoryResponse>? categoryResponses = null;
+			if (includeCategories)
+			{
+				categoryResponses = projectConfig.Categories.ConvertAll(x => new GetProjectCategoryResponse(x));
+				if (streamConfigs != null)
+				{
+					foreach (StreamConfig streamConfig in streamConfigs)
+					{
+						GetProjectCategoryResponse? categoryResponse = categoryResponses.FirstOrDefault(x => MatchCategory(streamConfig.Name, x));
+						if (categoryResponse == null)
+						{
+							int row = (categoryResponses.Count > 0) ? categoryResponses.Max(x => x.Row) : 0;
+							if (categoryResponses.Count(x => x.Row == row) >= 3)
+							{
+								row++;
+							}
+
+							ProjectCategoryConfig otherCategory = new ProjectCategoryConfig();
+							otherCategory.Name = "Other";
+							otherCategory.Row = row;
+							otherCategory.IncludePatterns.Add(".*");
+
+							categoryResponse = new GetProjectCategoryResponse(otherCategory);
+							categoryResponses.Add(categoryResponse);
+						}
+						categoryResponse.Streams!.Add(streamConfig.Id.ToString());
+					}
+				}
+			}
+
+			return new GetProjectResponse(projectConfig.Id.ToString(), projectConfig.Name, projectConfig.Order, streamResponses, categoryResponses);
+		}
+
+		/// <summary>
+		/// Tests if a category response matches a given stream name
+		/// </summary>
+		/// <param name="name">The stream name</param>
+		/// <param name="category">The category response</param>
+		/// <returns>True if the category matches</returns>
+		static bool MatchCategory(string name, GetProjectCategoryResponse category)
+		{
+			if (category.IncludePatterns.Any(x => Regex.IsMatch(name, x)))
+			{
+				if (!category.ExcludePatterns.Any(x => Regex.IsMatch(name, x)))
+				{
+					return true;
+				}
+			}
+			return false;
+		}
+	}
+}
diff --git a/Engine/Source/Programs/Horde/Horde.Server/Projects/ProjectsController.cs b/Engine/Source/Programs/Horde/Horde.Server/Projects/ProjectsController.cs
new file mode 100644
--- /dev/null
+++ b/Engine/Source/Programs/Horde/Horde.Server/Projects/ProjectsController.cs
@@ -0,0 +1,135 @@
+// Copyright Epic Games, Inc. All Rights Reserved.
+
+using System.Collections.Generic;
+using System.Linq;
+using System.Threading.Tasks;
+using Horde.Server.Acls;
+using Horde.Server.Configuration;
+using Horde.Server.Server;
+using Horde.Server.Streams;
+using Horde.Server.Utilities;
+using Microsoft.AspNetCore.Authorization;
+using Microsoft.AspNetCore.Mvc;
+using Microsoft.AspNetCore.StaticFiles;
+using Microsoft.Extensions.Options;
+
+namespace Horde.Server.Projects
+{
+	/// <summary>
+	/// Controller for the /api/v1/projects endpoint
+	/// </summary>
+	[ApiController]
+	[Authorize]
+	[Route("[controller]")]
+	public class ProjectsController : HordeControllerBase
+	{
+		private readonly IOptionsSnapshot<GlobalConfig> _globalConfig;
+
+		private static readonly FileExtensionContentTypeProvider s_contentTypeProvider = new FileExtensionContentTypeProvider();
+
+		/// <summary>
+		/// Constructor
+		/// </summary>
+		public ProjectsController(IOptionsSnapshot<GlobalConfig> globalConfig)
+		{
+			_globalConfig = globalConfig;
+		}
+
+		/// <summary>
+		/// Query all the projects
+		/// </summary>
+		/// <param name="includeStreams">Whether to include streams in the response</param>
+		/// <param name="includeCategories">Whether to include categories in the response</param>
+		/// <param name="filter">Filter for the properties to return</param>
+		/// <returns>Information about all the projects</returns>
+		[HttpGet]
+		[Route("/api/v1/projects")]
+		[ProducesResponseType(typeof(List<GetProjectResponse>), 200)]
+		public ActionResult<List<object>> GetProjects([FromQuery(Name = "Streams")] bool includeStreams = false, [FromQuery(Name = "Categories")] bool includeCategories = false, [FromQuery] PropertyFilter? filter = null)
+		{
+			GlobalConfig globalConfig = _globalConfig.Value;
+
+			List<object> responses = new List<object>();
+			foreach (ProjectConfig projectConfig in globalConfig.Projects)
+			{
+				if (projectConfig.Authorize(AclAction.ViewProject, User))
+				{
+					List<StreamConfig>? visibleStreams = null;
+					if (includeStreams || includeCategories)
+					{
+						visibleStreams = projectConfig.Streams.Where(x => x.Authorize(AclAction.ViewStream, User)).ToList();
+					}
+					responses.Add(GetProjectResponse.FromConfig(projectConfig, includeStreams, includeCategories, visibleStreams).ApplyFilter(filter));
+				}
+			}
+			return responses;
+		}
+
+		/// <summary>
+		/// Retrieve information about a specific project
+		/// </summary>
+		/// <param name="projectId">Id of the project to get information about</param>
+		/// <param name="filter">Filter for the properties to return</param>
+		/// <returns>Information about the requested project</returns>
+		[HttpGet]
+		[Route("/api/v1/projects/{projectId}")]
+		[ProducesResponseType(typeof(List<GetProjectResponse>), 200)]
+		public ActionResult<object> GetProject(ProjectId projectId, [FromQuery] PropertyFilter? filter = null)
+		{
+			ProjectConfig? projectConfig;
+			if(!_globalConfig.Value.TryGetProject(projectId, out projectConfig))
+			{
+				return NotFound(projectId);
+			}
+			if (!projectConfig.Authorize(AclAction.ViewProject, User))
+			{
+				return Forbid(AclAction.ViewProject, projectId);
+			}
+
+			bool includeStreams = PropertyFilter.Includes(filter, nameof(GetProjectResponse.Streams));
+			bool includeCategories = PropertyFilter.Includes(filter, nameof(GetProjectResponse.Categories));
+
+			List<StreamConfig>? visibleStreams = null;
+			if (includeStreams || includeCategories)
+			{
+				visibleStreams = projectConfig.Streams.Where(x => x.Authorize(AclAction.ViewStream, User)).ToList();
+			}
+
+			return GetProjectResponse.FromConfig(projectConfig, includeStreams, includeCategories, visibleStreams).ApplyFilter(filter);
+		}
+
+		/// <summary>
+		/// Retrieve information about a specific project
+		/// </summary>
+		/// <param name="projectId">Id of the project to get information about</param>
+		/// <returns>Information about the requested project</returns>
+		[HttpGet]
+		[Route("/api/v1/projects/{projectId}/logo")]
+		public ActionResult<object> GetProjectLogo(ProjectId projectId)
+		{
+			ProjectConfig? projectConfig;
+			if (!_globalConfig.Value.TryGetProject(projectId, out projectConfig))
+			{
+				return NotFound(projectId);
+			}
+			if (!projectConfig.Authorize(AclAction.ViewProject, User))
+			{
+				return Forbid(AclAction.ViewProject, projectId);
+			}
+
+			ConfigResource? logoResource = projectConfig.Logo;
+			if (logoResource == null || logoResource.Path == null || logoResource.Data.Length == 0)
+			{
+				return NotFound("Missing logo resource data");
+			}
+
+			string? contentType;
+			if (!s_contentTypeProvider.TryGetContentType(logoResource.Path, out contentType))
+			{
+				contentType = "application/octet-stream";
+			}
+
+			return new FileContentResult(logoResource.Data.ToArray(), contentType);
+		}
+	}
+}
diff --git a/Engine/Source/Programs/Horde/Horde.Server/Properties/AssemblyInfo.cs b/Engine/Source/Programs/Horde/Horde.Server/Properties/AssemblyInfo.cs
new file mode 100644
--- /dev/null
+++ b/Engine/Source/Programs/Horde/Horde.Server/Properties/AssemblyInfo.cs
@@ -0,0 +1,5 @@
+// Copyright Epic Games, Inc. All Rights Reserved.
+
+using System.Runtime.CompilerServices;
+
+[assembly: InternalsVisibleTo("Horde.Server.Tests")]
diff --git a/Engine/Source/Programs/Horde/Horde.Server/Properties/launchSettings.json b/Engine/Source/Programs/Horde/Horde.Server/Properties/launchSettings.json
new file mode 100644
--- /dev/null
+++ b/Engine/Source/Programs/Horde/Horde.Server/Properties/launchSettings.json
@@ -0,0 +1,15 @@
+{
+  "profiles": {
+    "HordeServer": {
+      "commandName": "Project",
+      "commandLineArgs": "",
+      "environmentVariables": {
+        "ASPNETCORE_ENVIRONMENT": "Local",
+        "AWS_DEFAULT_PROFILE": "okta-team-devtools-dev_team-devtools-developer",
+        "AWS_PROFILE": "okta-team-devtools-dev_team-devtools-developer",
+        "AWS:Region": "us-east-1"
+      },
+      "hotReloadEnabled": false
+    }
+  }
+}
\ No newline at end of file
diff --git a/Engine/Source/Programs/Horde/Horde.Server/Secrets/Credential.cs b/Engine/Source/Programs/Horde/Horde.Server/Secrets/Credential.cs
new file mode 100644
--- /dev/null
+++ b/Engine/Source/Programs/Horde/Horde.Server/Secrets/Credential.cs
@@ -0,0 +1,69 @@
+// Copyright Epic Games, Inc. All Rights Reserved.
+
+using System.Collections.Generic;
+using Horde.Server.Acls;
+using MongoDB.Bson;
+using MongoDB.Bson.Serialization.Attributes;
+
+namespace Horde.Server.Secrets
+{
+	/// <summary>
+	/// Stores information about a credential
+	/// </summary>
+	public class Credential
+	{
+		/// <summary>
+		/// Unique id for this credential
+		/// </summary>
+		[BsonId]
+		public ObjectId Id { get; set; }
+
+		/// <summary>
+		/// Name of this credential
+		/// </summary>
+		[BsonRequired]
+		public string Name { get; set; }
+
+		/// <summary>
+		/// The normalized name of this credential
+		/// </summary>
+		[BsonRequired]
+		public string NormalizedName { get; set; }
+
+		/// <summary>
+		/// Properties for this credential
+		/// </summary>
+		public Dictionary<string, string> Properties { get; set; } = new Dictionary<string, string>();
+
+		/// <summary>
+		/// Private constructor for serialization
+		/// </summary>
+		[BsonConstructor]
+		private Credential()
+		{
+			Name = null!;
+			NormalizedName = null!;
+		}
+
+		/// <summary>
+		/// Constructor
+		/// </summary>
+		/// <param name="name">Name of this credential</param>
+		public Credential(string name)
+		{
+			Id = ObjectId.GenerateNewId();
+			Name = name;
+			NormalizedName = GetNormalizedName(name);
+		}
+
+		/// <summary>
+		/// Gets the normalized form of a name, used for case-insensitive comparisons
+		/// </summary>
+		/// <param name="name">The name to normalize</param>
+		/// <returns>The normalized name</returns>
+		public static string GetNormalizedName(string name)
+		{
+			return name.ToUpperInvariant();
+		}
+	}
+}
diff --git a/Engine/Source/Programs/Horde/Horde.Server/Secrets/CredentialService.cs b/Engine/Source/Programs/Horde/Horde.Server/Secrets/CredentialService.cs
new file mode 100644
--- /dev/null
+++ b/Engine/Source/Programs/Horde/Horde.Server/Secrets/CredentialService.cs
@@ -0,0 +1,151 @@
+// Copyright Epic Games, Inc. All Rights Reserved.
+
+using System.Collections.Generic;
+using System.Linq;
+using System.Security.Claims;
+using System.Threading.Tasks;
+using Horde.Server.Acls;
+using Horde.Server.Server;
+using MongoDB.Bson;
+using MongoDB.Driver;
+
+namespace Horde.Server.Secrets
+{
+	/// <summary>
+	/// Service for accessing and modifying credentials
+	/// </summary>
+	public class CredentialService
+	{
+		/// <summary>
+		/// The ACL service instance
+		/// </summary>
+		readonly AclService _aclService;
+
+		/// <summary>
+		/// Collection of credential documents
+		/// </summary>
+		readonly IMongoCollection<Credential> _credentials; 
+
+		/// <summary>
+		/// Constructor
+		/// </summary>
+		/// <param name="aclService">The ACL service instance</param>
+		/// <param name="mongoService">The database service instance</param>
+		public CredentialService(AclService aclService, MongoService mongoService)
+		{
+			_aclService = aclService;
+
+			_credentials = mongoService.Credentials;
+
+			if (!mongoService.ReadOnlyMode)
+			{
+				_credentials.Indexes.CreateOne(new CreateIndexModel<Credential>(Builders<Credential>.IndexKeys.Ascending(x => x.NormalizedName), new CreateIndexOptions { Unique = true }));
+			}
+		}
+
+		/// <summary>
+		/// Creates a new credential
+		/// </summary>
+		/// <param name="name">Name of the new credential</param>
+		/// <param name="properties">Properties for the new credential</param>
+		/// <returns>The new credential document</returns>
+		public async Task<Credential> CreateCredentialAsync(string name, Dictionary<string, string>? properties)
+		{
+			Credential newCredential = new Credential(name);
+			if (properties != null)
+			{
+				newCredential.Properties = new Dictionary<string, string>(properties, newCredential.Properties.Comparer);
+			}
+
+			await _credentials.InsertOneAsync(newCredential);
+			return newCredential;
+		}
+
+		/// <summary>
+		/// Updates an existing credential
+		/// </summary>
+		/// <param name="id">Unique id of the credential</param>
+		/// <param name="newName">The new name for the credential</param>
+		/// <param name="newProperties">Properties on the credential to update. Any properties with a value of null will be removed.</param>
+		/// <returns>Async task object</returns>
+		public async Task UpdateCredentialAsync(ObjectId id, string? newName, Dictionary<string, string?>? newProperties)
+		{
+			UpdateDefinitionBuilder<Credential> updateBuilder = Builders<Credential>.Update;
+
+			List<UpdateDefinition<Credential>> updates = new List<UpdateDefinition<Credential>>();
+			if (newName != null)
+			{
+				updates.Add(updateBuilder.Set(x => x.Name, newName));
+				updates.Add(updateBuilder.Set(x => x.NormalizedName, Credential.GetNormalizedName(newName)));
+			}
+			if (newProperties != null)
+			{
+				foreach (KeyValuePair<string, string?> pair in newProperties)
+				{
+					if (pair.Value == null)
+					{
+						updates.Add(updateBuilder.Unset(x => x.Properties[pair.Key]));
+					}
+					else
+					{
+						updates.Add(updateBuilder.Set(x => x.Properties[pair.Key], pair.Value));
+					}
+				}
+			}
+
+			if (updates.Count > 0)
+			{
+				await _credentials.FindOneAndUpdateAsync<Credential>(x => x.Id == id, updateBuilder.Combine(updates));
+			}
+		}
+
+		/// <summary>
+		/// Gets all the available credentials
+		/// </summary>
+		/// <returns>List of project documents</returns>
+		public Task<List<Credential>> FindCredentialsAsync(string? name)
+		{
+			FilterDefinitionBuilder<Credential> filterBuilder = Builders<Credential>.Filter;
+
+			FilterDefinition<Credential> filter = FilterDefinition<Credential>.Empty;
+			if (name != null)
+			{
+				filter &= filterBuilder.Eq(x => x.NormalizedName, Credential.GetNormalizedName(name));
+			}
+
+			return _credentials.Find(filter).ToListAsync();
+		}
+
+		/// <summary>
+		/// Gets a credential by ID
+		/// </summary>
+		/// <param name="id">Unique id of the credential</param>
+		/// <returns>The credential document</returns>
+		public async Task<Credential?> GetCredentialAsync(ObjectId id)
+		{
+			return await _credentials.Find<Credential>(x => x.Id == id).FirstOrDefaultAsync();
+		}
+
+		/// <summary>
+		/// Gets a credential by name
+		/// </summary>
+		/// <param name="name">Name of the credential</param>
+		/// <returns>The credential document</returns>
+		public async Task<Credential?> GetCredentialAsync(string name)
+		{
+			string normalizedName = Credential.GetNormalizedName(name);
+			return await _credentials.Find<Credential>(x => x.NormalizedName == normalizedName).FirstOrDefaultAsync();
+		}
+
+		/// <summary>
+		/// Deletes a credential by id
+		/// </summary>
+		/// <param name="id">Unique id of the credential</param>
+		/// <returns>True if the credential was deleted</returns>
+		public async Task<bool> DeleteCredentialAsync(ObjectId id)
+		{
+			DeleteResult result = await _credentials.DeleteOneAsync<Credential>(x => x.Id == id);
+			return result.DeletedCount > 0;
+		}
+	}
+}
diff --git a/Engine/Source/Programs/Horde/Horde.Server/Secrets/CredentialsController.cs b/Engine/Source/Programs/Horde/Horde.Server/Secrets/CredentialsController.cs
new file mode 100644
--- /dev/null
+++ b/Engine/Source/Programs/Horde/Horde.Server/Secrets/CredentialsController.cs
@@ -0,0 +1,168 @@
+// Copyright Epic Games, Inc. All Rights Reserved.
+
+using System.Collections.Generic;
+using System.Threading.Tasks;
+using Horde.Server.Acls;
+using Horde.Server.Credentials;
+using Horde.Server.Server;
+using Horde.Server.Utilities;
+using Microsoft.AspNetCore.Authorization;
+using Microsoft.AspNetCore.Mvc;
+using Microsoft.Extensions.Options;
+using MongoDB.Bson;
+
+namespace Horde.Server.Secrets
+{
+	/// <summary>
+	/// Controller for the /api/v1/credentials endpoint
+	/// </summary>
+	[ApiController]
+	[Authorize]
+	[Route("[controller]")]
+	public class CredentialsController : ControllerBase
+	{
+		private readonly CredentialService _credentialService;
+		private readonly IOptionsSnapshot<GlobalConfig> _globalConfig;
+
+		/// <summary>
+		/// Constructor
+		/// </summary>
+		public CredentialsController(CredentialService credentialService, IOptionsSnapshot<GlobalConfig> globalConfig)
+		{
+			_credentialService = credentialService;
+			_globalConfig = globalConfig;
+		}
+
+		/// <summary>
+		/// Creates a new credential
+		/// </summary>
+		/// <param name="create">Parameters for the new credential.</param>
+		/// <returns>Http result code</returns>
+		[HttpPost]
+		[Route("/api/v1/credentials")]
+		public async Task<ActionResult<CreateCredentialResponse>> CreateCredentialAsync([FromBody] CreateCredentialRequest create)
+		{
+			if(!_globalConfig.Value.Authorize(AclAction.CreateCredential, User))
+			{
+				return Forbid();
+			}
+
+			Credential newCredential = await _credentialService.CreateCredentialAsync(create.Name, create.Properties);
+			return new CreateCredentialResponse(newCredential.Id.ToString());
+		}
+
+		/// <summary>
+		/// Query all the credentials
+		/// </summary>
+		/// <param name="name">Id of the credential to get information about</param>
+		/// <param name="filter">Filter for the properties to return</param>
+		/// <returns>Information about all the credentials</returns>
+		[HttpGet]
+		[Route("/api/v1/credentials")]
+		[ProducesResponseType(typeof(List<GetCredentialResponse>), 200)]
+		public async Task<ActionResult<object>> FindCredentialAsync([FromQuery] string? name = null, [FromQuery] PropertyFilter? filter = null)
+		{
+			if (!_globalConfig.Value.Authorize(AclAction.ListCredentials, User))
+			{
+				return Forbid();
+			}
+
+			List<Credential> credentials = await _credentialService.FindCredentialsAsync(name);
+
+			List<object> responses = new List<object>();
+			foreach (Credential credential in credentials)
+			{
+				if (_globalConfig.Value.Authorize(AclAction.ViewCredential, User))
+				{
+					responses.Add(new GetCredentialResponse(credential).ApplyFilter(filter));
+				}
+			}
+			return responses;
+		}
+
+		/// <summary>
+		/// Retrieve information about a specific credential
+		/// </summary>
+		/// <param name="credentialId">Id of the credential to get information about</param>
+		/// <param name="filter">Filter for properties to return</param>
+		/// <returns>Information about the requested credential</returns>
+		[HttpGet]
+		[Route("/api/v1/credentials/{credentialId}")]
+		[ProducesResponseType(typeof(GetCredentialResponse), 200)]
+		public async Task<ActionResult<object>> GetCredentialAsync(string credentialId, [FromQuery] PropertyFilter? filter = null)
+		{
+			ObjectId projectIdValue = credentialId.ToObjectId();
+
+			Credential? credential = await _credentialService.GetCredentialAsync(projectIdValue);
+			if (credential == null)
+			{
+				return NotFound();
+			}
+
+			if (!_globalConfig.Value.Authorize(AclAction.ViewCredential, User))
+			{
+				return Forbid();
+			}
+
+			return new GetCredentialResponse(credential).ApplyFilter(filter);
+		}
+
+		/// <summary>
+		/// Update a credential's properties.
+		/// </summary>
+		/// <param name="credentialId">Id of the credential to update</param>
+		/// <param name="update">Items on the credential to update</param>
+		/// <returns>Http result code</returns>
+		[HttpPut]
+		[Route("/api/v1/credentials/{credentialId}")]
+		public async Task<ActionResult> UpdateCredentialAsync(string credentialId, [FromBody] UpdateCredentialRequest update)
+		{
+			ObjectId credentialIdValue = credentialId.ToObjectId();
+
+			Credential? credential = await _credentialService.GetCredentialAsync(credentialIdValue);
+			if(credential == null)
+			{
+				return NotFound();
+			}
+
+			if (!_globalConfig.Value.Authorize(AclAction.UpdateCredential, User))
+			{
+				return Forbid();
+			}
+			if (update.Acl != null && !_globalConfig.Value.Authorize(AclAction.ChangePermissions, User))
+			{
+				return Forbid();
+			}
+
+			await _credentialService.UpdateCredentialAsync(credentialIdValue, update.Name, update.Properties);
+			return new OkResult();
+		}
+
+		/// <summary>
+		/// Delete a credential
+		/// </summary>
+		/// <param name="credentialId">Id of the credential to delete</param>
+		/// <returns>Http result code</returns>
+		[HttpDelete]
+		[Route("/api/v1/credentials/{credentialId}")]
+		public async Task<ActionResult> DeleteCredentialAsync(string credentialId)
+		{
+			ObjectId credentialIdValue = credentialId.ToObjectId();
+
+			Credential? credential = await _credentialService.GetCredentialAsync(credentialIdValue);
+			if (credential == null)
+			{
+				return NotFound();
+			}
+			if (!_globalConfig.Value.Authorize(AclAction.DeleteCredential, User))
+			{
+				return Forbid();
+			}
+			if (!await _credentialService.DeleteCredentialAsync(credentialIdValue))
+			{
+				return NotFound();
+			}
+			return new OkResult();
+		}
+	}
+}
diff --git a/Engine/Source/Programs/Horde/Horde.Server/Server/AccountController.cs b/Engine/Source/Programs/Horde/Horde.Server/Server/AccountController.cs
new file mode 100644
--- /dev/null
+++ b/Engine/Source/Programs/Horde/Horde.Server/Server/AccountController.cs
@@ -0,0 +1,138 @@
+// Copyright Epic Games, Inc. All Rights Reserved.
+
+using System;
+using System.Globalization;
+using System.Net;
+using System.Text;
+using System.Threading.Tasks;
+using Horde.Server.Acls;
+using Horde.Server.Authentication;
+using Microsoft.AspNetCore.Authentication;
+using Microsoft.AspNetCore.Authentication.Cookies;
+using Microsoft.AspNetCore.Authentication.OpenIdConnect;
+using Microsoft.AspNetCore.Mvc;
+using Microsoft.Extensions.Options;
+
+namespace Horde.Server.Server
+{
+	/// <summary>
+	/// Controller managing account status
+	/// </summary>
+	[ApiController]
+	[Route("[controller]")]
+	public class AccountController : Controller
+	{
+		/// <summary>
+		/// Style sheet for HTML responses
+		/// </summary>
+		const string StyleSheet =
+			"body { font-family: 'Segoe UI', 'Roboto', arial, sans-serif; } " +
+			"p { margin:20px; font-size:13px; } " +
+			"h1 { margin:20px; font-size:32px; font-weight:200; } " +
+			"table { margin:10px 20px; } " +
+			"td { margin:5px; font-size:13px; }";
+
+		readonly string _authenticationScheme;
+		readonly IOptionsSnapshot<GlobalConfig> _globalConfig;
+
+		/// <summary>
+		/// Constructor
+		/// </summary>
+		public AccountController(IOptionsMonitor<ServerSettings> serverSettings, IOptionsSnapshot<GlobalConfig> globalConfig)
+		{
+			_authenticationScheme = GetAuthScheme(serverSettings.CurrentValue.AuthMethod);
+			_globalConfig = globalConfig;
+		}
+
+		/// <summary>
+		/// Get auth scheme name for a given auth method
+		/// </summary>
+		/// <param name="method">Authentication method</param>
+		/// <returns>Name of authentication scheme</returns>
+		public static string GetAuthScheme(AuthMethod method)
+		{
+			return method switch
+			{
+				AuthMethod.Anonymous => AnonymousAuthenticationHandler.AuthenticationScheme,
+				AuthMethod.Okta => OktaDefaults.AuthenticationScheme,
+				AuthMethod.OpenIdConnect => OpenIdConnectDefaults.AuthenticationScheme,
+				_ => throw new ArgumentOutOfRangeException(nameof(method), method, null)
+			};
+		}
+
+		/// <summary>
+		/// Gets the current login status
+		/// </summary>
+		/// <returns>The current login state</returns>
+		[HttpGet]
+		[Route("/account")]
+		public ActionResult State()
+		{
+			StringBuilder content = new StringBuilder();
+			content.Append($"<html><style>{StyleSheet}</style><h1>Horde Server</h1>");
+			if (User.Identity?.IsAuthenticated ?? false)
+			{
+				content.Append(CultureInfo.InvariantCulture, $"<p>User <b>{User.Identity?.Name}</b> is logged in. <a href=\"/account/logout\">Log out</a></p>");
+				if (_globalConfig.Value.Authorize(AclAction.AdminWrite, User))
+				{
+					content.Append("<p>");
+					content.Append("<a href=\"/api/v1/admin/token\">Get bearer token</a><br/>");
+					content.Append("<a href=\"/api/v1/admin/registrationtoken\">Get agent registration token</a><br/>");
+					content.Append("<a href=\"/api/v1/admin/softwaretoken\">Get agent software upload token</a><br/>");
+					content.Append("<a href=\"/api/v1/admin/softwaredownloadtoken\">Get agent software download token</a><br/>");
+					content.Append("<a href=\"/api/v1/admin/configtoken\">Get configuration token</a><br/>");
+					content.Append("<a href=\"/api/v1/admin/chainedjobtoken\">Get chained job token</a><br/>");
+					content.Append("</p>");
+				}
+				content.Append(CultureInfo.InvariantCulture, $"<p>Claims for {User.Identity?.Name}:");
+				content.Append("<table>");
+				foreach (System.Security.Claims.Claim claim in User.Claims)
+				{
+					content.Append(CultureInfo.InvariantCulture, $"<tr><td>{claim.Type}</td><td>{claim.Value}</td></tr>");
+				}
+				content.Append("</table>");
+				content.Append("</p>");
+
+				content.Append(CultureInfo.InvariantCulture, $"<p>Built from Perforce</p>");
+			}
+			else
+			{
+				content.Append("<p><a href=\"/account/login\"><b>Login with OAuth2</b></a></p>");
+			}
+			content.Append("</html>");
+			return new ContentResult { ContentType = "text/html", StatusCode = (int)HttpStatusCode.OK, Content = content.ToString() };
+		}
+
+		/// <summary>
+		/// Login to the server
+		/// </summary>
+		/// <returns>Http result</returns>
+		[HttpGet]
+		[Route("/account/login")]
+		public IActionResult Login()
+		{
+			return new ChallengeResult(_authenticationScheme, new AuthenticationProperties { RedirectUri = "/account" });
+		}
+
+		/// <summary>
+		/// Logout of the current account
+		/// </summary>
+		/// <returns>Http result</returns>
+		[HttpGet]
+		[Route("/account/logout")]
+		public async Task<IActionResult> Logout()
+		{
+			await HttpContext.SignOutAsync(CookieAuthenticationDefaults.AuthenticationScheme);
+			try
+			{
+				await HttpContext.SignOutAsync(_authenticationScheme);
+			}
+			catch
+			{
+			}
+
+			string content = $"<html><style>{StyleSheet}</style><body onload=\"setTimeout(function(){{ window.location = '/account'; }}, 2000)\"><p>User has been logged out. Returning to login page.</p></body></html>";
+			return new ContentResult { ContentType = "text/html", StatusCode = (int)HttpStatusCode.OK, Content = content };
+		}
+	}
+}
diff --git a/Engine/Source/Programs/Horde/Horde.Server/Server/AdminController.cs b/Engine/Source/Programs/Horde/Horde.Server/Server/AdminController.cs
new file mode 100644
--- /dev/null
+++ b/Engine/Source/Programs/Horde/Horde.Server/Server/AdminController.cs
@@ -0,0 +1,237 @@
+// Copyright Epic Games, Inc. All Rights Reserved.
+
+using System;
+using System.Collections.Generic;
+using System.Linq;
+using System.Security.Claims;
+using System.Threading.Tasks;
+using Horde.Server.Acls;
+using Horde.Server.Utilities;
+using Microsoft.AspNetCore.Authorization;
+using Microsoft.AspNetCore.Mvc;
+using Microsoft.Extensions.Options;
+
+namespace Horde.Server.Server
+{
+	/// <summary>
+	/// Object containing settings for the server
+	/// </summary>
+	public class AdminSettings
+	{
+		/// <summary>
+		/// The default perforce server
+		/// </summary>
+		public string? DefaultServerAndPort { get; set; }
+
+		/// <summary>
+		/// The default perforce username
+		/// </summary>
+		public string? DefaultUserName { get; set; }
+
+		/// <summary>
+		/// The default perforce password
+		/// </summary>
+		public string? DefaultPassword { get; set; }
+	}
+
+	/// <summary>
+	/// The conform limit value
+	/// </summary>
+	public class ConformSettings
+	{
+		/// <summary>
+		/// Maximum number of conforms allowed at once
+		/// </summary>
+		public int MaxCount { get; set; }
+	}
+
+	/// <summary>
+	/// Controller managing account status
+	/// </summary>
+	[ApiController]
+	[Authorize]
+	[Route("[controller]")]
+	public class AdminController : HordeControllerBase
+	{
+		readonly AclService _aclService;
+		readonly IOptionsSnapshot<GlobalConfig> _globalConfig;
+		readonly IOptions<ServerSettings> _serverSettings;
+
+		/// <summary>
+		/// Constructor
+		/// </summary>
+		public AdminController(AclService aclService, IOptionsSnapshot<GlobalConfig> globalConfig, IOptions<ServerSettings> serverSettings)
+		{
+			_aclService = aclService;
+			_globalConfig = globalConfig;
+			_serverSettings = serverSettings;
+		}
+
+		/// <summary>
+		/// Issues a token for the given roles. Issues a token for the current user if not specified.
+		/// </summary>
+		/// <returns>Administrative settings for the server</returns>
+		[HttpGet]
+		[Route("/api/v1/admin/token")]
+		public async Task<ActionResult<string>> GetTokenAsync()
+		{
+			if (!_globalConfig.Value.Authorize(AclAction.IssueBearerToken, User))
+			{
+				return Forbid(AclAction.IssueBearerToken);
+			}
+
+			return await _aclService.IssueBearerTokenAsync(User.Claims, GetDefaultExpiryTime());
+		}
+
+		/// <summary>
+		/// Returns the fully parsed config object.
+		/// </summary>
+		[HttpGet]
+		[Route("/api/v1/admin/config")]
+		public ActionResult<object> GetConfig()
+		{
+			if (!_globalConfig.Value.Authorize(AclAction.AdminRead, User))
+			{
+				return Forbid(AclAction.AdminRead);
+			}
+
+			return _globalConfig.Value;
+		}
+
+		/// <summary>
+		/// Issues a token for the given roles. Issues a token for the current user if not specified.
+		/// </summary>
+		/// <param name="roles">Roles for the new token</param>
+		/// <returns>Administrative settings for the server</returns>
+		[HttpGet]
+		[Route("/api/v1/admin/roletoken")]
+		public async Task<ActionResult<string>> GetRoleTokenAsync([FromQuery] string roles)
+		{
+			if (!_globalConfig.Value.Authorize(AclAction.AdminWrite, User))
+			{
+				return Forbid(AclAction.AdminWrite);
+			}
+
+			List<Claim> claims = new List<Claim>();
+			claims.AddRange(roles.Split('+').Select(x => new Claim(ClaimTypes.Role, x)));
+
+			return await _aclService.IssueBearerTokenAsync(claims, GetDefaultExpiryTime());
+		}
+
+		/// <summary>
+		/// Issues a token for the given roles. Issues a token for the current user if not specified.
+		/// </summary>
+		/// <returns>Administrative settings for the server</returns>
+		[HttpGet]
+		[Route("/api/v1/admin/registrationtoken")]
+		public async Task<ActionResult<string>> GetRegistrationTokenAsync()
+		{
+			if (!_globalConfig.Value.Authorize(AclAction.AdminWrite, User))
+			{
+				return Forbid(AclAction.AdminWrite);
+			}
+
+			List<AclClaimConfig> claims = new List<AclClaimConfig>();
+			claims.Add(new AclClaimConfig(ClaimTypes.Name, User.Identity?.Name ?? "Unknown"));
+			claims.Add(HordeClaims.AgentRegistrationClaim);
+
+			return await _aclService.IssueBearerTokenAsync(claims, null);
+		}
+
+		/// <summary>
+		/// Issues a token valid to upload new versions of the agent software.
+		/// </summary>
+		/// <returns>Administrative settings for the server</returns>
+		[HttpGet]
+		[Route("/api/v1/admin/softwaretoken")]
+		public async Task<ActionResult<string>> GetSoftwareTokenAsync()
+		{
+			if (!_globalConfig.Value.Authorize(AclAction.AdminWrite, User))
+			{
+				return Forbid(AclAction.AdminWrite);
+			}
+
+			List<AclClaimConfig> claims = new List<AclClaimConfig>();
+			claims.Add(new AclClaimConfig(ClaimTypes.Name, User.Identity?.Name ?? "Unknown"));
+			claims.Add(HordeClaims.UploadSoftwareClaim);
+
+			return await _aclService.IssueBearerTokenAsync(claims, null);
+		}
+
+		/// <summary>
+		/// Issues a token valid to download new versions of the agent software.
+		/// </summary>
+		/// <returns>Administrative settings for the server</returns>
+		[HttpGet]
+		[Route("/api/v1/admin/softwaredownloadtoken")]
+		public async Task<ActionResult<string>> GetSoftwareDownloadTokenAsync()
+		{
+			if (!_globalConfig.Value.Authorize(AclAction.AdminRead, User))
+			{
+				return Forbid(AclAction.AdminRead);
+			}
+
+			List<AclClaimConfig> claims = new List<AclClaimConfig>();
+			claims.Add(new AclClaimConfig(ClaimTypes.Name, User.Identity?.Name ?? "Unknown"));
+			claims.Add(HordeClaims.DownloadSoftwareClaim);
+
+			return await _aclService.IssueBearerTokenAsync(claims, null);
+		}
+
+		/// <summary>
+		/// Issues a token valid to configure streams and projects
+		/// </summary>
+		/// <returns>Administrative settings for the server</returns>
+		[HttpGet]
+		[Route("/api/v1/admin/configtoken")]
+		public async Task<ActionResult<string>> GetConfigToken()
+		{
+			if (!_globalConfig.Value.Authorize(AclAction.AdminRead, User))
+			{
+				return Forbid(AclAction.AdminRead);
+			}
+
+			List<AclClaimConfig> claims = new List<AclClaimConfig>();
+			claims.Add(new AclClaimConfig(ClaimTypes.Name, User.Identity?.Name ?? "Unknown"));
+			claims.Add(HordeClaims.ConfigureProjectsClaim);
+
+			return await _aclService.IssueBearerTokenAsync(claims, null);
+		}
+
+		/// <summary>
+		/// Issues a token valid to start chained jobs
+		/// </summary>
+		/// <returns>Administrative settings for the server</returns>
+		[HttpGet]
+		[Route("/api/v1/admin/chainedjobtoken")]
+		public async Task<ActionResult<string>> GetChainedJobToken()
+		{
+			if (!_globalConfig.Value.Authorize(AclAction.AdminRead, User))
+			{
+				return Forbid(AclAction.AdminRead);
+			}
+
+			List<AclClaimConfig> claims = new List<AclClaimConfig>();
+			//Claims.Add(new AclClaim(ClaimTypes.Name, User.Identity.Name ?? "Unknown"));
+			claims.Add(HordeClaims.StartChainedJobClaim);
+
+			return await _aclService.IssueBearerTokenAsync(claims, null);
+		}
+
+		/// <summary>
+		/// Gets the default expiry time for a token
+		/// </summary>
+		/// <returns></returns>
+		private TimeSpan? GetDefaultExpiryTime()
+		{
+			ServerSettings serverSettings = _serverSettings.Value;
+
+			TimeSpan? expiryTime = null;
+			if (serverSettings.JwtExpiryTimeHours != -1)
+			{
+				expiryTime = TimeSpan.FromHours(serverSettings.JwtExpiryTimeHours);
+			}
+			return expiryTime;
+		}
+	}
+}
diff --git a/Engine/Source/Programs/Horde/Horde.Server/Server/ConsistencyService.cs b/Engine/Source/Programs/Horde/Horde.Server/Server/ConsistencyService.cs
new file mode 100644
--- /dev/null
+++ b/Engine/Source/Programs/Horde/Horde.Server/Server/ConsistencyService.cs
@@ -0,0 +1,71 @@
+// Copyright Epic Games, Inc. All Rights Reserved.
+
+using System;
+using System.Collections.Generic;
+using System.Linq;
+using System.Threading;
+using System.Threading.Tasks;
+using Horde.Server.Agents.Sessions;
+using Horde.Server.Agents.Leases;
+using Horde.Server.Utilities;
+using HordeCommon;
+using Microsoft.Extensions.Hosting;
+using Microsoft.Extensions.Logging;
+
+namespace Horde.Server.Server
+{
+	/// <summary>
+	/// Service which checks the database for consistency and fixes up any errors
+	/// </summary>
+	class ConsistencyService : IHostedService, IDisposable
+	{
+		readonly ISessionCollection _sessionCollection;
+		readonly ILeaseCollection _leaseCollection;
+		readonly ITicker _ticker;
+		readonly ILogger<ConsistencyService> _logger;
+
+		/// <summary>
+		/// Constructor
+		/// </summary>
+		public ConsistencyService(ISessionCollection sessionCollection, ILeaseCollection leaseCollection, IClock clock, ILogger<ConsistencyService> logger)
+		{
+			_sessionCollection = sessionCollection;
+			_leaseCollection = leaseCollection;
+			_ticker = clock.AddSharedTicker<ConsistencyService>(TimeSpan.FromMinutes(20.0), TickLeaderAsync, logger);
+			_logger = logger;
+		}
+
+		/// <inheritdoc/>
+		public Task StartAsync(CancellationToken cancellationToken) => _ticker.StartAsync();
+
+		/// <inheritdoc/>
+		public Task StopAsync(CancellationToken cancellationToken) => _ticker.StopAsync();
+
+		/// <inheritdoc/>
+		public void Dispose() => _ticker.Dispose();
+
+		/// <summary>
+		/// Poll for inconsistencies in the database
+		/// </summary>
+		/// <param name="stoppingToken">Stopping token</param>
+		/// <returns>Async task</returns>
+		async ValueTask TickLeaderAsync(CancellationToken stoppingToken)
+		{
+			List<ISession> sessions = await _sessionCollection.FindActiveSessionsAsync();
+			Dictionary<SessionId, ISession> sessionIdToInstance = sessions.ToDictionary(x => x.Id, x => x);
+
+			// Find any leases that are still running when their session has terminated
+			List<ILease> leases = await _leaseCollection.FindActiveLeasesAsync();
+			foreach (ILease lease in leases)
+			{
+				if (!sessionIdToInstance.ContainsKey(lease.SessionId))
+				{
+					ISession? session = await _sessionCollection.GetAsync(lease.SessionId);
+					DateTime finishTime = session?.FinishTime ?? DateTime.UtcNow;
+					_logger.LogWarning("Setting finish time for lease {LeaseId} to {FinishTime}", lease.Id, finishTime);
+					await _leaseCollection.TrySetOutcomeAsync(lease.Id, finishTime, LeaseOutcome.Cancelled, null);
+				}
+			}
+		}
+	}
+}
diff --git a/Engine/Source/Programs/Horde/Horde.Server/Server/DebugController.cs b/Engine/Source/Programs/Horde/Horde.Server/Server/DebugController.cs
new file mode 100644
--- /dev/null
+++ b/Engine/Source/Programs/Horde/Horde.Server/Server/DebugController.cs
@@ -0,0 +1,530 @@
+// Copyright Epic Games, Inc. All Rights Reserved.
+
+//#define ENABLE_PUBLIC_DEBUG_CONTROLLER
+#define ENABLE_SECURE_DEBUG_CONTROLLER
+
+
+using System;
+using System.Collections.Generic;
+using System.IO;
+using System.Linq;
+using System.Net;
+using System.Text;
+using System.Text.Json;
+using System.Threading;
+using System.Threading.Tasks;
+using System.Web;
+using EpicGames.Core;
+using Horde.Server.Acls;
+using Horde.Server.Jobs;
+using Horde.Server.Jobs.Graphs;
+using Horde.Server.Jobs.Templates;
+using Horde.Server.Logs;
+using Horde.Server.Utilities;
+using JetBrains.Profiler.SelfApi;
+using Microsoft.AspNetCore.Authorization;
+using Microsoft.AspNetCore.Mvc;
+using Microsoft.Extensions.Logging;
+using Microsoft.Extensions.Options;
+using MongoDB.Driver;
+
+namespace Horde.Server.Server
+{
+#if ENABLE_PUBLIC_DEBUG_CONTROLLER
+	/// <summary>
+	/// Public endpoints for the debug controller
+	/// </summary>
+	[ApiController]
+	public class PublicDebugController : ControllerBase
+	{
+		/// <summary>
+		/// The connection tracker service singleton
+		/// </summary>
+		RequestTrackerService RequestTrackerService;
+
+		IHostApplicationLifetime ApplicationLifetime;
+
+		IDogStatsd DogStatsd;
+
+		/// <summary>
+		/// Constructor
+		/// </summary>
+		/// <param name="RequestTrackerService"></param>
+		/// <param name="ApplicationLifetime"></param>
+		/// <param name="DogStatsd"></param>
+		public PublicDebugController(RequestTrackerService RequestTrackerService, IHostApplicationLifetime ApplicationLifetime, IDogStatsd DogStatsd)
+		{
+			RequestTrackerService = RequestTrackerService;
+			ApplicationLifetime = ApplicationLifetime;
+			DogStatsd = DogStatsd;
+		}
+
+		/// <summary>
+		/// Prints all the headers for the incoming request
+		/// </summary>
+		/// <returns>Http result</returns>
+		[HttpGet]
+		[Route("/api/v1/debug/headers")]
+		public IActionResult GetRequestHeaders()
+		{
+			StringBuilder Content = new StringBuilder();
+			Content.AppendLine("<html><body><pre>");
+			foreach (KeyValuePair<string, StringValues> Pair in HttpContext.Request.Headers)
+			{
+				foreach (string Value in Pair.Value)
+				{
+					Content.AppendLine(HttpUtility.HtmlEncode($"{Pair.Key}: {Value}"));
+				}
+			}
+			Content.Append("</pre></body></html>");
+			return new ContentResult { ContentType = "text/html", StatusCode = (int)HttpStatusCode.OK, Content = Content.ToString() };
+		}
+
+		/// <summary>
+		/// Waits specified number of milliseconds and then returns a response
+		/// Used for testing timeouts proxy settings.
+		/// </summary>
+		/// <returns>Http result</returns>
+		[HttpGet]
+		[Route("/api/v1/debug/wait")]
+		public async Task<ActionResult> GetAndWait([FromQuery] int WaitTimeMs = 1000)
+		{
+			await Task.Delay(WaitTimeMs);
+			string Content = $"Waited {WaitTimeMs} ms. " + new Random().Next(0, 10000000);
+			return new ContentResult { ContentType = "text/plain", StatusCode = (int)HttpStatusCode.OK, Content = Content };
+		}
+
+		/// <summary>
+		/// Waits specified number of milliseconds and then throws an exception
+		/// Used for testing graceful shutdown and interruption of outstanding requests.
+		/// </summary>
+		/// <returns>Http result</returns>
+		[HttpGet]
+		[Route("/api/v1/debug/exception")]
+		public async Task<ActionResult> ThrowException([FromQuery] int WaitTimeMs = 0)
+		{
+			await Task.Delay(WaitTimeMs);
+			throw new Exception("Test exception triggered by debug controller!");
+		}
+
+		/// <summary>
+		/// Trigger an increment of a DogStatsd metric
+		/// </summary>
+		/// <returns>Http result</returns>
+		[HttpGet]
+		[Route("/api/v1/debug/metric")]
+		public ActionResult TriggerMetric([FromQuery] int Value = 10)
+		{
+			DogStatsd.Increment("hordeMetricTest", Value);
+			return Ok("Incremented metric 'hordeMetricTest' Type: " + DogStatsd.GetType());
+		}
+
+		/// <summary>
+		/// Display metrics related to the .NET runtime
+		/// </summary>
+		/// <returns>Http result</returns>
+		[HttpGet]
+		[Route("/api/v1/debug/dotnet-metrics")]
+		public ActionResult DotNetMetrics()
+		{
+			ThreadPool.GetMaxThreads(out int MaxWorkerThreads, out int MaxIoThreads);
+			ThreadPool.GetAvailableThreads(out int FreeWorkerThreads, out int FreeIoThreads);
+			ThreadPool.GetMinThreads(out int MinWorkerThreads, out int MinIoThreads);
+
+			int BusyIoThreads = MaxIoThreads - FreeIoThreads;
+			int BusyWorkerThreads = MaxWorkerThreads - FreeWorkerThreads;
+
+			StringBuilder Content = new StringBuilder();
+			Content.AppendLine("Threads:");
+			Content.AppendLine("-------------------------------------------------------------");
+			Content.AppendLine("Worker busy={0,-5} free={1,-5} min={2,-5} max={3,-5}", BusyWorkerThreads, FreeWorkerThreads, MinWorkerThreads, MaxWorkerThreads);
+			Content.AppendLine("  IOCP busy={0,-5} free={1,-5} min={2,-5} max={3,-5}", BusyIoThreads, FreeIoThreads, MinIoThreads, MaxWorkerThreads);
+
+
+			NumberFormatInfo Nfi = (NumberFormatInfo)CultureInfo.InvariantCulture.NumberFormat.Clone();
+			Nfi.NumberGroupSeparator = " ";
+
+			string FormatBytes(long Number)
+			{
+				return (Number / 1024 / 1024).ToString("#,0", Nfi) + " MB";
+			}
+
+			GCMemoryInfo GcMemoryInfo = GC.GetGCMemoryInfo();
+			Content.AppendLine("");
+			Content.AppendLine("");
+			Content.AppendLine("Garbage collection (GC):");
+			Content.AppendLine("-------------------------------------------------------------");
+			Content.AppendLine("              Latency mode: " + GCSettings.LatencyMode);
+			Content.AppendLine("              Is server GC: " + GCSettings.IsServerGC);
+			Content.AppendLine("              Total memory: " + FormatBytes(GC.GetTotalMemory(false)));
+			Content.AppendLine("           Total allocated: " + FormatBytes(GC.GetTotalAllocatedBytes(false)));
+			Content.AppendLine("                 Heap size: " + FormatBytes(GcMemoryInfo.HeapSizeBytes));
+			Content.AppendLine("                Fragmented: " + FormatBytes(GcMemoryInfo.FragmentedBytes));
+			Content.AppendLine("               Memory Load: " + FormatBytes(GcMemoryInfo.MemoryLoadBytes));
+			Content.AppendLine("    Total available memory: " + FormatBytes(GcMemoryInfo.TotalAvailableMemoryBytes));
+			Content.AppendLine("High memory load threshold: " + FormatBytes(GcMemoryInfo.HighMemoryLoadThresholdBytes));
+
+			return Ok(Content.ToString());
+		}
+
+		/// <summary>
+		/// Force a full GC of all generations
+		/// </summary>
+		/// <returns>Prints time taken in ms</returns>
+		[HttpGet]
+		[Route("/api/v1/debug/force-gc")]
+		public ActionResult ForceTriggerGc()
+		{
+			Stopwatch Timer = new Stopwatch();
+			Timer.Start();
+			GC.Collect();
+			Timer.Stop();
+			return Ok($"Time taken: {Timer.Elapsed.TotalMilliseconds} ms");
+		}
+
+		/// <summary>
+		/// Lists requests in progress
+		/// </summary>
+		/// <returns>HTML result</returns>
+		[HttpGet]
+		[Route("/api/v1/debug/requests-in-progress")]
+		public ActionResult GetRequestsInProgress()
+		{
+			StringBuilder Content = new StringBuilder();
+			Content.AppendLine("<html><body>");
+			Content.AppendLine("<h1>Requests in progress</h1>");
+			Content.AppendLine("<table border=\"1\">");
+			Content.AppendLine("<tr>");
+			Content.AppendLine("<th>Request Trace ID</th>");
+			Content.AppendLine("<th>Path</th>");
+			Content.AppendLine("<th>Started At</th>");
+			Content.AppendLine("<th>Age</th>");
+			Content.AppendLine("</tr>");
+
+			List<KeyValuePair<string, TrackedRequest>> Requests = RequestTrackerService.GetRequestsInProgress().ToList();
+			Requests.Sort((A, B) => A.Value.StartedAt.CompareTo(B.Value.StartedAt));
+
+			foreach (KeyValuePair<string, TrackedRequest> Entry in Requests)
+			{
+				Content.Append("<tr>");
+				Content.AppendLine($"<td>{Entry.Key}</td>");
+				Content.AppendLine($"<td>{Entry.Value.Request.Path}</td>");
+				Content.AppendLine($"<td>{Entry.Value.StartedAt}</td>");
+				Content.AppendLine($"<td>{Entry.Value.GetTimeSinceStartInMs()} ms</td>");
+				Content.Append("</tr>");
+			}
+			Content.Append("</table>\n</body>\n</html>");
+
+			return new ContentResult { ContentType = "text/html", StatusCode = (int)HttpStatusCode.OK, Content = Content.ToString() };
+		}
+
+		/*
+		// Used during development only
+		[HttpGet]
+		[Route("/api/v1/debug/stop")]
+		public ActionResult StopApp()
+		{
+			Task.Run(async () =>
+			{
+				await Task.Delay(100);
+				ApplicationLifetime.StopApplication();
+			});
+			
+			return new ContentResult { ContentType = "text/plain", StatusCode = (int)HttpStatusCode.OK, Content = "App stopping..." };
+		}
+		/**/
+	}
+#endif
+#if ENABLE_SECURE_DEBUG_CONTROLLER
+	/// <summary>
+	/// Controller managing account status
+	/// </summary>
+	[ApiController]
+	[Authorize]
+	public class SecureDebugController : ControllerBase
+	{
+		private static readonly Random s_random = new ();
+		
+		private readonly AclService _aclService;
+		private readonly MongoService _mongoService;
+		private readonly JobTaskSource _jobTaskSource;
+		private readonly ITemplateCollection _templateCollection;
+		private readonly IGraphCollection _graphCollection;
+		private readonly ILogFileCollection _logFileCollection;
+		private readonly IOptions<ServerSettings> _settings;
+		private readonly IOptionsSnapshot<GlobalConfig> _globalConfig;
+		private readonly ILogger<SecureDebugController> _logger;
+
+		/// <summary>
+		/// Constructor
+		/// </summary>
+		public SecureDebugController(AclService aclService, MongoService mongoService, JobTaskSource jobTaskSource,
+			ITemplateCollection templateCollection, IGraphCollection graphCollection,
+			ILogFileCollection logFileCollection, IOptions<ServerSettings> settings, IOptionsSnapshot<GlobalConfig> globalConfig, ILogger<SecureDebugController> logger)
+		{
+			_aclService = aclService;
+			_mongoService = mongoService;
+			_jobTaskSource = jobTaskSource;
+			_templateCollection = templateCollection;
+			_graphCollection = graphCollection;
+			_logFileCollection = logFileCollection;
+			_settings = settings;
+			_globalConfig = globalConfig;
+			_logger = logger;
+		}
+
+		/// <summary>
+		/// Prints all the environment variables
+		/// </summary>
+		/// <returns>Http result</returns>
+		[HttpGet]
+		[Route("/api/v1/debug/environment")]
+		public ActionResult GetServerEnvVars()
+		{
+			if (!_globalConfig.Value.Authorize(AclAction.AdminRead, User))
+			{
+				return Forbid();
+			}
+
+			StringBuilder content = new StringBuilder();
+			content.AppendLine("<html><body><pre>");
+			foreach (System.Collections.DictionaryEntry? pair in System.Environment.GetEnvironmentVariables())
+			{
+				if (pair != null)
+				{
+					content.AppendLine(HttpUtility.HtmlEncode($"{pair.Value.Key}={pair.Value.Value}"));
+				}
+			}
+			content.Append("</pre></body></html>");
+			return new ContentResult { ContentType = "text/html", StatusCode = (int)HttpStatusCode.OK, Content = content.ToString() };
+		}
+
+		/// <summary>
+		/// Returns diagnostic information about the current state of the queue
+		/// </summary>
+		/// <returns>Information about the queue</returns>
+		[HttpGet]
+		[Route("/api/v1/debug/queue")]
+		public ActionResult<object> GetQueueStatus()
+		{
+			if (!_globalConfig.Value.Authorize(AclAction.AdminRead, User))
+			{
+				return Forbid();
+			}
+
+			return _jobTaskSource.GetStatus();
+		}
+
+		/// <summary>
+		/// Returns the complete config Horde uses
+		/// </summary>
+		/// <returns>Information about the config</returns>
+		[HttpGet]
+		[Route("/api/v1/debug/config")]
+		public ActionResult GetConfig()
+		{
+			if (!_globalConfig.Value.Authorize(AclAction.AdminRead, User))
+			{
+				return Forbid();
+			}
+
+			JsonSerializerOptions options = new JsonSerializerOptions();
+			options.WriteIndented = true;
+
+			Startup.ConfigureJsonSerializer(options);
+
+			return Ok(JsonSerializer.Serialize(_settings.Value, options));
+		}
+		
+		/// <summary>
+		/// Generate log message of varying size
+		/// </summary>
+		/// <returns>Information about the log message generated</returns>
+		[HttpGet]
+		[Route("/api/v1/debug/generate-log-msg")]
+		public ActionResult GenerateLogMessage(
+			[FromQuery] string? logLevel = null,
+			[FromQuery] int messageLen = 0,
+			[FromQuery] int exceptionMessageLen = 0,
+			[FromQuery] int argCount = 0,
+			[FromQuery] int argLen = 10)
+		{
+			if (!_globalConfig.Value.Authorize(AclAction.AdminRead, User))
+			{
+				return Forbid();
+			}
+
+			string RandomString(int length)
+			{
+				const string Chars = "ABCDEFGHIJKLMNOPQRSTUVWXYZabcdefghijklmnopqrstuvwxyz0123456789";
+				return new string(Enumerable.Repeat(Chars, length).Select(s => s[s_random.Next(s.Length)]).ToArray());
+			}
+			
+			if (!Enum.TryParse(logLevel, out LogLevel logLevelInternal))
+			{
+				logLevelInternal = LogLevel.Information;
+			}
+
+			Exception? exception = null;
+			string message = "Message generated by /api/v1/debug/generate-log-msg";
+			message += RandomString(messageLen);
+
+			if (exceptionMessageLen > 0)
+			{
+				exception = new Exception("Exception from /api/v1/debug/generate-log-msg " + RandomString(exceptionMessageLen));
+			}
+
+			Dictionary<string, object> args = new ();
+			if (argCount > 0)
+			{
+				for (int i = 0; i < argCount; i++)
+				{
+					args["Arg" + i] = "Arg 1 - " + RandomString(argLen);
+				}
+			}
+
+			using IDisposable logScope = _logger.BeginScope(args);
+			
+			// Ignore warning as we explicitly want to build this message manually
+#pragma warning disable CA2254 // Template should be a static expression
+			_logger.Log(logLevelInternal, exception, message);
+#pragma warning restore CA2254
+			
+			return Ok($"Log message generated logLevel={logLevelInternal} messageLen={messageLen} exceptionMessageLen={exceptionMessageLen} argCount={argCount} argLen={argLen}");
+		}
+
+		/// <summary>
+		/// Queries for all graphs
+		/// </summary>
+		/// <returns>The graph definitions</returns>
+		[HttpGet]
+		[Route("/api/v1/debug/graphs")]
+		[ProducesResponseType(200, Type = typeof(GetGraphResponse))]
+		public async Task<ActionResult<List<object>>> GetGraphsAsync([FromQuery] int? index = null, [FromQuery] int? count = null, [FromQuery] PropertyFilter? filter = null)
+		{
+			if (!_globalConfig.Value.Authorize(AclAction.AdminRead, User))
+			{
+				return Forbid();
+			}
+
+			List<IGraph> graphs = await _graphCollection.FindAllAsync(null, index, count);
+			return graphs.ConvertAll(x => new GetGraphResponse(x).ApplyFilter(filter));
+		}
+
+		/// <summary>
+		/// Queries for a particular graph by hash
+		/// </summary>
+		/// <returns>The graph definition</returns>
+		[HttpGet]
+		[Route("/api/v1/debug/graphs/{GraphId}")]
+		[ProducesResponseType(200, Type = typeof(GetGraphResponse))]
+		public async Task<ActionResult<object>> GetGraphAsync(string graphId, [FromQuery] PropertyFilter? filter = null)
+		{
+			if (!_globalConfig.Value.Authorize(AclAction.AdminRead, User))
+			{
+				return Forbid();
+			}
+
+			IGraph graph = await _graphCollection.GetAsync(ContentHash.Parse(graphId));
+			return new GetGraphResponse(graph).ApplyFilter(filter);
+		}
+
+		/// <summary>
+		/// Retrieve metadata about a specific log file
+		/// </summary>
+		/// <param name="logFileId">Id of the log file to get information about</param>
+		/// <param name="filter">Filter for the properties to return</param>
+		/// <returns>Information about the requested project</returns>
+		[HttpGet]
+		[Route("/api/v1/debug/logs/{LogFileId}")]
+		public async Task<ActionResult<object>> GetLogAsync(LogId logFileId, [FromQuery] PropertyFilter? filter = null)
+		{
+			if (!_globalConfig.Value.Authorize(AclAction.AdminRead, User))
+			{
+				return Forbid();
+			}
+
+			ILogFile? logFile = await _logFileCollection.GetLogFileAsync(logFileId, CancellationToken.None);
+			if (logFile == null)
+			{
+				return NotFound();
+			}
+
+			return logFile.ApplyFilter(filter);
+		}
+
+		/// <summary>
+		/// Populate the database with test data
+		/// </summary>
+		/// <returns>Async task</returns>
+		[HttpGet]
+		[Route("/api/v1/debug/collections/{Name}")]
+		public async Task<ActionResult<object>> GetDocumentsAsync(string name, [FromQuery] string? filter = null, [FromQuery] int index = 0, [FromQuery] int count = 10)
+		{
+			if (!_globalConfig.Value.Authorize(AclAction.AdminRead, User))
+			{
+				return Forbid();
+			}
+
+			IMongoCollection<Dictionary<string, object>> collection = _mongoService.GetCollection<Dictionary<string, object>>(name);
+			List<Dictionary<string, object>> documents = await collection.Find(filter ?? "{}").Skip(index).Limit(count).ToListAsync();
+			return documents;
+		}
+		
+		/// <summary>
+		/// Starts the profiler session
+		/// </summary>
+		/// <returns>Text message</returns>
+		[HttpGet]
+		[Route("/api/v1/debug/profiler/start")]
+		public async Task<ActionResult> StartProfiler()
+		{
+			await DotTrace.EnsurePrerequisiteAsync();
+
+			string snapshotDir = Path.Join(Path.GetTempPath(), "horde-profiler-snapshots");
+			if (!Directory.Exists(snapshotDir))
+			{
+				Directory.CreateDirectory(snapshotDir);
+			}
+
+			DotTrace.Config config = new ();
+			config.SaveToDir(snapshotDir);
+			DotTrace.Attach(config);
+			DotTrace.StartCollectingData();
+			
+			return new ContentResult { ContentType = "text/plain", StatusCode = (int)HttpStatusCode.OK, Content = "Profiling session started. Using dir " + snapshotDir };
+		}
+		
+		/// <summary>
+		/// Stops the profiler session
+		/// </summary>
+		/// <returns>Text message</returns>
+		[HttpGet]
+		[Route("/api/v1/debug/profiler/stop")]
+		public ActionResult StopProfiler()
+		{
+			DotTrace.SaveData();
+			DotTrace.Detach();
+			return new ContentResult { ContentType = "text/plain", StatusCode = (int)HttpStatusCode.OK, Content = "Profiling session stopped" };
+		}
+		
+		/// <summary>
+		/// Downloads the captured profiling snapshots
+		/// </summary>
+		/// <returns>A .zip file containing the profiling snapshots</returns>
+		[HttpGet]
+		[Route("/api/v1/debug/profiler/download")]
+		public ActionResult DownloadProfilingData()
+		{
+			string snapshotZipFile = DotTrace.GetCollectedSnapshotFilesArchive(false);
+			if (!System.IO.File.Exists(snapshotZipFile))
+			{
+				return NotFound("The generated snapshot .zip file was not found");
+			}
+			
+			return PhysicalFile(snapshotZipFile, "application/zip", Path.GetFileName(snapshotZipFile));
+		}
+	}
+}
+
+#endif
diff --git a/Engine/Source/Programs/Horde/Horde.Server/Server/MetricService.cs b/Engine/Source/Programs/Horde/Horde.Server/Server/MetricService.cs
new file mode 100644
--- /dev/null
+++ b/Engine/Source/Programs/Horde/Horde.Server/Server/MetricService.cs
@@ -0,0 +1,84 @@
+// Copyright Epic Games, Inc. All Rights Reserved.
+
+using System;
+using System.Threading;
+using System.Threading.Tasks;
+using HordeCommon;
+using Microsoft.Extensions.Hosting;
+using Microsoft.Extensions.Logging;
+using StatsdClient;
+
+namespace Horde.Server.Server
+{
+	/// <summary>
+	/// Periodically send metrics for the CLR and other services that cannot be collected on a per-request basis
+	/// </summary>
+	public sealed class MetricService : IHostedService, IDisposable
+	{
+		readonly IDogStatsd _dogStatsd;
+		readonly ITicker _ticker;
+
+		/// <summary>
+		/// Constructor
+		/// </summary>
+		public MetricService(IDogStatsd dogStatsd, IClock clock, ILogger<MetricService> logger)
+		{
+			_dogStatsd = dogStatsd;
+			_ticker = clock.AddTicker<MetricService>(TimeSpan.FromSeconds(20.0), TickAsync, logger);
+		}
+
+		/// <inheritdoc/>
+		public Task StartAsync(CancellationToken cancellationToken) => _ticker.StartAsync();
+
+		/// <inheritdoc/>
+		public Task StopAsync(CancellationToken cancellationToken) => _ticker.StopAsync();
+
+		/// <inheritdoc/>
+		public void Dispose() => _ticker.Dispose();
+
+		/// <summary>
+		/// Execute the background task
+		/// </summary>
+		/// <param name="stoppingToken">Cancellation token for the async task</param>
+		/// <returns>Async task</returns>
+		ValueTask TickAsync(CancellationToken stoppingToken)
+		{
+			ReportGcMetrics();
+			ReportThreadMetrics();
+			return new ValueTask();
+		}
+
+		private void ReportThreadMetrics()
+		{
+			ThreadPool.GetMaxThreads(out int maxWorkerThreads, out int maxIoThreads);
+			ThreadPool.GetAvailableThreads(out int freeWorkerThreads, out int freeIoThreads);
+			ThreadPool.GetMinThreads(out int minWorkerThreads, out int minIoThreads);
+
+			int busyIoThreads = maxIoThreads - freeIoThreads;
+			int busyWorkerThreads = maxWorkerThreads - freeWorkerThreads;
+
+			_dogStatsd.Gauge("horde.clr.threadpool.io.max", maxIoThreads);
+			_dogStatsd.Gauge("horde.clr.threadpool.io.min", minIoThreads);
+			_dogStatsd.Gauge("horde.clr.threadpool.io.free", freeIoThreads);
+			_dogStatsd.Gauge("horde.clr.threadpool.io.busy", busyIoThreads);
+
+			_dogStatsd.Gauge("horde.clr.threadpool.worker.max", maxWorkerThreads);
+			_dogStatsd.Gauge("horde.clr.threadpool.worker.min", minWorkerThreads);
+			_dogStatsd.Gauge("horde.clr.threadpool.worker.free", freeWorkerThreads);
+			_dogStatsd.Gauge("horde.clr.threadpool.worker.busy", busyWorkerThreads);
+		}
+
+		private void ReportGcMetrics()
+		{
+			GCMemoryInfo gcMemoryInfo = GC.GetGCMemoryInfo();
+
+			_dogStatsd.Gauge("horde.clr.gc.totalMemory", GC.GetTotalMemory(false));
+			_dogStatsd.Gauge("horde.clr.gc.totalAllocated", GC.GetTotalAllocatedBytes());
+			_dogStatsd.Gauge("horde.clr.gc.heapSize", gcMemoryInfo.HeapSizeBytes);
+			_dogStatsd.Gauge("horde.clr.gc.fragmented", gcMemoryInfo.FragmentedBytes);
+			_dogStatsd.Gauge("horde.clr.gc.memoryLoad", gcMemoryInfo.MemoryLoadBytes);
+			_dogStatsd.Gauge("horde.clr.gc.totalAvailableMemory", gcMemoryInfo.TotalAvailableMemoryBytes);
+			_dogStatsd.Gauge("horde.clr.gc.highMemoryLoadThreshold", gcMemoryInfo.HighMemoryLoadThresholdBytes);
+		}
+	}
+}
