Horde: Rename the Horde.Build project to Horde.Server.

diff --git a/Engine/Source/Programs/Horde/Horde.Server/Server/Notices/INoticeCollection.cs b/Engine/Source/Programs/Horde/Horde.Server/Server/Notices/INoticeCollection.cs
new file mode 100644
--- /dev/null
+++ b/Engine/Source/Programs/Horde/Horde.Server/Server/Notices/INoticeCollection.cs
@@ -0,0 +1,57 @@
+// Copyright Epic Games, Inc. All Rights Reserved.
+
+using System;
+using System.Collections.Generic;
+using System.Threading.Tasks;
+using Horde.Server.Users;
+using Horde.Server.Utilities;
+using MongoDB.Bson;
+
+namespace Horde.Server.Server.Notices
+{
+	/// <summary>
+	/// A collection of user specified notices
+	/// </summary>
+	public interface INoticeCollection
+	{
+		/// <summary>
+		/// Add a notice to the collection
+		/// </summary>
+		/// <param name="message"></param>
+		/// <param name="userId"></param>
+		/// <param name="startTime"></param>
+		/// <param name="finishTime"></param>
+		/// <returns></returns>
+		Task<INotice?> AddNoticeAsync(string message, UserId? userId, DateTime? startTime, DateTime? finishTime);
+
+		/// <summary>
+		/// Update an existing notice
+		/// </summary>
+		/// <param name="id"></param>
+		/// <param name="message"></param>
+		/// <param name="startTime"></param>
+		/// <param name="finishTime"></param>
+		/// <returns></returns>
+		Task<bool> UpdateNoticeAsync(ObjectId id, string? message, DateTime? startTime, DateTime? finishTime);
+
+		/// <summary>
+		/// Get a notice by id
+		/// </summary>
+		/// <param name="noticeId"></param>
+		/// <returns></returns>
+		Task<INotice?> GetNoticeAsync(ObjectId noticeId);
+
+		/// <summary>
+		/// Get all notices
+		/// </summary>
+		/// <returns></returns>
+		Task<List<INotice>> GetNoticesAsync();
+
+		/// <summary>
+		/// Remove a notice
+		/// </summary>
+		/// <param name="id"></param>
+		/// <returns></returns>
+		Task<bool> RemoveNoticeAsync(ObjectId id);
+	}
+}
\ No newline at end of file
diff --git a/Engine/Source/Programs/Horde/Horde.Server/Server/Notices/NoticeCollection.cs b/Engine/Source/Programs/Horde/Horde.Server/Server/Notices/NoticeCollection.cs
new file mode 100644
--- /dev/null
+++ b/Engine/Source/Programs/Horde/Horde.Server/Server/Notices/NoticeCollection.cs
@@ -0,0 +1,115 @@
+// Copyright Epic Games, Inc. All Rights Reserved.
+
+using MongoDB.Bson.Serialization.Attributes;
+using System.Threading.Tasks;
+using MongoDB.Driver;
+using System;
+using System.Collections.Generic;
+using System.Linq;
+using Horde.Server.Utilities;
+using MongoDB.Bson;
+using Horde.Server.Users;
+
+namespace Horde.Server.Server.Notices
+{
+	/// <summary>
+	/// Collection of notice documents
+	/// </summary>
+	public class NoticeCollection : INoticeCollection
+	{
+
+		/// <summary>
+		/// Document representing a notice 
+		/// </summary>
+		class NoticeDocument : INotice
+		{
+			[BsonRequired, BsonId]
+			public ObjectId Id { get; set; }
+
+			[BsonIgnoreIfNull]
+			public UserId? UserId { get; set; }
+
+			[BsonIgnoreIfNull]
+			public DateTime? StartTime { get; set; }
+
+			[BsonIgnoreIfNull]
+			public DateTime? FinishTime { get; set; }
+
+			public string Message { get; set; } = String.Empty;
+
+			[BsonConstructor]
+			private NoticeDocument()
+			{
+			}
+
+			public NoticeDocument(ObjectId id, string message, UserId? userId, DateTime? startTime, DateTime? finishTime)
+			{
+				Id = id;
+				Message = message;
+				UserId = userId;
+				StartTime = startTime;
+				FinishTime = finishTime;
+			}
+		}
+
+		readonly IMongoCollection<NoticeDocument> _notices;
+
+		/// <summary>
+		/// Constructor
+		/// </summary>
+		public NoticeCollection(MongoService mongoService)
+		{
+			_notices = mongoService.GetCollection<NoticeDocument>("Notices");
+		}
+
+		/// <inheritdoc/>
+		public async Task<INotice?> AddNoticeAsync(string message, UserId? userId, DateTime? startTime, DateTime? finishTime)
+		{
+			NoticeDocument newNotice = new NoticeDocument(ObjectId.GenerateNewId(), message, userId, startTime, finishTime);
+			await _notices.InsertOneAsync(newNotice);
+			return newNotice;
+		}
+
+		/// <inheritdoc/>
+		public async Task<bool> UpdateNoticeAsync(ObjectId id, string? message, DateTime? startTime, DateTime? finishTime)
+		{
+
+			UpdateDefinitionBuilder<NoticeDocument> updateBuilder = Builders<NoticeDocument>.Update;
+			List<UpdateDefinition<NoticeDocument>> updates = new List<UpdateDefinition<NoticeDocument>>();
+
+			if (message != null)
+			{
+				updates.Add(updateBuilder.Set(x => x.Message, message));
+			}
+
+			updates.Add(updateBuilder.Set(x => x.StartTime, startTime));
+
+			updates.Add(updateBuilder.Set(x => x.FinishTime, finishTime));
+
+			NoticeDocument? document = await _notices.FindOneAndUpdateAsync(x => x.Id == id, updateBuilder.Combine(updates));
+
+			return document != null;
+		}
+
+		/// <inheritdoc/>
+		public async Task<INotice?> GetNoticeAsync(ObjectId noticeId)
+		{
+			return await _notices.Find(x => x.Id == noticeId).FirstOrDefaultAsync();
+		}
+
+		/// <inheritdoc/>
+		public async Task<List<INotice>> GetNoticesAsync()
+		{
+			List<NoticeDocument> results = await _notices.Find(x => true).ToListAsync();
+			return results.Select<NoticeDocument, INotice>(x => x).ToList();
+		}
+
+		/// <inheritdoc/>
+		public async Task<bool> RemoveNoticeAsync(ObjectId id)
+		{
+			DeleteResult result = await _notices.DeleteOneAsync(x => x.Id == id);
+			return result.DeletedCount > 0;
+		}
+	}
+}
+
diff --git a/Engine/Source/Programs/Horde/Horde.Server/Server/Notices/NoticeService.cs b/Engine/Source/Programs/Horde/Horde.Server/Server/Notices/NoticeService.cs
new file mode 100644
--- /dev/null
+++ b/Engine/Source/Programs/Horde/Horde.Server/Server/Notices/NoticeService.cs
@@ -0,0 +1,85 @@
+// Copyright Epic Games, Inc. All Rights Reserved.
+
+using System;
+using System.Collections.Generic;
+using System.Threading.Tasks;
+using Horde.Server.Users;
+using Horde.Server.Utilities;
+using MongoDB.Bson;
+
+namespace Horde.Server.Server.Notices
+{
+	/// <summary>
+	/// Notice service
+	/// </summary>
+	public sealed class NoticeService 
+	{
+		/// <summary>
+		/// Notice collection
+		/// </summary>
+		readonly INoticeCollection _notices;
+
+		/// <summary>
+		/// Device service constructor
+		/// </summary>
+		public NoticeService(INoticeCollection notices)
+		{
+			_notices = notices;
+		}
+
+		/// <summary>
+		/// Add a notice
+		/// </summary>
+		/// <param name="message"></param>
+		/// <param name="userId"></param>
+		/// <param name="startTime"></param>
+		/// <param name="finishTime"></param>
+		/// <returns></returns>
+		public Task<INotice?> AddNoticeAsync(string message, UserId? userId, DateTime? startTime, DateTime? finishTime)
+		{
+			return _notices.AddNoticeAsync(message, userId, startTime, finishTime);
+		}
+
+		/// <summary>
+		/// Update an existing notice
+		/// </summary>
+		/// <param name="id"></param>
+		/// <param name="message"></param>
+		/// <param name="startTime"></param>
+		/// <param name="finishTime"></param>
+		/// <returns></returns>
+		public Task<bool> UpdateNoticeAsync(ObjectId id, string? message, DateTime? startTime, DateTime? finishTime)
+		{
+			return _notices.UpdateNoticeAsync(id, message, startTime, finishTime);
+		}
+
+		/// <summary>
+		/// Get a notice by id
+		/// </summary>
+		/// <param name="noticeId"></param>
+		/// <returns></returns>
+		public Task<INotice?> GetNoticeAsync(ObjectId noticeId)
+		{
+			return _notices.GetNoticeAsync(noticeId);
+		}
+
+		/// <summary>
+		/// Get all notices
+		/// </summary>
+		/// <returns></returns>
+		public Task<List<INotice>> GetNoticesAsync()
+		{
+			return _notices.GetNoticesAsync();
+		}
+
+		/// <summary>
+		/// Remove an existing notice
+		/// </summary>
+		/// <param name="id"></param>
+		/// <returns></returns>
+		public Task<bool> RemoveNoticeAsync(ObjectId id)
+		{
+			return _notices.RemoveNoticeAsync(id);
+		}
+	}
+}
\ No newline at end of file
diff --git a/Engine/Source/Programs/Horde/Horde.Server/Server/Notices/NoticesController.cs b/Engine/Source/Programs/Horde/Horde.Server/Server/Notices/NoticesController.cs
new file mode 100644
--- /dev/null
+++ b/Engine/Source/Programs/Horde/Horde.Server/Server/Notices/NoticesController.cs
@@ -0,0 +1,137 @@
+// Copyright Epic Games, Inc. All Rights Reserved.
+
+using System;
+using System.Collections.Generic;
+using System.Threading.Tasks;
+using Horde.Server.Acls;
+using Horde.Server.Users;
+using Horde.Server.Utilities;
+using HordeCommon;
+using Microsoft.AspNetCore.Authorization;
+using Microsoft.AspNetCore.Mvc;
+using Microsoft.Extensions.Options;
+using MongoDB.Bson;
+
+namespace Horde.Server.Server.Notices
+{
+	/// <summary>
+	/// Controller for the /api/v1/notices endpoint
+	/// </summary>
+	[ApiController]
+	[Authorize]
+	[Route("[controller]")]
+	public class NoticesController : ControllerBase
+	{
+		private readonly AclService _aclService;
+		private readonly IUserCollection _userCollection;
+		private readonly NoticeService _noticeService;
+		private readonly LazyCachedValue<Task<IGlobals>> _cachedGlobals;
+		private readonly LazyCachedValue<Task<List<INotice>>> _cachedNotices;
+		private readonly IOptionsSnapshot<GlobalConfig> _globalConfig;
+		private readonly IClock _clock;
+
+		/// <summary>
+		/// Constructor
+		/// </summary>
+		public NoticesController(GlobalsService globalsService, NoticeService noticeService, AclService aclService, IUserCollection userCollection, IClock clock, IOptionsSnapshot<GlobalConfig> globalConfig)
+		{			
+			_aclService = aclService;
+			_userCollection = userCollection;
+			_noticeService = noticeService;
+			_cachedGlobals = new LazyCachedValue<Task<IGlobals>>(async () => await globalsService.GetAsync(), TimeSpan.FromSeconds(30.0));
+			_cachedNotices = new LazyCachedValue<Task<List<INotice>>>(() => noticeService.GetNoticesAsync(), TimeSpan.FromMinutes(1));
+			_clock = clock;
+			_globalConfig = globalConfig;
+		}
+
+		/// <summary>
+		/// Add a status message
+		/// </summary>
+		/// <param name="request"></param>
+		/// <returns></returns>
+		[HttpPost("/api/v1/notices")]
+		public async Task<ActionResult> AddNoticeAsync(CreateNoticeRequest request)
+		{
+			if (!_globalConfig.Value.Authorize(AclAction.AdminWrite, User))
+			{
+				return Forbid();
+			}
+
+			INotice? notice = await _noticeService.AddNoticeAsync(request.Message, User.GetUserId(), request.StartTime, request.FinishTime);
+
+			return notice == null ? NotFound() : Ok();
+
+		}
+
+		/// <summary>
+		/// Update a status message
+		/// </summary>
+		/// <param name="request"></param>
+		/// <returns></returns>
+		[HttpPut("/api/v1/notices")]
+		public async Task<ActionResult> UpdateNoticeAsync(UpdateNoticeRequest request)
+		{
+			if (!_globalConfig.Value.Authorize(AclAction.AdminWrite, User))
+			{
+				return Forbid();
+			}
+
+			return await _noticeService.UpdateNoticeAsync(new ObjectId(request.Id), request.Message, request.StartTime, request.FinishTime) ? Ok() : NotFound();
+		}
+
+		/// <summary>
+		/// Remove a manually added status message
+		/// </summary>
+		/// <returns></returns>
+		[HttpDelete("/api/v1/notices/{id}")]
+		public async Task<ActionResult> DeleteNoticeAsync(string id)
+		{
+			if (!_globalConfig.Value.Authorize(AclAction.AdminWrite, User))
+			{
+				return Forbid();
+			}
+
+			await _noticeService.RemoveNoticeAsync(new ObjectId(id));
+			
+			return Ok();
+		}
+
+		/// <summary>
+		/// Gets the current status messages
+		/// </summary>
+		/// <returns>The status messages</returns>
+		[HttpGet("/api/v1/notices")]
+		public async Task<List<GetNoticeResponse>> GetNoticesAsync()
+		{
+			List<GetNoticeResponse> messages = new List<GetNoticeResponse>();
+
+			DateTimeOffset now = TimeZoneInfo.ConvertTime(new DateTimeOffset(_clock.UtcNow), _clock.TimeZone);
+						
+			foreach (ScheduledDowntime schedule in _globalConfig.Value.Downtime)
+			{
+				DateTimeOffset start = TimeZoneInfo.ConvertTimeBySystemTimeZoneId(schedule.GetNext(now).StartTime, "UTC");
+				DateTimeOffset finish = TimeZoneInfo.ConvertTimeBySystemTimeZoneId(schedule.GetNext(now).FinishTime, "UTC");
+				messages.Add(new GetNoticeResponse() { ScheduledDowntime = true, StartTime = new DateTime(start.Ticks, DateTimeKind.Utc), FinishTime = new DateTime(finish.Ticks, DateTimeKind.Utc), Active = schedule.IsActive(now) });
+			}
+
+			// Add user notices
+			List<INotice> notices = await _cachedNotices.GetCached();
+
+			for (int i = 0; i < notices.Count; i++)
+			{
+				INotice notice = notices[i];
+				GetThinUserInfoResponse? userInfo = null;
+
+				if (notice.UserId != null)
+				{
+					userInfo = new GetThinUserInfoResponse(await _userCollection.GetCachedUserAsync(notice.UserId));
+				}
+
+				messages.Add(new GetNoticeResponse() { Id = notice.Id.ToString(), Active = true, Message = notice.Message, CreatedByUser = userInfo});
+			}
+
+			return messages;
+			
+		}
+	}
+}
diff --git a/Engine/Source/Programs/Horde/Horde.Server/Server/RedisService.cs b/Engine/Source/Programs/Horde/Horde.Server/Server/RedisService.cs
new file mode 100644
--- /dev/null
+++ b/Engine/Source/Programs/Horde/Horde.Server/Server/RedisService.cs
@@ -0,0 +1,296 @@
+// Copyright Epic Games, Inc. All Rights Reserved.
+
+using System;
+using System.Collections.Generic;
+using System.Diagnostics;
+using System.IO;
+using System.Linq;
+using System.Net;
+using System.Net.NetworkInformation;
+using System.Runtime.InteropServices;
+using System.Threading;
+using System.Threading.Tasks;
+using EpicGames.Core;
+using EpicGames.Redis;
+using Microsoft.Extensions.Diagnostics.HealthChecks;
+using Microsoft.Extensions.Logging;
+using Microsoft.Extensions.Options;
+using StackExchange.Redis;
+
+namespace Horde.Server.Server
+{
+	/// <summary>
+	/// Manages the lifetime of a bundled Redis instance
+	/// </summary>
+	public sealed class RedisService : IHealthCheck, IDisposable
+	{
+		/// <summary>
+		/// Default Redis port
+		/// </summary>
+		const int RedisPort = 6379;
+
+		/// <summary>
+		/// The managed process group containing the Redis server
+		/// </summary>
+		ManagedProcessGroup? _redisProcessGroup;
+
+		/// <summary>
+		/// The server process
+		/// </summary>
+		ManagedProcess? _redisProcess;
+
+		/// <summary>
+		/// Connection multiplexer
+		/// </summary>
+		private readonly ConnectionMultiplexer _multiplexer;
+
+		/// <summary>
+		/// The database interface.
+		/// If possible, use ConnectionPool instead. 
+		/// </summary>
+		public IDatabase DatabaseSingleton { get; }
+
+		/// <summary>
+		/// Connection pool
+		/// </summary>
+		public RedisConnectionPool ConnectionPool { get; }
+
+		/// <summary>
+		/// Logger factory
+		/// </summary>
+		readonly ILoggerFactory _loggerFactory;
+
+		/// <summary>
+		/// Logging instance
+		/// </summary>
+		readonly ILogger<RedisService> _logger;
+
+		/// <summary>
+		/// Hack to initialize RedisService early enough to use data protection
+		/// </summary>
+		/// <param name="settings"></param>
+		public RedisService(ServerSettings settings)
+			: this(Options.Create(settings), new Serilog.Extensions.Logging.SerilogLoggerFactory(Serilog.Log.Logger))
+		{
+		}
+
+		/// <summary>
+		/// Constructor
+		/// </summary>
+		public RedisService(IOptions<ServerSettings> options, ILoggerFactory loggerFactory)
+			: this(options.Value.RedisConnectionConfig, -1, loggerFactory)
+		{
+		}
+
+		/// <summary>
+		/// Constructor
+		/// </summary>
+		/// <param name="connectionString">Redis connection string. If null, we will start a temporary redis instance on the local machine.</param>
+		/// <param name="dbNum">Override for the database to use. Set to -1 to use the default from the connection string.</param>
+		/// <param name="loggerFactory"></param>
+		public RedisService(string? connectionString, int dbNum, ILoggerFactory loggerFactory)
+		{
+			_loggerFactory = loggerFactory;
+			_logger = loggerFactory.CreateLogger<RedisService>();
+
+			if (connectionString == null)
+			{
+				if (IsRunningOnDefaultPort())
+				{
+					connectionString = $"localhost:{RedisPort}";
+				}
+				else if (TryStartRedisServer())
+				{
+					connectionString = $"localhost:{RedisPort}";
+				}
+				else
+				{
+					throw new Exception($"Unable to connect to Redis. Please set {nameof(ServerSettings.RedisConnectionConfig)} in {Program.UserConfigFile}");
+				}
+			}
+
+			_multiplexer = ConnectionMultiplexer.Connect(connectionString);
+			DatabaseSingleton = _multiplexer.GetDatabase(dbNum);
+			ConnectionPool = new RedisConnectionPool(20, connectionString, dbNum);
+		}
+
+		/// <inheritdoc/>
+		public async Task<HealthCheckResult> CheckHealthAsync(HealthCheckContext context, CancellationToken cancellationToken = default)
+		{
+			try
+			{
+				TimeSpan span = await GetDatabase().PingAsync();
+				return HealthCheckResult.Healthy(data: new Dictionary<string, object> { ["Latency"] = span.ToString() });
+			}
+			catch (Exception ex)
+			{
+				return HealthCheckResult.Unhealthy("Unable to ping Redis", ex);
+			}
+		}
+
+		/// <summary>
+		/// Get the least-loaded Redis database from the connection pool
+		/// Don't store the returned object and try to resolve this as late as possible to ensure load is balanced.
+		/// </summary>
+		/// <returns>A Redis database</returns>
+		public IDatabase GetDatabase()
+		{
+			return ConnectionPool.GetDatabase();
+		}
+
+		/// <inheritdoc/>
+		public void Dispose()
+		{
+			_multiplexer.Dispose();
+
+			if (_redisProcess != null)
+			{
+				_redisProcess.Dispose();
+				_redisProcess = null;
+			}
+			if (_redisProcessGroup != null)
+			{
+				_redisProcessGroup.Dispose();
+				_redisProcessGroup = null;
+			}
+		}
+
+		/// <summary>
+		/// Checks if Redis is already running on the default port
+		/// </summary>
+		/// <returns></returns>
+		static bool IsRunningOnDefaultPort()
+		{
+			IPGlobalProperties ipGlobalProperties = IPGlobalProperties.GetIPGlobalProperties();
+
+			IPEndPoint[] listeners = ipGlobalProperties.GetActiveTcpListeners();
+			if (listeners.Any(x => x.Port == RedisPort))
+			{
+				return true;
+			}
+
+			return false;
+		}
+
+		/// <summary>
+		/// Attempts to start a local instance of Redis
+		/// </summary>
+		/// <returns></returns>
+		bool TryStartRedisServer()
+		{
+			if (!RuntimeInformation.IsOSPlatform(OSPlatform.Windows))
+			{
+				return false;
+			}
+
+			FileReference redisExe = FileReference.Combine(Program.AppDir, "ThirdParty", "Redis", "redis-server.exe");
+			if (!FileReference.Exists(redisExe))
+			{
+				_logger.LogDebug("Redis executable does not exist at {ExePath}", redisExe);
+				return false;
+			}
+
+			DirectoryReference redisDir = DirectoryReference.Combine(Program.DataDir, "Redis");
+			DirectoryReference.CreateDirectory(redisDir);
+
+			FileReference redisConfigFile = FileReference.Combine(redisDir, "redis.conf");
+			if (!FileReference.Exists(redisConfigFile))
+			{
+				using (StreamWriter writer = new StreamWriter(redisConfigFile.FullName))
+				{
+					writer.WriteLine("# redis.conf");
+				}
+			}
+
+			_redisProcessGroup = new ManagedProcessGroup();
+			try
+			{
+				_redisProcess = new ManagedProcess(_redisProcessGroup, redisExe.FullName, "", null, null, ProcessPriorityClass.Normal);
+				_redisProcess.StdIn.Close();
+				Task.Run(() => RelayRedisOutput());
+				return true;
+			}
+			catch (Exception ex)
+			{
+				_logger.LogWarning(ex, "Unable to start Redis server process");
+				return false;
+			}
+		}
+
+		/// <summary>
+		/// Copies output from the redis process to the logger
+		/// </summary>
+		/// <returns></returns>
+		async Task RelayRedisOutput()
+		{
+			ILogger redisLogger = _loggerFactory.CreateLogger("Redis");
+			for (; ; )
+			{
+				string? line = await _redisProcess!.ReadLineAsync();
+				if (line == null)
+				{
+					break;
+				}
+				if (line.Length > 0)
+				{
+					redisLogger.Log(LogLevel.Information, "{Output}", line);
+				}
+			}
+			redisLogger.LogInformation("Exit code {ExitCode}", _redisProcess.ExitCode);
+		}
+
+		/// <summary>
+		/// Publish a message to a channel
+		/// </summary>
+		/// <param name="channel">Channel to post to</param>
+		/// <param name="message">Message to post to the channel</param>
+		/// <param name="flags">Flags for the request</param>
+		public Task PublishAsync(RedisChannel channel, RedisValue message, CommandFlags flags = CommandFlags.None)
+		{
+			return ConnectionPool.GetDatabase().PublishAsync(channel, message, flags);
+		}
+
+		/// <summary>
+		/// Publish a message to a channel
+		/// </summary>
+		/// <typeparam name="T">Type of elements sent over the channel</typeparam>
+		/// <param name="channel">Channel to post to</param>
+		/// <param name="message">Message to post to the channel</param>
+		/// <param name="flags">Flags for the request</param>
+		public Task PublishAsync<T>(RedisChannel<T> channel, T message, CommandFlags flags = CommandFlags.None)
+		{
+			return ConnectionPool.GetDatabase().PublishAsync(channel, message, flags);
+		}
+
+		/// <inheritdoc cref="SubscribeAsync{T}(RedisChannel{T}, Action{RedisChannel{T}, T})"/>
+		public Task<RedisSubscription> SubscribeAsync(RedisChannel channel, Action<RedisValue> callback) => SubscribeAsync(channel, (ch, x) => callback(x));
+
+		/// <inheritdoc cref="SubscribeAsync{T}(RedisChannel{T}, Action{RedisChannel{T}, T})"/>
+		public Task<RedisSubscription> SubscribeAsync<T>(RedisChannel<T> channel, Action<T> callback) => SubscribeAsync(channel, (ch, x) => callback(x));
+
+		/// <summary>
+		/// Subscribe to notifications on a channel
+		/// </summary>
+		/// <param name="channel">Channel to monitor</param>
+		/// <param name="callback">Callback for new events</param>
+		/// <returns>Subscription object</returns>
+		public async Task<RedisSubscription> SubscribeAsync(RedisChannel channel, Action<RedisChannel, RedisValue> callback)
+		{
+			IConnectionMultiplexer connection = ConnectionPool.GetConnection();
+			return await connection.SubscribeAsync(channel, callback);
+		}
+
+		/// <summary>
+		/// Subscribe to notifications on a channel
+		/// </summary>
+		/// <typeparam name="T">Type of elements sent over the channel</typeparam>
+		/// <param name="channel">Channel to monitor</param>
+		/// <param name="callback">Callback for new events</param>
+		/// <returns>Subscription object</returns>
+		public async Task<RedisSubscription> SubscribeAsync<T>(RedisChannel<T> channel, Action<RedisChannel<T>, T> callback)
+		{
+			IConnectionMultiplexer connection = ConnectionPool.GetConnection();
+			return await connection.SubscribeAsync(channel, callback);
+		}
+	}
+}
diff --git a/Engine/Source/Programs/Horde/Horde.Server/Server/RpcService.cs b/Engine/Source/Programs/Horde/Horde.Server/Server/RpcService.cs
new file mode 100644
--- /dev/null
+++ b/Engine/Source/Programs/Horde/Horde.Server/Server/RpcService.cs
@@ -0,0 +1,476 @@
+// Copyright Epic Games, Inc. All Rights Reserved.
+
+using System;
+using System.Buffers;
+using System.Collections.Generic;
+using System.IO;
+using System.Linq;
+using System.Net;
+using System.Threading;
+using System.Threading.Tasks;
+using Amazon.EC2.Model;
+using Google.Protobuf;
+using Google.Protobuf.WellKnownTypes;
+using Grpc.Core;
+using Horde.Server.Acls;
+using Horde.Server.Agents;
+using Horde.Server.Agents.Pools;
+using Horde.Server.Agents.Sessions;
+using Horde.Server.Agents.Software;
+using Horde.Server.Jobs;
+using Horde.Server.Tasks;
+using Horde.Server.Tools;
+using Horde.Server.Utilities;
+using HordeCommon.Rpc;
+using Microsoft.AspNetCore.Authorization;
+using Microsoft.Extensions.Logging;
+using Microsoft.Extensions.Options;
+
+namespace Horde.Server.Server
+{
+	using RpcAgentCapabilities = HordeCommon.Rpc.Messages.AgentCapabilities;
+	using RpcDeviceCapabilities = HordeCommon.Rpc.Messages.DeviceCapabilities;
+	using RpcGetStreamResponse = HordeCommon.Rpc.GetStreamResponse;
+	using RpcGetJobResponse = HordeCommon.Rpc.GetJobResponse;
+	using RpcGetStepResponse = HordeCommon.Rpc.GetStepResponse;
+	using RpcUpdateJobRequest = HordeCommon.Rpc.UpdateJobRequest;
+	using RpcUpdateStepRequest = HordeCommon.Rpc.UpdateStepRequest;
+
+	/// <summary>
+	/// Implements the Horde gRPC service for bots updating their status and dequeing work
+	/// </summary>
+	[Authorize]
+	public class RpcService : HordeRpc.HordeRpcBase
+	{
+		/// <summary>
+		/// Timeout before closing a long-polling request (client will retry again) 
+		/// </summary>
+		internal TimeSpan _longPollTimeout = TimeSpan.FromMinutes(9);
+
+		readonly AgentService _agentService;
+		readonly PoolService _poolService;
+		readonly LifetimeService _lifetimeService;
+		readonly ConformTaskSource _conformTaskSource;
+		readonly JobRpcCommon _jobRpcCommon;
+		readonly IToolCollection _toolCollection;
+		readonly IOptionsSnapshot<GlobalConfig> _globalConfig;
+		readonly ILogger _logger;
+
+		/// <summary>
+		/// Constructor
+		/// </summary>
+		public RpcService(AgentService agentService, PoolService poolService, LifetimeService lifetimeService, ConformTaskSource conformTaskSource, JobRpcCommon jobRpcCommon, IToolCollection toolCollection, IOptionsSnapshot<GlobalConfig> globalConfig, ILogger<RpcService> logger)
+		{
+			_agentService = agentService;
+			_poolService = poolService;
+			_lifetimeService = lifetimeService;
+			_conformTaskSource = conformTaskSource;
+			_jobRpcCommon = jobRpcCommon;
+			_toolCollection = toolCollection;
+			_globalConfig = globalConfig;
+			_logger = logger;
+		}
+
+		/// <summary>
+		/// Waits until the server is terminating
+		/// </summary>
+		/// <param name="reader">Request reader</param>
+		/// <param name="writer">Response writer</param>
+		/// <param name="context">Context for the call</param>
+		/// <returns>Response object</returns>
+		public override async Task QueryServerState(IAsyncStreamReader<QueryServerStateRequest> reader, IServerStreamWriter<QueryServerStateResponse> writer, ServerCallContext context)
+		{
+			if (await reader.MoveNext())
+			{
+				QueryServerStateRequest request = reader.Current;
+				_logger.LogInformation("Start server query for client {Name}", request.Name);
+
+				// Return the current response
+				QueryServerStateResponse response = new QueryServerStateResponse();
+				response.Name = Dns.GetHostName();
+				await writer.WriteAsync(response);
+
+				// Move to the next request from the client. This should always be the end of the stream, but will not occur until the client stops requesting responses.
+				Task<bool> moveNextTask = reader.MoveNext();
+
+				// Wait for the client to close the stream or a shutdown to start
+				Task longPollDelay = Task.Delay(_longPollTimeout);
+				Task waitTask = await Task.WhenAny(moveNextTask, _lifetimeService.StoppingTask, longPollDelay);
+
+				if (waitTask == moveNextTask)
+				{
+					throw new Exception("Unexpected request to QueryServerState posted from client.");
+				}
+				else if (waitTask == _lifetimeService.StoppingTask)
+				{
+					_logger.LogInformation("Notifying client {Name} of server shutdown", request.Name);
+					await writer.WriteAsync(response);
+				}
+				else if (waitTask == longPollDelay)
+				{
+					// Send same response as server shutdown. In the agent perspective, they will be identical.
+					await writer.WriteAsync(response);
+				}
+			}
+		}
+
+		/// <summary>
+		/// Waits until the server is terminating
+		/// </summary>
+		/// <param name="reader">Request reader</param>
+		/// <param name="writer">Response writer</param>
+		/// <param name="context">Context for the call</param>
+		/// <returns>Response object</returns>
+		public override async Task QueryServerStateV2(IAsyncStreamReader<QueryServerStateRequest> reader, IServerStreamWriter<QueryServerStateResponse> writer, ServerCallContext context)
+		{
+			if (await reader.MoveNext())
+			{
+				QueryServerStateRequest request = reader.Current;
+				_logger.LogDebug("Start server query for client {Name}", request.Name);
+
+				try
+				{
+					// Return the current response
+					QueryServerStateResponse response = new QueryServerStateResponse();
+					response.Name = Dns.GetHostName();
+					response.Stopping = _lifetimeService.IsStopping;
+					await writer.WriteAsync(response);
+
+					// Move to the next request from the client. This should always be the end of the stream, but will not occur until the client stops requesting responses.
+					Task<bool> moveNextTask = reader.MoveNext();
+
+					// Wait for the client to close the stream or a shutdown to start
+					if (await Task.WhenAny(moveNextTask, _lifetimeService.StoppingTask) == _lifetimeService.StoppingTask)
+					{
+						response.Stopping = true;
+						await writer.WriteAsync(response);
+					}
+
+					// Wait until the client has finished sending
+					while (await moveNextTask)
+					{
+						moveNextTask = reader.MoveNext();
+					}
+				}
+				catch (Exception ex)
+				{
+					_logger.LogError(ex, "Exception in QueryServerState for {Name}", request.Name);
+					throw;
+				}
+			}
+		}
+
+		/// <summary>
+		/// Updates the workspaces synced for an agent
+		/// </summary>
+		/// <param name="request">The request parameters</param>
+		/// <param name="context">Context for the call</param>
+		/// <returns>Response object</returns>
+		public override async Task<UpdateAgentWorkspacesResponse> UpdateAgentWorkspaces(UpdateAgentWorkspacesRequest request, ServerCallContext context)
+		{
+			for (; ; )
+			{
+				// Get the current agent state
+				IAgent? agent = await _agentService.GetAgentAsync(new AgentId(request.AgentId));
+				if (agent == null)
+				{
+					throw new StructuredRpcException(StatusCode.OutOfRange, "Agent {AgentId} does not exist", request.AgentId);
+				}
+
+				// Get the new workspaces
+				List<AgentWorkspace> newWorkspaces = request.Workspaces.Select(x => new AgentWorkspace(x)).ToList();
+
+				// Get the set of workspaces that are currently required
+				HashSet<AgentWorkspace> conformWorkspaces = await _poolService.GetWorkspacesAsync(agent, DateTime.UtcNow, _globalConfig.Value);
+				bool pendingConform = !conformWorkspaces.SetEquals(newWorkspaces) || (agent.RequestFullConform && !request.RemoveUntrackedFiles);
+
+				// Update the workspaces
+				if (await _agentService.TryUpdateWorkspacesAsync(agent, newWorkspaces, pendingConform))
+				{
+					UpdateAgentWorkspacesResponse response = new UpdateAgentWorkspacesResponse();
+					if (pendingConform)
+					{
+						response.Retry = await _conformTaskSource.GetWorkspacesAsync(agent, response.PendingWorkspaces);
+						response.RemoveUntrackedFiles = request.RemoveUntrackedFiles || agent.RequestFullConform;
+					}
+					return response;
+				}
+			}
+		}
+
+		static void CopyPropertyToResource(string name, List<string> properties, Dictionary<string, int> resources)
+		{
+			foreach (string property in properties)
+			{
+				if (property.Length > name.Length && property.StartsWith(name, StringComparison.OrdinalIgnoreCase) && property[name.Length] == '=')
+				{
+					int value;
+					if (Int32.TryParse(property.AsSpan(name.Length + 1), out value))
+					{
+						resources[name] = value;
+					}
+				}
+			}
+		}
+
+		static void GetCapabilities(RpcAgentCapabilities? capabilities, out List<string> properties, out Dictionary<string, int> resources)
+		{
+			properties = new List<string>();
+			resources = new Dictionary<string, int>();
+
+			if (capabilities == null)
+			{
+				return;
+			}
+			properties.AddRange(capabilities.Properties);
+
+			if (capabilities.Devices.Count <= 0)
+			{
+				return;
+			}
+
+			RpcDeviceCapabilities device = capabilities.Devices[0];
+			if (device.Properties == null)
+			{
+				return;
+			}
+
+			properties.AddRange(device.Properties);
+			CopyPropertyToResource(KnownPropertyNames.LogicalCores, properties, resources);
+			CopyPropertyToResource(KnownPropertyNames.Ram, properties, resources);
+		}
+
+		/// <summary>
+		/// Creates a new session
+		/// </summary>
+		/// <param name="request">Request to create a new agent</param>
+		/// <param name="context">Context for the RPC call</param>
+		/// <returns>Information about the new agent</returns>
+		public override async Task<CreateSessionResponse> CreateSession(CreateSessionRequest request, ServerCallContext context)
+		{
+			if (request.Capabilities == null)
+			{
+				throw new StructuredRpcException(StatusCode.InvalidArgument, "Capabilities may not be null");
+			}
+
+			AgentId agentId = new AgentId(request.Name);
+			using IDisposable scope = _logger.BeginScope("CreateSession({AgentId})", agentId.ToString());
+
+			GlobalConfig globalConfig = _globalConfig.Value;
+
+			// Find the agent
+			IAgent? agent = await _agentService.GetAgentAsync(agentId);
+			if (agent == null)
+			{
+				if (!globalConfig.Authorize(AclAction.CreateAgent, context.GetHttpContext().User))
+				{
+					throw new StructuredRpcException(StatusCode.PermissionDenied, "User is not authenticated to create new agents");
+				}
+
+				agent = await _agentService.CreateAgentAsync(request.Name, true, null);
+			}
+
+			// Make sure we're allowed to create sessions on this agent
+			if (!globalConfig.Authorize(AclAction.CreateSession, context.GetHttpContext().User))
+			{
+				throw new StructuredRpcException(StatusCode.PermissionDenied, "User is not authenticated to create session for {AgentId}", request.Name);
+			}
+
+			// Get the known properties for this agent
+			GetCapabilities(request.Capabilities, out List<string> properties, out Dictionary<string, int> resources);
+
+			// Create a new session
+			agent = await _agentService.CreateSessionAsync(agent, request.Status, properties, resources, request.Version);
+			if (agent == null)
+			{
+				throw new StructuredRpcException(StatusCode.NotFound, "Agent {AgentId} not found", request.Name);
+			}
+
+			// Create the response
+			CreateSessionResponse response = new CreateSessionResponse();
+			response.AgentId = agent.Id.ToString();
+			response.SessionId = agent.SessionId.ToString();
+			response.ExpiryTime = Timestamp.FromDateTime(agent.SessionExpiresAt!.Value);
+			response.Token = await _agentService.IssueSessionTokenAsync(agent.Id, agent.SessionId!.Value);
+			return response;
+		}
+
+		/// <summary>
+		/// Updates an agent session
+		/// </summary>
+		/// <param name="reader">Request to create a new agent</param>
+		/// <param name="writer">Writer for response objects</param>
+		/// <param name="context">Context for the RPC call</param>
+		/// <returns>Information about the new agent</returns>
+		public override async Task UpdateSession(IAsyncStreamReader<UpdateSessionRequest> reader, IServerStreamWriter<UpdateSessionResponse> writer, ServerCallContext context)
+		{
+			// Read the request object
+			Task<bool> nextRequestTask = reader.MoveNext();
+			if (await nextRequestTask)
+			{
+				UpdateSessionRequest request = reader.Current;
+				using IDisposable scope = _logger.BeginScope("UpdateSession for agent {AgentId}, session {SessionId}", request.AgentId, request.SessionId);
+
+				_logger.LogDebug("Updating session for {AgentId}", request.AgentId);
+				foreach (HordeCommon.Rpc.Messages.Lease lease in request.Leases)
+				{
+					_logger.LogDebug("Session {SessionId}, Lease {LeaseId} - State: {LeaseState}, Outcome: {LeaseOutcome}", request.SessionId, lease.Id, lease.State, lease.Outcome);
+				}
+
+				// Get a task for moving to the next item. This will only complete once the call has closed.
+				using CancellationTokenSource cancellationSource = CancellationTokenSource.CreateLinkedTokenSource(context.CancellationToken);
+				nextRequestTask = reader.MoveNext();
+				nextRequestTask = nextRequestTask.ContinueWith(task =>
+				{
+					cancellationSource.Cancel();
+					return task.Result;
+				}, TaskScheduler.Current);
+
+				// Get the current agent state
+				IAgent? agent = await _agentService.GetAgentAsync(new AgentId(request.AgentId));
+				if (agent != null)
+				{
+					// Check we're authorized to update it
+					if (!_agentService.AuthorizeSession(agent, context.GetHttpContext().User))
+					{
+						throw new StructuredRpcException(StatusCode.PermissionDenied, "Not authenticated for {AgentId}", request.AgentId);
+					}
+
+					// Get the new capabilities of this agent
+					List<string>? properties = null;
+					Dictionary<string, int>? resources = null;
+					if (request.Capabilities != null)
+					{
+						GetCapabilities(request.Capabilities, out properties, out resources);
+					}
+
+					// Update the session
+					try
+					{
+						agent = await _agentService.UpdateSessionWithWaitAsync(agent, SessionId.Parse(request.SessionId), request.Status, properties, resources, request.Leases, cancellationSource.Token);
+					}
+					catch (Exception ex)
+					{
+						_logger.LogError(ex, "Swallowed exception while updating session for {AgentId}.", request.AgentId);
+					}
+				}
+
+				// Handle the invalid agent case
+				if (agent == null)
+				{
+					throw new StructuredRpcException(StatusCode.NotFound, "Invalid agent name '{AgentId}'", request.AgentId);
+				}
+
+				// Create the new session info
+				if (!context.CancellationToken.IsCancellationRequested)
+				{
+					UpdateSessionResponse response = new UpdateSessionResponse();
+					response.Leases.Add(agent.Leases.Select(x => x.ToRpcMessage()));
+					response.ExpiryTime = (agent.SessionExpiresAt == null) ? new Timestamp() : Timestamp.FromDateTime(agent.SessionExpiresAt.Value);
+					await writer.WriteAsync(response);
+				}
+
+				// Wait for the client to close the stream
+				while (await nextRequestTask)
+				{
+					nextRequestTask = reader.MoveNext();
+				}
+			}
+		}
+
+		/// <inheritdoc/>
+		public override Task<RpcGetStreamResponse> GetStream(GetStreamRequest request, ServerCallContext context) => _jobRpcCommon.GetStream(request, context);
+
+		/// <inheritdoc/>
+		public override Task<RpcGetJobResponse> GetJob(GetJobRequest request, ServerCallContext context) => _jobRpcCommon.GetJob(request, context);
+
+		/// <inheritdoc/>
+		public override Task<Empty> UpdateJob(RpcUpdateJobRequest request, ServerCallContext context) => _jobRpcCommon.UpdateJob(request, context);
+
+		/// <inheritdoc/>
+		public override Task<BeginBatchResponse> BeginBatch(BeginBatchRequest request, ServerCallContext context) => _jobRpcCommon.BeginBatch(request, context);
+
+		/// <inheritdoc/>
+		public override Task<Empty> FinishBatch(FinishBatchRequest request, ServerCallContext context) => _jobRpcCommon.FinishBatch(request, context);
+
+		/// <inheritdoc/>
+		public override Task<BeginStepResponse> BeginStep(BeginStepRequest request, ServerCallContext context) => _jobRpcCommon.BeginStep(request, context);
+
+		/// <inheritdoc/>
+		public override Task<Empty> UpdateStep(RpcUpdateStepRequest request, ServerCallContext context) => _jobRpcCommon.UpdateStep(request, context);
+
+		/// <inheritdoc/>
+		public override Task<RpcGetStepResponse> GetStep(GetStepRequest request, ServerCallContext context) => _jobRpcCommon.GetStep(request, context);
+
+		/// <inheritdoc/>
+		public override Task<UpdateGraphResponse> UpdateGraph(UpdateGraphRequest request, ServerCallContext context) => _jobRpcCommon.UpdateGraph(request, context);
+
+		/// <inheritdoc/>
+		public override Task<Empty> CreateEvents(CreateEventsRequest request, ServerCallContext context) => _jobRpcCommon.CreateEvents(request, context);
+
+		/// <summary>
+		/// Downloads a new agent archive
+		/// </summary>
+		/// <param name="request">Request arguments</param>
+		/// <param name="responseStream">Writer for the output data</param>
+		/// <param name="context">Context for the RPC call</param>
+		/// <returns>Information about the new agent</returns>
+		public override async Task DownloadSoftware(DownloadSoftwareRequest request, IServerStreamWriter<DownloadSoftwareResponse> responseStream, ServerCallContext context)
+		{
+			int colonIdx = request.Version.IndexOf(':', StringComparison.Ordinal);
+			ToolId toolId = new ToolId(request.Version.Substring(0, colonIdx));
+			string version = request.Version.Substring(colonIdx + 1);
+
+			ToolConfig? toolConfig;
+			if (!_globalConfig.Value.TryGetTool(toolId, out toolConfig))
+			{
+				throw new StructuredRpcException(StatusCode.NotFound, $"Missing tool {toolId}");
+			}
+			if (!toolConfig.Authorize(AclAction.DownloadTool, context.GetHttpContext().User))
+			{
+				throw new StructuredRpcException(StatusCode.NotFound, "Access to software is forbidden");
+			}
+
+			ITool? tool = await _toolCollection.GetAsync(toolId, _globalConfig.Value);
+			if (tool == null)
+			{
+				throw new StructuredRpcException(StatusCode.NotFound, $"Missing tool {toolId}");
+			}
+
+			IToolDeployment? deployment = tool.Deployments.LastOrDefault(x => x.Version.Equals(version, StringComparison.Ordinal));
+			if (deployment == null)
+			{
+				throw new StructuredRpcException(StatusCode.NotFound, $"Missing tool version {version}");
+			}
+
+			using Stream stream = await _toolCollection.GetDeploymentZipAsync(tool, deployment, context.CancellationToken);
+			using (IMemoryOwner<byte> buffer = MemoryPool<byte>.Shared.Rent(128 * 1024))
+			{
+				for(; ;)
+				{
+					int read = await stream.ReadAsync(buffer.Memory, context.CancellationToken);
+					if (read == 0)
+					{
+						break;
+					}
+
+					DownloadSoftwareResponse response = new DownloadSoftwareResponse();
+					response.Data = UnsafeByteOperations.UnsafeWrap(buffer.Memory.Slice(0, read));
+					await responseStream.WriteAsync(response);
+				}
+			}
+		}
+
+		/// <inheritdoc/>
+		public override Task<Empty> WriteOutput(WriteOutputRequest request, ServerCallContext context) => _jobRpcCommon.WriteOutput(request, context);
+
+		/// <inheritdoc/>
+		public override Task<UploadArtifactResponse> UploadArtifact(IAsyncStreamReader<UploadArtifactRequest> reader, ServerCallContext context) => _jobRpcCommon.UploadArtifact(reader, context);
+
+		/// <inheritdoc/>
+		public override Task<UploadTestDataResponse> UploadTestData(IAsyncStreamReader<UploadTestDataRequest> reader, ServerCallContext context) => _jobRpcCommon.UploadTestData(reader, context);
+
+		/// <inheritdoc/>
+		public override Task<CreateReportResponse> CreateReport(CreateReportRequest request, ServerCallContext context) => _jobRpcCommon.CreateReport(request, context);
+	}
+}
diff --git a/Engine/Source/Programs/Horde/Horde.Server/Server/ServerController.cs b/Engine/Source/Programs/Horde/Horde.Server/Server/ServerController.cs
new file mode 100644
--- /dev/null
+++ b/Engine/Source/Programs/Horde/Horde.Server/Server/ServerController.cs
@@ -0,0 +1,55 @@
+// Copyright Epic Games, Inc. All Rights Reserved.
+
+using System.Diagnostics;
+using System.Reflection;
+using Microsoft.AspNetCore.Authorization;
+using Microsoft.AspNetCore.Mvc;
+using Microsoft.Extensions.Options;
+
+namespace Horde.Server.Server
+{
+	/// <summary>
+	/// Controller managing account status
+	/// </summary>
+	[ApiController]
+	[Authorize]
+	[Route("[controller]")]
+	public class ServerController : ControllerBase
+	{
+
+		/// <summary>
+		/// Settings for the server
+		/// </summary>
+		readonly IOptionsMonitor<ServerSettings> _settings;
+
+		/// <summary>
+		/// Constructor
+		/// </summary>
+		public ServerController(IOptionsMonitor<ServerSettings> settings)
+		{
+			_settings = settings;
+		}
+		
+		/// <summary>
+		/// Get server version
+		/// </summary>
+		[HttpGet]
+		[Route("/api/v1/server/version")]
+		public ActionResult GetVersionAsync()
+		{
+			FileVersionInfo fileVersionInfo = FileVersionInfo.GetVersionInfo(Assembly.GetExecutingAssembly().Location);
+			return Ok(fileVersionInfo.ProductVersion);
+		}		
+
+		/// <summary>
+		/// Get server information
+		/// </summary>
+		[HttpGet]
+		[Route("/api/v1/server/info")]
+		[ProducesResponseType(typeof(GetServerInfoResponse), 200)]
+		public ActionResult<GetServerInfoResponse> GetServerInfo()
+		{
+			return new GetServerInfoResponse();
+		}
+	}
+}
diff --git a/Engine/Source/Programs/Horde/Horde.Server/ServerSettings.cs b/Engine/Source/Programs/Horde/Horde.Server/ServerSettings.cs
new file mode 100644
--- /dev/null
+++ b/Engine/Source/Programs/Horde/Horde.Server/ServerSettings.cs
@@ -0,0 +1,867 @@
+// Copyright Epic Games, Inc. All Rights Reserved.
+
+using System;
+using System.Collections.Generic;
+using System.Diagnostics.CodeAnalysis;
+using System.Linq;
+using System.Security.Claims;
+using System.Text.Json.Serialization;
+using EpicGames.Core;
+using EpicGames.Horde;
+using EpicGames.Horde.Storage;
+using EpicGames.Horde.Storage.Nodes;
+using EpicGames.Perforce;
+using Horde.Server.Acls;
+using Horde.Server.Agents.Fleet;
+using Horde.Server.Server;
+using Horde.Server.Storage.Backends;
+using Horde.Server.Telemetry;
+using Horde.Server.Tools;
+using Horde.Server.Utilities;
+using Serilog.Events;
+
+namespace Horde.Server
+{
+	/// <summary>
+	/// Types of storage backend to use
+	/// </summary>
+	public enum StorageBackendType
+	{
+		/// <summary>
+		/// Local filesystem
+		/// </summary>
+		FileSystem,
+
+		/// <summary>
+		/// AWS S3
+		/// </summary>
+		Aws,
+
+		/// <summary>
+		/// In-memory only (for testing)
+		/// </summary>
+		Memory,
+	};
+
+	/// <summary>
+	/// Common settings for different storage backends
+	/// </summary>
+	public interface IStorageBackendOptions : IFileSystemStorageOptions, IAwsStorageOptions
+	{
+		/// <summary>
+		/// The type of storage backend to use
+		/// </summary>
+		StorageBackendType? Type { get; }
+	}
+
+	/// <summary>
+	/// Common settings object for different providers
+	/// </summary>
+	public class StorageBackendOptions : IStorageBackendOptions
+	{
+		/// <inheritdoc/>
+		public StorageBackendType? Type { get; set; }
+
+		/// <inheritdoc/>
+		public string? BaseDir { get; set; }
+
+		/// <inheritdoc/>
+		public string? AwsBucketName { get; set; }
+
+		/// <inheritdoc/>
+		public string? AwsBucketPath { get; set; }
+
+		/// <inheritdoc/>
+		public AwsCredentialsType? AwsCredentials { get; set; }
+
+		/// <inheritdoc/>
+		public string? AwsRole { get; set; }
+
+		/// <inheritdoc/>
+		public string? AwsProfile { get; set; }
+
+		/// <inheritdoc/>
+		public string? AwsRegion { get; set; }
+	}
+
+	/// <summary>
+	/// Options for configuring a blob store
+	/// </summary>
+	public class BlobStoreOptions : StorageBackendOptions
+	{
+	}
+
+	/// <summary>
+	/// Options for configuring the default tree store implementation
+	/// </summary>
+	public interface ITreeStoreOptions
+	{
+		/// <summary>
+		/// Options for creating bundles
+		/// </summary>
+		TreeOptions Bundle { get; }
+
+		/// <summary>
+		/// Options for chunking content
+		/// </summary>
+		ChunkingOptions Chunking { get; }
+	}
+
+	/// <summary>
+	/// Options for storing trees
+	/// </summary>
+	public class TreeStoreOptions : BlobStoreOptions, ITreeStoreOptions
+	{
+		/// <inheritdoc/>
+		public TreeOptions Bundle { get; set; } = new TreeOptions();
+
+		/// <inheritdoc/>
+		public ChunkingOptions Chunking { get; set; } = new ChunkingOptions();
+	}
+
+	/// <summary>
+	/// Authentication method used for logging users in
+	/// </summary>
+	public enum AuthMethod
+	{
+		/// <summary>
+		/// No authentication enabled, mainly for demo and testing purposes
+		/// </summary>
+		Anonymous,
+
+		/// <summary>
+		/// OpenID Connect authentication, tailored for Okta
+		/// </summary>
+		Okta,
+		
+		/// <summary>
+		/// Generic OpenID Connect authentication, recommended for most
+		/// </summary>
+		OpenIdConnect,
+	}
+
+	/// <summary>
+	/// Type of run mode this process should use. Each carry different types of workloads. 
+	/// More than one mode can be active. But not all modes are not guaranteed to be compatible with each other and will
+	/// raise an error if combined in such a way.
+	/// </summary>
+	public enum RunMode
+	{
+		/// <summary>
+		/// Default no-op value (ASP.NET config will default to this for enums that cannot be parsed)
+		/// </summary> 
+		None,
+		
+		/// <summary>
+		/// Handle and respond to incoming external requests, such as HTTP REST and gRPC calls.
+		/// These requests are time-sensitive and short-lived, typically less than 5 secs.
+		/// If processes handling requests are unavailable, it will be very visible for users.
+		/// </summary>
+		Server,
+		
+		/// <summary>
+		/// Run non-request facing workloads. Such as background services, processing queues, running work
+		/// based on timers etc. Short periods of downtime or high CPU usage due to bursts are fine for this mode.
+		/// No user requests will be impacted directly. If auto-scaling is used, a much more aggressive policy can be
+		/// applied (tighter process packing, higher avg CPU usage).
+		/// </summary>
+		Worker
+	}
+
+	/// <summary>
+	/// Type of telemetry provider to use
+	/// </summary>
+	public enum TelemetrySinkType
+	{
+		/// <summary>
+		/// No telemetry sink (default)
+		/// </summary>
+		None,
+
+		/// <summary>
+		/// Use the Epic telemetry sink
+		/// </summary>
+		Epic,
+	}
+
+	/// <summary>
+	/// Configuration for the telemetry sink
+	/// </summary>
+	public class TelemetryConfig : IEpicTelemetrySinkConfig
+	{
+		/// <summary>
+		/// Type of telemetry sink
+		/// </summary>
+		public TelemetrySinkType Type { get; set; } = TelemetrySinkType.None;
+
+		/// <summary>
+		/// Base URL for the telemetry server
+		/// </summary>
+		public Uri? Url { get; set; }
+
+		/// <summary>
+		/// Application name to send in the event messages
+		/// </summary>
+		public string AppId { get; set; } = "Horde";
+	}
+	
+	/// <summary>
+	/// Feature flags to aid rollout of new features
+	///
+	/// Once a feature is running in its intended state and is stable, the flag should be removed.
+	/// A name and date of when the flag was created is noted next to it to help encourage this behavior.
+	/// Try having them be just a flag, a boolean.
+	/// </summary>
+	public class FeatureFlagSettings
+	{
+	}
+
+	/// <summary>
+	/// Options for the commit service
+	/// </summary>
+	public class CommitSettings
+	{
+		/// <summary>
+		/// Whether to mirror commit metadata to the database
+		/// </summary>
+		public bool ReplicateMetadata { get; set; } = true;
+
+		/// <summary>
+		/// Whether to mirror commit metadata to the database
+		/// </summary>
+		public bool ReplicateContent { get; set; } = false;
+
+		/// <summary>
+		/// Options for how objects are packed together
+		/// </summary>
+		public TreeOptions Bundle { get; set; } = new TreeOptions();
+
+		/// <summary>
+		/// Options for how objects are sliced
+		/// </summary>
+		public ChunkingOptions Chunking { get; set; } = new ChunkingOptions();
+	}
+
+	/// <summary>
+	/// Configuration for a tool bundled alongsize the server
+	/// </summary>
+	public class BundledToolConfig : ToolConfig
+	{
+		/// <summary>
+		/// Version string for the current tool data
+		/// </summary>
+		public string Version { get; set; } = "1.0";
+
+		/// <summary>
+		/// Ref name in the tools directory
+		/// </summary>
+		public RefName RefName { get; set; } = new RefName("default-ref");
+
+		/// <summary>
+		/// Directory containing blob data for this tool. If empty, the tools/{id} folder next to the server will be used.
+		/// </summary>
+		public string? DataDir { get; set; }
+
+		/// <summary>
+		/// Constructor
+		/// </summary>
+		public BundledToolConfig()
+		{
+			Public = true;
+		}
+	}
+
+	/// <summary>
+	/// Global settings for the application
+	/// </summary>
+	public class ServerSettings : IAclScope
+	{
+		/// <inheritdoc/>
+		[JsonIgnore]
+		public IAclScope? ParentScope => null;
+
+		/// <inheritdoc/>
+		[JsonIgnore]
+		public AclScopeName ScopeName => AclScopeName.Root;
+
+		/// <inheritdoc cref="RunMode" />
+		public RunMode[]? RunModes { get; set; } = null;
+
+		/// <summary>
+		/// Override the data directory used by Horde. Defaults to C:\ProgramData\HordeServer on Windows, {AppDir}/Data on other platforms.
+		/// </summary>
+		public string? DataDir { get; set; } = null;
+
+		/// <summary>
+		/// Output level for console
+		/// </summary>
+		public LogEventLevel ConsoleLogLevel { get; set; } = LogEventLevel.Debug;
+
+		/// <summary>
+		/// Main port for serving HTTP. Uses the default Kestrel port (5000) if not specified.
+		/// </summary>
+		public int HttpPort { get; set; }
+
+		/// <summary>
+		/// Port for serving HTTP with TLS enabled. Uses the default Kestrel port (5001) if not specified.
+		/// </summary>
+		public int HttpsPort { get; set; }
+
+		/// <summary>
+		/// Dedicated port for serving only HTTP/2.
+		/// </summary>
+		public int Http2Port { get; set; }
+
+		/// <summary>
+		/// Port for listening to compute tunnel initiator requests
+		/// </summary>
+		public int ComputeInitiatorPort { get; set; }
+
+		/// <summary>
+		/// Port for compute remotes to connect to
+		/// </summary>
+		public int ComputeRemotePort { get; set; }
+
+		/// <summary>
+		/// MongoDB connection string
+		/// </summary>
+		public string? DatabaseConnectionString { get; set; }
+
+		/// <summary>
+		/// MongoDB database name
+		/// </summary>
+		public string DatabaseName { get; set; } = "Horde";
+
+		/// <summary>
+		/// The claim type for administrators
+		/// </summary>
+		public string AdminClaimType { get; set; } = HordeClaimTypes.InternalRole;
+
+		/// <summary>
+		/// Value of the claim type for administrators
+		/// </summary>
+		public string AdminClaimValue { get; set; } = "admin";
+
+		/// <summary>
+		/// Optional certificate to trust in order to access the database (eg. AWS public cert for TLS)
+		/// </summary>
+		public string? DatabasePublicCert { get; set; }
+
+		/// <summary>
+		/// Access the database in read-only mode (avoids creating indices or updating content)
+		/// Useful for debugging a local instance of HordeServer against a production database.
+		/// </summary>
+		public bool DatabaseReadOnlyMode { get; set; } = false;
+
+		/// <summary>
+		/// Optional PFX certificate to use for encrypting agent SSL traffic. This can be a self-signed certificate, as long as it's trusted by agents.
+		/// </summary>
+		public string? ServerPrivateCert { get; set; }
+
+		/// <summary>
+		/// Issuer for tokens from the auth provider
+		/// </summary>
+		public AuthMethod AuthMethod { get; set; } = AuthMethod.Anonymous;
+
+		/// <summary>
+		/// Audience for OIDC validation
+		/// </summary>
+		public string? OidcAudience { get; set; }
+		
+		/// <summary>
+		/// Issuer for tokens from the auth provider
+		/// </summary>
+		public string? OidcAuthority { get; set; }
+
+		/// <summary>
+		/// Client id for the OIDC authority
+		/// </summary>
+		public string? OidcClientId { get; set; }
+
+		/// <summary>
+		/// Client secret for the OIDC authority
+		/// </summary>
+		public string? OidcClientSecret { get; set; }
+
+		/// <summary>
+		/// Optional redirect url provided to OIDC login
+		/// </summary>
+		public string? OidcSigninRedirect { get; set; }
+
+		/// <summary>
+		/// OpenID Connect scopes to request when signing in
+		/// </summary>
+		public string[] OidcRequestedScopes { get; set; } = { "profile", "email", "openid" };
+
+		/// <summary>
+		/// List of fields in /userinfo endpoint to try map to the standard name claim (see System.Security.Claims.ClaimTypes.Name)
+		/// </summary>
+		public string[] OidcClaimNameMapping { get; set; } = { "preferred_username", "email" };
+
+		/// <summary>
+		/// List of fields in /userinfo endpoint to try map to the standard email claim (see System.Security.Claims.ClaimTypes.Email)
+		/// </summary>
+		public string[] OidcClaimEmailMapping { get; set; } = { "email" };
+
+		/// <summary>
+		/// List of fields in /userinfo endpoint to try map to the Horde user claim (see HordeClaimTypes.User)
+		/// </summary>
+		public string[] OidcClaimHordeUserMapping { get; set; } = { "preferred_username", "email" };
+
+		/// <summary>
+		/// List of fields in /userinfo endpoint to try map to the Horde Perforce user claim (see HordeClaimTypes.PerforceUser)
+		/// </summary>
+		public string[] OidcClaimHordePerforceUserMapping { get; set; } = { "preferred_username", "email" };
+
+		/// <summary>
+		/// Name of the issuer in bearer tokens from the server
+		/// </summary>
+		public string? JwtIssuer { get; set; } = null!;
+
+		/// <summary>
+		/// Secret key used to sign JWTs. This setting is typically only used for development. In prod, a unique secret key will be generated and stored in the DB for each unique server instance.
+		/// </summary>
+		public string? JwtSecret { get; set; } = null!;
+
+		/// <summary>
+		/// Length of time before JWT tokens expire, in hours
+		/// </summary>
+		public int JwtExpiryTimeHours { get; set; } = 4;
+
+		/// <summary>
+		/// Whether to enable Cors, generally for development purposes
+		/// </summary>
+		public bool CorsEnabled { get; set; } = false;
+
+		/// <summary>
+		/// Allowed Cors origin 
+		/// </summary>
+		public string CorsOrigin { get; set; } = null!;
+
+		/// <summary>
+		/// Whether to enable a schedule in test data (false by default for development builds)
+		/// </summary>
+		public bool EnableScheduleInTestData { get; set; }
+
+		/// <summary>
+		/// The number of months to retain test data
+		/// </summary>
+		public int TestDataRetainMonths { get; set; } = 6;
+
+		/// <summary>
+		/// Interval between rebuilding the schedule queue with a DB query.
+		/// </summary>
+		public TimeSpan SchedulePollingInterval { get; set; } = TimeSpan.FromSeconds(60.0);
+
+		/// <summary>
+		/// Interval between polling for new jobs
+		/// </summary>
+		public TimeSpan NoResourceBackOffTime { get; set; } = TimeSpan.FromSeconds(30.0);
+
+		/// <summary>
+		/// Interval between attempting to assign agents to take on jobs
+		/// </summary>
+		public TimeSpan InitiateJobBackOffTime { get; set; } = TimeSpan.FromSeconds(180.0);
+
+		/// <summary>
+		/// Interval between scheduling jobs when an unknown error occurs
+		/// </summary>
+		public TimeSpan UnknownErrorBackOffTime { get; set; } = TimeSpan.FromSeconds(120.0);
+
+		/// <summary>
+		/// Config for connecting to Redis server(s).
+		/// Setting it to null will disable Redis use and connection
+		/// See format at https://stackexchange.github.io/StackExchange.Redis/Configuration.html
+		/// </summary>
+		public string? RedisConnectionConfig { get; set; }
+
+		/// <summary>
+		/// Type of write cache to use in log service
+		/// Currently Supported: "InMemory" or "Redis"
+		/// </summary>
+		public string LogServiceWriteCacheType { get; set; } = "InMemory";
+
+		/// <summary>
+		/// Settings for artifact storage
+		/// </summary>
+		public StorageBackendOptions LogStorage { get; set; } = new StorageBackendOptions() { BaseDir = "Logs" };
+
+		/// <summary>
+		/// Settings for artifact storage
+		/// </summary>
+		public StorageBackendOptions ArtifactStorage { get; set; } = new StorageBackendOptions() { BaseDir = "Artifacts" };
+
+		/// <summary>
+		/// Configuration of tree storage
+		/// </summary>
+		public TreeStoreOptions CommitStorage { get; set; } = new TreeStoreOptions() { BaseDir = "Commits" };
+
+		/// <summary>
+		/// Whether to log json to stdout
+		/// </summary>
+		public bool LogJsonToStdOut { get; set; } = false;
+
+		/// <summary>
+		/// Whether to log requests to the UpdateSession and QueryServerState RPC endpoints
+		/// </summary>
+		public bool LogSessionRequests { get; set; } = false;
+
+		/// <summary>
+		/// Default fleet manager to use (when not specified by pool)
+		/// </summary>
+		public FleetManagerType FleetManagerV2 { get; set; } = FleetManagerType.NoOp;
+
+		/// <summary>
+		/// Config for the fleet manager (serialized JSON)
+		/// </summary>
+		public string? FleetManagerV2Config { get; set; }
+
+		/// <summary>
+		/// Whether to run scheduled jobs.
+		/// </summary>
+		public bool DisableSchedules { get; set; }
+
+		/// <summary>
+		/// Timezone for evaluating schedules
+		/// </summary>
+		public string? ScheduleTimeZone { get; set; }
+
+		/// <summary>
+		/// Bot token for interacting with Slack (xoxb-*)
+		/// </summary>
+		public string? SlackToken { get; set; }
+
+		/// <summary>
+		/// Token for opening a socket to slack (xapp-*)
+		/// </summary>
+		public string? SlackSocketToken { get; set; }
+
+		/// <summary>
+		/// Admin user token for Slack (xoxp-*). This is only required when using the admin endpoints to invite users.
+		/// </summary>
+		public string? SlackAdminToken { get; set; }
+
+		/// <summary>
+		/// Filtered list of slack users to send notifications to. Should be Slack user ids, separated by commas.
+		/// </summary>
+		public string? SlackUsers { get; set; }
+
+		/// <summary>
+		/// Prefix to use when reporting errors
+		/// </summary>
+		public string SlackErrorPrefix { get; set; } = ":horde-error: ";
+
+		/// <summary>
+		/// Prefix to use when reporting warnings
+		/// </summary>
+		public string SlackWarningPrefix { get; set; } = ":horde-warning: ";
+
+		/// <summary>
+		/// Channel for sending messages related to config update failures
+		/// </summary>
+		public string? ConfigNotificationChannel { get; set; }
+
+		/// <summary>
+		/// Channel to send stream notification update failures to
+		/// </summary>
+		public string? UpdateStreamsNotificationChannel { get; set; }
+
+		/// <summary>
+		/// Slack channel to send job related notifications to. Multiple channels can be specified, separated by ;
+		/// </summary>
+		public string? JobNotificationChannel { get; set; }
+
+		/// <summary>
+		/// The URl to use for generating links back to the dashboard.
+		/// </summary>
+		public Uri DashboardUrl { get; set; } = new Uri("https://localhost:3000");
+
+		/// <summary>
+		/// Help email address that users can contact with issues
+		/// </summary>
+		public string? HelpEmailAddress { get; set; }
+
+		/// <summary>
+		/// Help slack channel that users can use for issues
+		/// </summary>
+		public string? HelpSlackChannel { get; set; }
+
+		/// <summary>
+		/// Url of P4 Swarm installation
+		/// </summary>
+		public Uri? P4SwarmUrl { get; set; }
+
+		/// <summary>
+		/// The Jira service account user name
+		/// </summary>
+		public string? JiraUsername { get; set; }
+
+		/// <summary>
+		/// The Jira service account API token
+		/// </summary>
+		public string? JiraApiToken { get; set; }
+
+		/// <summary>
+		/// The Uri for the Jira installation
+		/// </summary>
+		public Uri? JiraUrl { get; set; }
+
+		/// <summary>
+		/// The number of days shared device checkouts are held
+		/// </summary>
+		public int SharedDeviceCheckoutDays { get; set; } = 3;
+
+		/// <summary>
+		/// Default agent pool sizing strategy for pools that doesn't have one explicitly configured
+		/// </summary>
+		public PoolSizeStrategy DefaultAgentPoolSizeStrategy { get; set; } = PoolSizeStrategy.LeaseUtilization;
+
+		/// <summary>
+		/// Scale-out cooldown for auto-scaling agent pools (in seconds). Can be overridden by per-pool settings.
+		/// </summary>
+		public int AgentPoolScaleOutCooldownSeconds { get; set; } = 60; // 1 min
+
+		/// <summary>
+		/// Scale-in cooldown for auto-scaling agent pools (in seconds). Can be overridden by per-pool settings.
+		/// </summary>
+		public int AgentPoolScaleInCooldownSeconds { get; set; } = 1200; // 20 mins
+
+		/// <summary>
+		/// Set the minimum size of the global thread pool
+		/// This value has been found in need of tweaking to avoid timeouts with the Redis client during bursts
+		/// of traffic. Default is 16 for .NET Core CLR. The correct value is dependent on the traffic the Horde Server
+		/// is receiving. For Epic's internal deployment, this is set to 40.
+		/// </summary>
+		public int? GlobalThreadPoolMinSize { get; set; }
+
+		/// <summary>
+		/// Whether to enable Datadog integration for tracing
+		/// </summary>
+		public bool WithDatadog { get; set; }
+
+		/// <summary>
+		/// Whether to enable Amazon Web Services (AWS) specific features
+		/// </summary>
+		public bool WithAws { get; set; } = false;
+
+		/// <summary>
+		/// Path to the root config file
+		/// </summary>
+		public string ConfigPath { get; set; } = "Defaults/globals.json";
+
+		/// <summary>
+		/// Perforce connections for use by the Horde server (not agents)
+		/// </summary>
+		public List<PerforceConnectionSettings> Perforce { get; set; } = new List<PerforceConnectionSettings>();
+
+		/// <summary>
+		/// Whether to use the local Perforce environment
+		/// </summary>
+		public bool UseLocalPerforceEnv { get; set; }
+
+		/// <summary>
+		/// Number of pooled perforce connections to keep
+		/// </summary>
+		public int PerforceConnectionPoolSize { get; set; } = 5;
+
+		/// <summary>
+		/// Whether to enable the upgrade task source.
+		/// </summary>
+		public bool EnableUpgradeTasks { get; set; } = true;
+
+		/// <summary>
+		/// Whether to enable the conform task source.
+		/// </summary>
+		public bool EnableConformTasks { get; set; } = true;
+
+		/// <summary>
+		/// Forces configuration data to be read and updated as part of appplication startup, rather than on a schedule. Useful when running locally.
+		/// </summary>
+		public bool ForceConfigUpdateOnStartup { get; set; }
+
+		/// <summary>
+		/// Whether to open a browser on startup
+		/// </summary>
+		public bool OpenBrowser { get; set; } = false;
+
+		/// <inheritdoc cref="FeatureFlags" />
+		public FeatureFlagSettings FeatureFlags { get; set; } = new ();
+
+		/// <summary>
+		/// Options for the commit service
+		/// </summary>
+		public CommitSettings Commits { get; set; } = new CommitSettings();
+
+		/// <summary>
+		/// Settings for sending telemetry events
+		/// </summary>
+		public TelemetryConfig Telemetry { get; set; } = new TelemetryConfig();
+
+		/// <summary>
+		/// Tools bundled along with the server. Data for each tool can be produced using the 'bundle create' command, and should be stored in the /tools/{id} directory.
+		/// </summary>
+		public List<BundledToolConfig> BundledTools { get; set; } = new List<BundledToolConfig>();
+
+		/// <summary>
+		/// Default pre-baked ACL for authentication of well-known roles
+		/// </summary>
+		[JsonIgnore]
+		public AclConfig? Acl
+		{
+			get
+			{
+				_defaultAcl ??= GetDefaultAcl();
+				return _defaultAcl;
+			}
+		}
+		
+		[JsonIgnore]
+		AclConfig? _defaultAcl;
+
+		/// <summary>
+		/// Create the default ACL for the server, including all predefined roles.
+		/// </summary>
+		/// <returns></returns>
+		AclConfig GetDefaultAcl()
+		{
+			AclConfig defaultAcl = new AclConfig();
+			defaultAcl.Entries.Add(new AclEntryConfig(new AclClaimConfig(ClaimTypes.Role, "internal:AgentRegistration"), new[] { AclAction.CreateAgent, AclAction.CreateSession }));
+			defaultAcl.Entries.Add(new AclEntryConfig(HordeClaims.AgentRegistrationClaim, new[] { AclAction.CreateAgent, AclAction.CreateSession, AclAction.UpdateAgent, AclAction.DownloadSoftware, AclAction.CreatePool, AclAction.UpdatePool, AclAction.ViewPool, AclAction.DeletePool, AclAction.ListPools, AclAction.ViewStream, AclAction.ViewProject, AclAction.ViewJob, AclAction.ViewCosts }));
+			defaultAcl.Entries.Add(new AclEntryConfig(HordeClaims.AgentRoleClaim, new[] { AclAction.ViewProject, AclAction.ViewStream, AclAction.CreateEvent, AclAction.DownloadSoftware }));
+			defaultAcl.Entries.Add(new AclEntryConfig(HordeClaims.DownloadSoftwareClaim, new[] { AclAction.DownloadSoftware }));
+			defaultAcl.Entries.Add(new AclEntryConfig(HordeClaims.UploadSoftwareClaim, new[] { AclAction.UploadSoftware }));
+			defaultAcl.Entries.Add(new AclEntryConfig(HordeClaims.ConfigureProjectsClaim, new[] { AclAction.CreateProject, AclAction.UpdateProject, AclAction.ViewProject, AclAction.CreateStream, AclAction.UpdateStream, AclAction.ViewStream, AclAction.ChangePermissions }));
+			defaultAcl.Entries.Add(new AclEntryConfig(HordeClaims.StartChainedJobClaim, new[] { AclAction.CreateJob, AclAction.ExecuteJob, AclAction.UpdateJob, AclAction.ViewJob, AclAction.ViewTemplate, AclAction.ViewStream }));
+
+			if (AdminClaimType != null && AdminClaimValue != null)
+			{
+				AclAction[] actions = Enum.GetValues(typeof(AclAction)).OfType<AclAction>().ToArray();
+				defaultAcl.Entries.Add(new AclEntryConfig(new AclClaimConfig(AdminClaimType, AdminClaimValue), actions));
+			}
+
+			return defaultAcl;
+		}
+
+		/// <summary>
+		/// Helper method to check if this process has activated the given mode
+		/// </summary>
+		/// <param name="mode">Run mode</param>
+		/// <returns>True if mode is active</returns>
+		public bool IsRunModeActive(RunMode mode)
+		{
+			if (RunModes == null)
+			{
+				return true;
+			}
+			return RunModes.Contains(mode);
+		}
+
+		/// <summary>
+		/// Validate the settings object does not contain any invalid fields
+		/// </summary>
+		/// <exception cref="ArgumentException"></exception>
+		public void Validate()
+		{
+			if (RunModes != null && IsRunModeActive(RunMode.None))
+			{
+				throw new ArgumentException($"Settings key '{nameof(RunModes)}' contains one or more invalid entries");
+			}
+		}
+
+		/// <summary>
+		/// Attempts to get a bundled tool with the given id
+		/// </summary>
+		/// <param name="toolId">The tool id</param>
+		/// <param name="bundledToolConfig">Configuration for the bundled tool</param>
+		/// <returns>True if the tool was found</returns>
+		public bool TryGetBundledTool(ToolId toolId, [NotNullWhen(true)] out BundledToolConfig? bundledToolConfig)
+		{
+			bundledToolConfig = BundledTools.FirstOrDefault(x => x.Id == toolId);
+			return bundledToolConfig != null;
+		}
+	}
+
+	/// <summary>
+	/// Identifier for a pool
+	/// </summary>
+	/// <param name="Id">Id to construct from</param>
+	[JsonSchemaString]
+	[StringIdConverter(typeof(PerforceConnectionIdConverter))]
+	public record struct PerforceConnectionId(StringId Id)
+	{
+		/// <summary>
+		/// Constructor
+		/// </summary>
+		public PerforceConnectionId(string id) : this(new StringId(id))
+		{
+		}
+
+		/// <inheritdoc/>
+		public override string ToString() => Id.ToString();
+	}
+
+	/// <summary>
+	/// Converter to and from <see cref="StringId"/> instances.
+	/// </summary>
+	class PerforceConnectionIdConverter : StringIdConverter<PerforceConnectionId>
+	{
+		/// <inheritdoc/>
+		public override PerforceConnectionId FromStringId(StringId id) => new PerforceConnectionId(id);
+
+		/// <inheritdoc/>
+		public override StringId ToStringId(PerforceConnectionId value) => value.Id;
+	}
+
+	/// <summary>
+	/// Perforce connection information for use by the Horde server (for reading config files, etc...)
+	/// </summary>
+	public class PerforceConnectionSettings
+	{
+		/// <summary>
+		/// Identifier for the default perforce connection profile
+		/// </summary>
+		public static PerforceConnectionId Default { get; } = new PerforceConnectionId("default");
+
+		/// <summary>
+		/// Identifier for this server
+		/// </summary>
+		public PerforceConnectionId Id { get; set; } = Default;
+
+		/// <summary>
+		/// Server and port
+		/// </summary>
+		public string? ServerAndPort { get; set; }
+
+		/// <summary>
+		/// Credentials for the server
+		/// </summary>
+		public PerforceCredentials? Credentials { get; set; }
+
+		/// <summary>
+		/// Create a <see cref="PerforceSettings"/> object with these settings as overrides
+		/// </summary>
+		/// <returns>New perforce settings object</returns>
+		public PerforceSettings ToPerforceSettings()
+		{
+			PerforceSettings settings = new PerforceSettings(PerforceSettings.Default);
+			settings.PreferNativeClient = true;
+
+			if (!String.IsNullOrEmpty(ServerAndPort))
+			{
+				settings.ServerAndPort = ServerAndPort;
+			}
+			if (Credentials != null)
+			{
+				if (!String.IsNullOrEmpty(Credentials.UserName))
+				{
+					settings.UserName = Credentials.UserName;
+				}
+				if (!String.IsNullOrEmpty(Credentials.Password))
+				{
+					settings.Password = Credentials.Password;
+				}
+			}
+			return settings;
+		}
+	}
+}
