Horde: Rename the Horde.Build project to Horde.Server.

diff --git a/Engine/Source/Programs/Horde/Horde.Server/Startup.cs b/Engine/Source/Programs/Horde/Horde.Server/Startup.cs
new file mode 100644
--- /dev/null
+++ b/Engine/Source/Programs/Horde/Horde.Server/Startup.cs
@@ -0,0 +1,1140 @@
+// Copyright Epic Games, Inc. All Rights Reserved.
+
+using System;
+using System.Collections.Generic;
+using System.Diagnostics;
+using System.Globalization;
+using System.IO;
+using System.Linq;
+using System.Reflection;
+using System.Runtime.InteropServices;
+using System.Text.Json;
+using System.Text.Json.Serialization;
+using System.Threading;
+using System.Threading.Tasks;
+using Amazon;
+using Amazon.AutoScaling;
+using Amazon.CloudWatch;
+using Amazon.EC2;
+using Amazon.Extensions.NETCore.Setup;
+using EpicGames.AspNet;
+using EpicGames.Core;
+using EpicGames.Horde.Storage;
+using Grpc.Core;
+using Grpc.Core.Interceptors;
+using Horde.Server.Acls;
+using Horde.Server.Authentication;
+using Horde.Server.Agents;
+using Horde.Server.Agents.Fleet;
+using Horde.Server.Agents.Leases;
+using Horde.Server.Agents.Pools;
+using Horde.Server.Agents.Sessions;
+using Horde.Server.Agents.Software;
+using Horde.Server.Agents.Telemetry;
+using Horde.Server.Configuration;
+using Horde.Server.Dashboard;
+using Horde.Server.Devices;
+using Horde.Server.Issues;
+using Horde.Server.Issues.External;
+using Horde.Server.Jobs;
+using Horde.Server.Jobs.Artifacts;
+using Horde.Server.Jobs.Graphs;
+using Horde.Server.Jobs.Schedules;
+using Horde.Server.Jobs.Templates;
+using Horde.Server.Jobs.TestData;
+using Horde.Server.Jobs.Timing;
+using Horde.Server.Logs;
+using Horde.Server.Logs.Builder;
+using Horde.Server.Logs.Storage;
+using Horde.Server.Notifications;
+using Horde.Server.Secrets;
+using Horde.Server.Server;
+using Horde.Server.Storage;
+using Horde.Server.Storage.Backends;
+using Horde.Server.Tasks;
+using Horde.Server.Tools;
+using Horde.Server.Utilities;
+using HordeCommon;
+using Microsoft.AspNetCore.Authentication;
+using Microsoft.AspNetCore.Authentication.Cookies;
+using Microsoft.AspNetCore.Authentication.JwtBearer;
+using Microsoft.AspNetCore.Authentication.OpenIdConnect;
+using Microsoft.AspNetCore.Authorization;
+using Microsoft.AspNetCore.Builder;
+using Microsoft.AspNetCore.DataProtection;
+using Microsoft.AspNetCore.Hosting;
+using Microsoft.AspNetCore.Hosting.Server.Features;
+using Microsoft.AspNetCore.Http;
+using Microsoft.AspNetCore.HttpOverrides;
+using Microsoft.AspNetCore.Mvc;
+using Microsoft.AspNetCore.Mvc.ModelBinding;
+using Microsoft.Extensions.Configuration;
+using Microsoft.Extensions.DependencyInjection;
+using Microsoft.Extensions.Diagnostics.HealthChecks;
+using Microsoft.Extensions.Hosting;
+using Microsoft.Extensions.Logging;
+using Microsoft.Extensions.Options;
+using Microsoft.OpenApi.Models;
+using MongoDB.Bson;
+using MongoDB.Bson.Serialization;
+using MongoDB.Bson.Serialization.Conventions;
+using MongoDB.Bson.Serialization.Serializers;
+using OpenTracing.Contrib.Grpc.Interceptors;
+using OpenTracing.Util;
+using Serilog;
+using Serilog.Events;
+using StatsdClient;
+using Status = Grpc.Core.Status;
+using Horde.Server.Users;
+using Horde.Server.Perforce;
+using Horde.Server.Projects;
+using Horde.Server.Streams;
+using Horde.Server.Telemetry;
+using Horde.Server.Ugs;
+using Horde.Server.Auditing;
+using Horde.Server.Server.Notices;
+using Horde.Server.Notifications.Sinks;
+using StatusCode = Grpc.Core.StatusCode;
+using Horde.Server.Artifacts;
+using Horde.Server.Compute;
+
+namespace Horde.Server
+{
+	class Startup
+	{
+		static Startup()
+		{
+			ProtoBuf.Meta.RuntimeTypeModel.Default[typeof(ProjectId)].SetSurrogate(typeof(StringIdProto<ProjectId, ProjectIdConverter>));
+			ProtoBuf.Meta.RuntimeTypeModel.Default[typeof(StreamId)].SetSurrogate(typeof(StringIdProto<StreamId, StreamIdConverter>));
+		}
+
+		class GrpcExceptionInterceptor : Interceptor
+		{
+			readonly ILogger<GrpcExceptionInterceptor> _logger;
+
+			public GrpcExceptionInterceptor(ILogger<GrpcExceptionInterceptor> logger)
+			{
+				_logger = logger;
+			}
+
+			public override Task<TResponse> UnaryServerHandler<TRequest, TResponse>(TRequest request, ServerCallContext context, UnaryServerMethod<TRequest, TResponse> continuation)
+			{
+				return Guard(context, () => base.UnaryServerHandler(request, context, continuation));
+			}
+
+			public override Task<TResponse> ClientStreamingServerHandler<TRequest, TResponse>(IAsyncStreamReader<TRequest> requestStream, ServerCallContext context, ClientStreamingServerMethod<TRequest, TResponse> continuation) where TRequest : class where TResponse : class
+			{
+				return Guard(context, () => base.ClientStreamingServerHandler(requestStream, context, continuation));
+			}
+
+			public override Task ServerStreamingServerHandler<TRequest, TResponse>(TRequest request, IServerStreamWriter<TResponse> responseStream, ServerCallContext context, ServerStreamingServerMethod<TRequest, TResponse> continuation) where TRequest : class where TResponse : class
+			{
+				return Guard(context, () => base.ServerStreamingServerHandler(request, responseStream, context, continuation));
+			}
+
+			public override Task DuplexStreamingServerHandler<TRequest, TResponse>(IAsyncStreamReader<TRequest> requestStream, IServerStreamWriter<TResponse> responseStream, ServerCallContext context, DuplexStreamingServerMethod<TRequest, TResponse> continuation) where TRequest : class where TResponse : class
+			{
+				return Guard(context, () => base.DuplexStreamingServerHandler(requestStream, responseStream, context, continuation));
+			}
+
+			async Task<T> Guard<T>(ServerCallContext context, Func<Task<T>> callFunc) where T : class
+			{
+				T result = null!;
+				await Guard(context, async () => { result = await callFunc(); });
+				return result;
+			}
+
+			async Task Guard(ServerCallContext context, Func<Task> callFunc)
+			{
+				HttpContext httpContext = context.GetHttpContext();
+
+				AgentId? agentId = AclService.GetAgentId(httpContext.User);
+				if (agentId != null)
+				{
+					using IDisposable scope = _logger.BeginScope("Agent: {AgentId}, RemoteIP: {RemoteIP}, Method: {Method}", agentId.Value, httpContext.Connection.RemoteIpAddress, context.Method);
+					await GuardInner(context, callFunc);
+				}
+				else
+				{
+					using IDisposable scope = _logger.BeginScope("RemoteIP: {RemoteIP}, Method: {Method}", httpContext.Connection.RemoteIpAddress, context.Method);
+					await GuardInner(context, callFunc);
+				}
+			}
+
+			async Task GuardInner(ServerCallContext context, Func<Task> callFunc)
+			{
+				try
+				{
+					await callFunc();
+				}
+				catch (StructuredRpcException ex)
+				{
+#pragma warning disable CA2254 // Template should be a static expression
+					_logger.LogError(ex, ex.Format, ex.Args);
+#pragma warning restore CA2254 // Template should be a static expression
+					throw;
+				}
+				catch (Exception ex)
+				{
+					if (context.CancellationToken.IsCancellationRequested)
+					{
+						_logger.LogInformation(ex, "Call to method {Method} was cancelled", context.Method);
+						throw;
+					}
+					else
+					{
+						_logger.LogError(ex, "Exception in call to {Method}", context.Method);
+						throw new RpcException(new Status(StatusCode.Internal, $"An exception was thrown on the server: {ex}"));
+					}
+				}
+			}
+		}
+
+		class BsonSerializationProvider : IBsonSerializationProvider
+		{
+			public IBsonSerializer? GetSerializer(Type type)
+			{
+				if (type == typeof(ContentHash))
+				{
+					return new ContentHashSerializer();
+				}
+				if (type == typeof(DateTimeOffset))
+				{
+					return new DateTimeOffsetStringSerializer();
+				}
+				return null;
+			}
+		}
+
+		class JsonObjectIdConverter : JsonConverter<ObjectId>
+		{
+			public override ObjectId Read(ref Utf8JsonReader reader, Type typeToConvert, JsonSerializerOptions options)
+			{
+				string? str = reader.GetString();
+				if (str == null)
+				{
+					throw new InvalidDataException("Unable to parse object id");
+				}
+				return str.ToObjectId();
+			}
+
+			public override void Write(Utf8JsonWriter writer, ObjectId objectId, JsonSerializerOptions options)
+			{
+				writer.WriteStringValue(objectId.ToString());
+			}
+		}
+		
+		class JsonDateTimeConverter : JsonConverter<DateTime>
+		{
+			public override DateTime Read(ref Utf8JsonReader reader, Type typeToConvert, JsonSerializerOptions options)
+			{
+				Debug.Assert(typeToConvert == typeof(DateTime));
+
+				string? str = reader.GetString();
+				if (str == null)
+				{
+					throw new InvalidDataException("Unable to parse DateTime");
+				}
+				return DateTime.Parse(str, CultureInfo.CurrentCulture);
+			}
+
+			public override void Write(Utf8JsonWriter writer, DateTime dateTime, JsonSerializerOptions options)
+			{
+				writer.WriteStringValue(dateTime.ToUniversalTime().ToString("yyyy'-'MM'-'dd'T'HH':'mm':'ssZ", CultureInfo.CurrentCulture));
+			}
+		}
+
+		class JsonTimeSpanConverter : JsonConverter<TimeSpan>
+		{
+			public override TimeSpan Read(ref Utf8JsonReader reader, Type typeToConvert, JsonSerializerOptions options)
+			{
+				string? str = reader.GetString();
+				if (str == null)
+				{
+					throw new InvalidDataException("Unable to parse TimeSpan");
+				}
+				return TimeSpan.Parse(str, CultureInfo.CurrentCulture);
+			}
+
+			public override void Write(Utf8JsonWriter writer, TimeSpan timeSpan, JsonSerializerOptions options)
+			{
+				writer.WriteStringValue(timeSpan.ToString("c"));
+			}
+		}
+
+		static IStorageBackend CreateStorageBackend(IServiceProvider sp, StorageBackendOptions options)
+		{
+			switch (options.Type ?? StorageBackendType.FileSystem)
+			{
+				case StorageBackendType.FileSystem:
+					return new FileSystemStorageBackend(options);
+				case StorageBackendType.Aws:
+					return new AwsStorageBackend(sp.GetRequiredService<IConfiguration>(), options, sp.GetRequiredService<ILogger<AwsStorageBackend>>());
+				case StorageBackendType.Memory:
+					return new MemoryStorageBackend();
+				default:
+					throw new NotImplementedException();
+			}
+		}
+
+		public Startup(IConfiguration configuration)
+		{
+			Configuration = configuration;
+		}
+
+		public IConfiguration Configuration { get; }
+
+		// This method gets called *multiple times* by the runtime. Use this method to add services to the container.
+		public void ConfigureServices(IServiceCollection services)
+		{
+			// IOptionsMonitor pattern for live updating of configuration settings
+			services.Configure<ServerSettings>(Configuration.GetSection("Horde"));
+
+			// Settings used for configuring services
+			IConfigurationSection configSection = Configuration.GetSection("Horde");
+			ServerSettings settings = new ServerSettings();
+			configSection.Bind(settings);
+			settings.Validate();
+
+			if (settings.GlobalThreadPoolMinSize != null)
+			{
+				// Min thread pool size is set to combat timeouts seen with the Redis client.
+				// See comments for <see cref="ServerSettings.GlobalThreadPoolMinSize" /> and
+				// [REDACTED_PROFILE_URL]/StackExchange.Redis/issues/1680
+				int min = settings.GlobalThreadPoolMinSize.Value;
+				ThreadPool.SetMinThreads(min, min);
+			}
+
+#pragma warning disable CA2000 // Dispose objects before losing scope
+			RedisService redisService = new RedisService(settings);
+#pragma warning restore CA2000 // Dispose objects before losing scope
+			services.AddSingleton<RedisService>(sp => redisService);
+			services.AddDataProtection().PersistKeysToStackExchangeRedis(() => redisService.DatabaseSingleton, "aspnet-data-protection");
+
+			if (settings.CorsEnabled)
+			{
+				services.AddCors(options =>
+				{
+					options.AddPolicy("CorsPolicy",
+						builder => builder.WithOrigins(settings.CorsOrigin.Split(";"))
+						.AllowAnyMethod()
+						.AllowAnyHeader()
+						.AllowCredentials());
+				});
+			}
+
+			services.AddGrpc(options =>
+			{
+				options.EnableDetailedErrors = true;
+				options.MaxReceiveMessageSize = 200 * 1024 * 1024; // 100 MB (packaged builds of Horde agent can be large) 
+				options.Interceptors.Add(typeof(GrpcExceptionInterceptor));
+				options.Interceptors.Add<ServerTracingInterceptor>(GlobalTracer.Instance);
+			});
+			services.AddGrpcReflection();
+
+			services.AddHttpClient<JobRpcCommon>();
+			services.AddScoped<JobRpcCommon>();
+
+			services.AddSingleton<IAgentCollection, AgentCollection>();
+			services.AddSingleton<IArtifactCollection, ArtifactCollection>();
+			services.AddSingleton<IArtifactCollectionV1, ArtifactCollectionV1>();
+			services.AddSingleton<IGraphCollection, GraphCollection>();
+			services.AddSingleton<IIssueCollection, IssueCollection>();
+			services.AddSingleton<IJobCollection, JobCollection>();
+			services.AddSingleton<IJobStepRefCollection, JobStepRefCollection>();
+			services.AddSingleton<IJobTimingCollection, JobTimingCollection>();
+			services.AddSingleton<ILeaseCollection, LeaseCollection>();
+			services.AddSingleton<ILogEventCollection, LogEventCollection>();
+			services.AddSingleton<ILogFileCollection, LogFileCollection>();
+			services.AddSingleton<INotificationTriggerCollection, NotificationTriggerCollection>();
+			services.AddSingleton<IPoolCollection, PoolCollection>();
+			services.AddSingleton<ISessionCollection, SessionCollection>();
+			services.AddSingleton<IServiceAccountCollection, ServiceAccountCollection>();
+			services.AddSingleton<ISubscriptionCollection, SubscriptionCollection>();
+			services.AddSingleton<IStreamCollection, StreamCollection>();
+			services.AddSingleton<ITemplateCollection, TemplateCollection>();
+			services.AddSingleton<ITestDataCollection, TestDataCollection>();
+			services.AddSingleton<ITelemetryCollection, TelemetryCollection>();
+			services.AddSingleton<ITemplateCollection, TemplateCollection>();
+			services.AddSingleton<IUgsMetadataCollection, UgsMetadataCollection>();
+			services.AddSingleton<IUserCollection, UserCollectionV2>();
+			services.AddSingleton<IDeviceCollection, DeviceCollection>();
+			services.AddSingleton<INoticeCollection, NoticeCollection>();
+			services.AddSingleton<IDashboardPreviewCollection, DashboardPreviewCollection>();
+
+			services.AddSingleton<IToolCollection, ToolCollection>();
+
+			services.AddSingleton<IConfigSource, InMemoryConfigSource>();
+			services.AddSingleton<IConfigSource, FileConfigSource>();
+			services.AddSingleton<IConfigSource, PerforceConfigSource>();
+
+			services.AddSingleton<ConfigService>();
+			services.AddSingleton<IOptionsFactory<GlobalConfig>>(sp => sp.GetRequiredService<ConfigService>());
+			services.AddSingleton<IOptionsChangeTokenSource<GlobalConfig>>(sp => sp.GetRequiredService<ConfigService>());
+
+			// Always run the hosted config service, regardless of the server role, so we receive updates on the latest config values.
+			services.AddHostedService(provider => provider.GetRequiredService<ConfigService>());
+
+			// Auditing
+			services.AddSingleton<IAuditLog<AgentId>>(sp => sp.GetRequiredService<IAuditLogFactory<AgentId>>().Create("Agents.Log", "AgentId"));
+
+			services.AddSingleton(typeof(IAuditLogFactory<>), typeof(AuditLogFactory<>));
+			services.AddSingleton(typeof(ISingletonDocument<>), typeof(SingletonDocument<>));
+
+			services.AddSingleton<FleetService>();
+			services.AddSingleton<IFleetManagerFactory, FleetManagerFactory>();
+			
+			// Associate IFleetManager interface with the default implementation from config for convenience
+			// Though most fleet managers are created on a per-pool basis
+			services.AddSingleton<IFleetManager>(ctx => ctx.GetRequiredService<IFleetManagerFactory>().CreateFleetManager(FleetManagerType.Default));
+
+			// Run the tunnel service for all runmodes
+			services.AddSingleton<TunnelService>();
+			services.AddHostedService<TunnelService>(sp => sp.GetRequiredService<TunnelService>());
+
+			services.AddSingleton<AclService>();
+			services.AddSingleton<AgentService>();
+			services.AddSingleton<ConsistencyService>();
+			services.AddSingleton<RequestTrackerService>();
+			services.AddSingleton<CredentialService>();
+			services.AddSingleton<MongoService>();
+			services.AddSingleton<GlobalsService>();
+			services.AddSingleton<IDogStatsd>(ctx =>
+			{
+				string? datadogAgentHost = Environment.GetEnvironmentVariable("DD_AGENT_HOST");
+				if (datadogAgentHost != null)
+				{
+					// Datadog agent is configured, enable DogStatsD for metric collection
+					StatsdConfig config = new StatsdConfig
+					{
+						StatsdServerName = datadogAgentHost,
+						StatsdPort = 8125,
+					};
+
+					DogStatsdService dogStatsdService = new DogStatsdService();
+					dogStatsdService.Configure(config);
+					return dogStatsdService;
+				}
+				return new NoOpDogStatsd();
+			});
+			services.AddSingleton<ICommitService, CommitService>();
+			services.AddSingleton<IClock, Clock>();
+			services.AddSingleton<IDowntimeService, DowntimeService>();
+			services.AddSingleton<IssueService>();
+			services.AddSingleton<JobService>();
+			services.AddSingleton<LifetimeService>();
+			services.AddSingleton<ILogFileService, LogFileService>();
+			services.AddSingleton<LogTailService>();
+			services.AddHostedService(provider => provider.GetRequiredService<LogTailService>()); // NB: Runs even on workers, to receive tail notifications.
+			services.AddSingleton<INotificationService, NotificationService>();
+
+			if (settings.Commits.ReplicateMetadata)
+			{
+				services.AddSingleton<PerforceServiceCache>();
+				services.AddSingleton<IPerforceService>(sp => sp.GetRequiredService<PerforceServiceCache>());
+			}
+			else
+			{
+				services.AddSingleton<PerforceService>();
+				services.AddSingleton<IPerforceService>(sp => sp.GetRequiredService<PerforceService>());
+			}
+			services.AddSingleton<PerforceReplicator>();
+
+			services.AddSingleton<PerforceLoadBalancer>();
+			services.AddSingleton<PoolService>();
+			services.AddSingleton<ReplicationService>();
+			services.AddSingleton<ScheduleService>();
+
+			if (settings.SlackToken != null)
+			{
+				services.AddSingleton<SlackNotificationSink>();
+				services.AddSingleton<IAvatarService, SlackNotificationSink>(sp => sp.GetRequiredService<SlackNotificationSink>());
+				services.AddSingleton<INotificationSink, SlackNotificationSink>(sp => sp.GetRequiredService<SlackNotificationSink>());
+			}
+			else
+			{
+				services.AddSingleton<IAvatarService, NullAvatarService>();
+			}
+
+			services.AddSingleton<DeviceService>();						
+			services.AddSingleton<NoticeService>();
+			services.AddSingleton<StorageService>();
+			services.AddSingleton<IStorageClientFactory>(sp => sp.GetRequiredService<StorageService>());
+			services.AddSingleton<TestDataService>();
+
+			if (settings.JiraUrl != null)
+			{
+				services.AddSingleton<IExternalIssueService, JiraService>();
+			}
+			else
+			{
+				services.AddSingleton<IExternalIssueService, DefaultExternalIssueService>();
+			}
+
+			// Storage providers
+			services.AddSingleton(sp => CreateStorageBackend(sp, settings.LogStorage).ForType<PersistentLogStorage>());
+			services.AddSingleton(sp => CreateStorageBackend(sp, settings.ArtifactStorage).ForType<ArtifactCollectionV1>());
+
+			if (settings.WithAws)
+			{
+				AWSOptions awsOptions = Configuration.GetAWSOptions();
+				services.AddDefaultAWSOptions(awsOptions);
+				if (awsOptions.Region == null && Environment.GetEnvironmentVariable("AWS_REGION") == null)
+				{
+					awsOptions.Region = RegionEndpoint.USEast1;
+				}
+				
+				services.AddAWSService<IAmazonCloudWatch>();
+				services.AddAWSService<IAmazonAutoScaling>();
+				services.AddAWSService<IAmazonEC2>();
+			}
+
+			ConfigureLogStorage(services);
+
+			AuthenticationBuilder authBuilder = services.AddAuthentication(options =>
+				{
+					switch (settings.AuthMethod)
+					{
+						case AuthMethod.Anonymous:
+							options.DefaultAuthenticateScheme = AnonymousAuthenticationHandler.AuthenticationScheme;
+							options.DefaultSignInScheme = AnonymousAuthenticationHandler.AuthenticationScheme;
+							options.DefaultChallengeScheme = AnonymousAuthenticationHandler.AuthenticationScheme;
+							break;
+						
+						case AuthMethod.Okta:
+							// If an authentication cookie is present, use it to get authentication information
+							options.DefaultAuthenticateScheme = CookieAuthenticationDefaults.AuthenticationScheme;
+							options.DefaultSignInScheme = CookieAuthenticationDefaults.AuthenticationScheme;
+
+							// If authentication is required, and no cookie is present, use OIDC to sign in
+							options.DefaultChallengeScheme = OktaDefaults.AuthenticationScheme;
+							break;
+						
+						case AuthMethod.OpenIdConnect:
+							// If an authentication cookie is present, use it to get authentication information
+							options.DefaultAuthenticateScheme = CookieAuthenticationDefaults.AuthenticationScheme;
+							options.DefaultSignInScheme = CookieAuthenticationDefaults.AuthenticationScheme;
+
+							// If authentication is required, and no cookie is present, use OIDC to sign in
+							options.DefaultChallengeScheme = OpenIdConnectDefaults.AuthenticationScheme;
+							break;
+						
+						default:
+							throw new ArgumentException($"Invalid auth method {settings.AuthMethod}");
+					}
+				});
+
+			List<string> schemes = new List<string>();
+
+			authBuilder.AddCookie(options =>
+				 {
+					 options.Events.OnValidatePrincipal = context =>
+					 {
+						 if (context.Principal?.FindFirst(HordeClaimTypes.UserId) == null)
+						 {
+							 context.RejectPrincipal();
+						 }
+						 return Task.CompletedTask;
+					 };
+
+					 options.Events.OnRedirectToAccessDenied = context =>
+					 {
+						 context.Response.StatusCode = StatusCodes.Status403Forbidden;
+						 return context.Response.CompleteAsync();
+					 };
+				 });
+			schemes.Add(CookieAuthenticationDefaults.AuthenticationScheme);
+
+			authBuilder.AddServiceAccount(options => { });
+			schemes.Add(ServiceAccountAuthHandler.AuthenticationScheme);
+
+			
+			switch (settings.AuthMethod)
+			{
+				case AuthMethod.Anonymous:
+					authBuilder.AddAnonymous(options =>
+					{
+						options.AdminClaimType = settings.AdminClaimType;
+						options.AdminClaimValue = settings.AdminClaimValue;
+					});
+					schemes.Add(AnonymousAuthenticationHandler.AuthenticationScheme);
+					break;
+						
+				case AuthMethod.Okta:
+					authBuilder.AddOkta(OktaDefaults.AuthenticationScheme, OpenIdConnectDefaults.DisplayName, options =>
+						{
+							options.Authority = settings.OidcAuthority;
+							options.ClientId = settings.OidcClientId;
+								options.Scope.Remove("groups");								
+
+							if (!String.IsNullOrEmpty(settings.OidcSigninRedirect))
+							{
+								options.Events = new OpenIdConnectEvents
+								{
+									OnRedirectToIdentityProvider = async redirectContext =>
+									{
+										redirectContext.ProtocolMessage.RedirectUri = settings.OidcSigninRedirect;
+										await Task.CompletedTask;
+									}
+								};
+							}
+						});
+					schemes.Add(OktaDefaults.AuthenticationScheme);
+					break;
+						
+				case AuthMethod.OpenIdConnect:
+					authBuilder.AddHordeOpenId(settings, OpenIdConnectDefaults.AuthenticationScheme, OpenIdConnectDefaults.DisplayName, options =>
+						{
+							options.Authority = settings.OidcAuthority;
+							options.ClientId = settings.OidcClientId;
+							options.ClientSecret = settings.OidcClientSecret;
+							foreach (string scope in settings.OidcRequestedScopes)
+							{
+								options.Scope.Add(scope);
+							}
+
+							if (!String.IsNullOrEmpty(settings.OidcSigninRedirect))
+							{
+								options.Events = new OpenIdConnectEvents
+								{
+									OnRedirectToIdentityProvider = async redirectContext =>
+									{
+										redirectContext.ProtocolMessage.RedirectUri = settings.OidcSigninRedirect;
+										await Task.CompletedTask;
+									}
+								};
+							}
+						});
+					schemes.Add(OpenIdConnectDefaults.AuthenticationScheme);
+					break;
+						
+				default:
+					throw new ArgumentException($"Invalid auth method {settings.AuthMethod}");
+			}
+
+			TelemetryConfig telemetryConfig = settings.Telemetry;
+			switch (telemetryConfig.Type)
+			{
+				case TelemetrySinkType.None:
+					services.AddSingleton<ITelemetrySink, NullTelemetrySink>();
+					break;
+				case TelemetrySinkType.Epic:
+					services.AddHttpClient(EpicTelemetrySink.HttpClientName, client => { });
+					services.AddSingleton<EpicTelemetrySink>();
+					services.AddSingleton<ITelemetrySink>(sp => sp.GetRequiredService<EpicTelemetrySink>());
+					services.AddHostedService(sp => sp.GetRequiredService<EpicTelemetrySink>());
+					break;
+			}
+
+			authBuilder.AddScheme<JwtBearerOptions, HordeServerJwtBearerHandler>(HordeServerJwtBearerHandler.AuthenticationScheme, options => { });
+			schemes.Add(HordeServerJwtBearerHandler.AuthenticationScheme);
+
+			if (settings.OidcAuthority != null && settings.OidcAudience != null)
+			{
+				HordeJwtBearerHandler hordeJwtBearer = new(settings);  
+				hordeJwtBearer.AddHordeJwtBearerConfiguration(authBuilder);  
+				schemes.Add(HordeJwtBearerHandler.AuthenticationScheme);				
+			}
+
+			services.AddAuthorization(options =>
+				{
+					options.DefaultPolicy = new AuthorizationPolicyBuilder(schemes.ToArray())
+						.RequireAuthenticatedUser()							
+						.Build();
+				});
+			
+			// Hosted service that needs to run no matter the run mode of the process (server vs worker)
+			services.AddHostedService(provider => (DowntimeService)provider.GetRequiredService<IDowntimeService>());
+
+			if (settings.IsRunModeActive(RunMode.Worker) && !settings.DatabaseReadOnlyMode)
+			{
+				services.AddHostedService<MongoUpgradeService>();
+
+				services.AddHostedService(provider => provider.GetRequiredService<FleetService>());
+				
+				services.AddHostedService(provider => provider.GetRequiredService<AgentService>());
+
+				if (settings.Commits.ReplicateMetadata)
+				{
+					services.AddHostedService(provider => provider.GetRequiredService<PerforceServiceCache>());
+				}
+
+				services.AddHostedService(provider => provider.GetRequiredService<ConsistencyService>());
+				services.AddHostedService(provider => provider.GetRequiredService<IssueService>());
+				services.AddHostedService<IssueReportService>();
+				services.AddHostedService(provider => (LogFileService)provider.GetRequiredService<ILogFileService>());
+				services.AddHostedService(provider => (NotificationService)provider.GetRequiredService<INotificationService>());
+				services.AddHostedService(provider => provider.GetRequiredService<ReplicationService>());
+				if (!settings.DisableSchedules)
+				{
+					services.AddHostedService(provider => provider.GetRequiredService<ScheduleService>());
+				}
+
+				services.AddHostedService<MetricService>();
+				services.AddHostedService(provider => provider.GetRequiredService<PerforceLoadBalancer>());
+				services.AddHostedService<PoolUpdateService>();
+				if (settings.SlackToken != null)
+				{
+					services.AddHostedService(provider => provider.GetRequiredService<SlackNotificationSink>());
+				}
+				services.AddHostedService<TelemetryService>();
+				services.AddHostedService(provider => provider.GetRequiredService<DeviceService>());
+				services.AddHostedService(provider => provider.GetRequiredService<TestDataService>());
+			}
+
+			services.AddHostedService(provider => provider.GetRequiredService<IExternalIssueService>());
+
+			// Task sources. Order of registration is important here; it dictates the priority in which sources are served.
+			if (!settings.DatabaseReadOnlyMode)
+			{
+				services.AddSingleton<JobTaskSource>();
+				services.AddHostedService<JobTaskSource>(provider => provider.GetRequiredService<JobTaskSource>());
+				services.AddSingleton<ConformTaskSource>();
+				services.AddHostedService<ConformTaskSource>(provider => provider.GetRequiredService<ConformTaskSource>());
+				services.AddSingleton<ComputeTaskSource>();
+
+				services.AddSingleton<ITaskSource, UpgradeTaskSource>();
+				services.AddSingleton<ITaskSource, ShutdownTaskSource>();
+				services.AddSingleton<ITaskSource, RestartTaskSource>();
+				services.AddSingleton<ITaskSource, ConformTaskSource>(provider => provider.GetRequiredService<ConformTaskSource>());
+				services.AddSingleton<ITaskSource, JobTaskSource>(provider => provider.GetRequiredService<JobTaskSource>());
+				services.AddSingleton<ITaskSource, ComputeTaskSource>(provider => provider.GetRequiredService<ComputeTaskSource>());
+			}
+
+			// Allow longer to shutdown so we can debug missing cancellation tokens
+			services.Configure<HostOptions>(options =>
+			{
+				options.ShutdownTimeout = TimeSpan.FromSeconds(30.0);
+			});
+
+			// Allow forwarded headers
+			services.Configure<ForwardedHeadersOptions>(options =>
+			{
+				options.ForwardedHeaders = ForwardedHeaders.XForwardedFor | ForwardedHeaders.XForwardedProto;
+				options.KnownProxies.Clear();
+				options.KnownNetworks.Clear();
+			});
+
+			services.AddMvc().AddJsonOptions(options => ConfigureJsonSerializer(options.JsonSerializerOptions));
+
+			services.AddControllers(options =>
+			{
+				options.InputFormatters.Add(new CbInputFormatter());
+				options.OutputFormatters.Add(new CbOutputFormatter());
+				options.OutputFormatters.Insert(0, new CbPreferredOutputFormatter());
+				options.FormatterMappings.SetMediaTypeMappingForFormat("uecb", CustomMediaTypeNames.UnrealCompactBinary);
+			});
+
+			services.AddSwaggerGen(config =>
+			{
+				config.SwaggerDoc("v1", new OpenApiInfo { Title = "Horde Server API", Version = "v1" });
+				config.IncludeXmlComments(Path.Combine(AppContext.BaseDirectory, $"{Assembly.GetExecutingAssembly().GetName().Name}.xml"));
+			});
+
+			services.Configure<ApiBehaviorOptions>(options =>
+			{
+				options.InvalidModelStateResponseFactory = context =>
+				{
+					foreach(KeyValuePair<string, ModelStateEntry> pair in context.ModelState)
+					{
+						ModelError? error = pair.Value.Errors.FirstOrDefault();
+						if (error != null)
+						{
+							string message = error.ErrorMessage;
+							if (String.IsNullOrEmpty(message))
+							{
+								message = error.Exception?.Message ?? "Invalid error object";
+							}
+							return new BadRequestObjectResult(EpicGames.Core.LogEvent.Create(LogLevel.Error, KnownLogEvents.None, error.Exception, "Invalid value for {Name}: {Message}", pair.Key, message));
+						}
+					}
+					return new BadRequestObjectResult(context.ModelState);
+				};
+			});
+
+			DirectoryReference dashboardDir = DirectoryReference.Combine(Program.AppDir, "DashboardApp");
+			if (DirectoryReference.Exists(dashboardDir)) 
+			{
+				services.AddSpaStaticFiles(config => { config.RootPath = "DashboardApp"; });
+			}
+
+			ConfigureMongoDbClient();
+			ConfigureFormatters();
+
+			OnAddHealthChecks(services);
+		}
+
+		public static void ConfigureFormatters()
+		{
+			LogValueFormatter.RegisterTypeAnnotation<AgentId>("AgentId");
+			LogValueFormatter.RegisterTypeAnnotation<JobId>("JobId");
+			LogValueFormatter.RegisterTypeAnnotation<LeaseId>("LeaseId");
+			LogValueFormatter.RegisterTypeAnnotation<LogId>("LogId");
+			LogValueFormatter.RegisterTypeAnnotation<UserId>("UserId");
+		}
+
+		public static void ConfigureJsonSerializer(JsonSerializerOptions options)
+		{
+			options.AllowTrailingCommas = true;
+			options.ReadCommentHandling = JsonCommentHandling.Skip;
+			options.DefaultIgnoreCondition = JsonIgnoreCondition.WhenWritingNull;
+			options.PropertyNamingPolicy = JsonNamingPolicy.CamelCase;
+			options.PropertyNameCaseInsensitive = true;
+			options.Converters.Add(new JsonObjectIdConverter());
+			options.Converters.Add(new JsonStringEnumConverter());
+			options.Converters.Add(new JsonKnownTypesConverterFactory());
+			options.Converters.Add(new StringIdJsonConverterFactory());
+			options.Converters.Add(new ObjectIdJsonConverterFactory());
+			options.Converters.Add(new JsonDateTimeConverter());
+			options.Converters.Add(new JsonTimeSpanConverter());
+		}
+
+		private static void ConfigureLogStorage(IServiceCollection services)
+		{
+			services.AddSingleton<ILogBuilder>(provider =>
+			{
+				RedisService? redisService = provider.GetService<RedisService>();
+				if(redisService == null)
+				{
+					return new LocalLogBuilder();
+				}
+				else
+				{
+					return new RedisLogBuilder(redisService.ConnectionPool, provider.GetRequiredService<ILogger<RedisLogBuilder>>());
+				}
+			});
+
+			services.AddSingleton<PersistentLogStorage>();
+
+			services.AddSingleton<ILogStorage>(provider =>
+			{
+				ILogStorage storage = provider.GetRequiredService<PersistentLogStorage>();
+
+//				IDatabase? RedisDb = Provider.GetService<IDatabase>();
+//				if (RedisDb != null)
+//				{
+//					Storage = new RedisLogStorage(RedisDb, Provider.GetService<ILogger<RedisLogStorage>>(), Storage);
+//				}
+
+				storage = new SequencedLogStorage(storage);
+				storage = new LocalLogStorage(50, storage);
+				return storage;
+			});
+		}
+		/*
+				private static void ConfigureLogFileWriteCache(IServiceCollection Services, ServerSettings Settings)
+				{
+					bool RedisConfigured = !String.IsNullOrEmpty(Settings.RedisConnectionConfig);
+					string CacheType = Settings.LogServiceWriteCacheType.ToLower(CultureInfo.CurrentCulture);
+
+					if (CacheType == "inmemory")
+					{
+						Services.AddSingleton<ILogFileWriteCache2>(Sp => new InMemoryLogFileWriteCache2());
+					}
+					else if (CacheType == "redis" && RedisConfigured)
+					{
+						Services.AddSingleton<ILogFileWriteCache2>(Sp => new RedisLogFileWriteCache2(Sp.GetService<ILogger<RedisLogFileWriteCache>>(), Sp.GetService<IDatabase>()));
+					}
+					else if (CacheType == "redis" && !RedisConfigured)
+					{
+						throw new Exception("Redis must be configured to use the Redis-backed log write cache");
+					}
+					else
+					{
+						throw new Exception("Unknown value set for LogServiceWriteCacheType in config: " + Settings.LogServiceWriteCacheType);
+					}
+				}
+		*/
+
+		public sealed class HostIdBsonSerializer : SerializerBase<HostId>
+		{
+			/// <inheritdoc/>
+			public override HostId Deserialize(BsonDeserializationContext context, BsonDeserializationArgs args) => new HostId(context.Reader.ReadString());
+
+			/// <inheritdoc/>
+			public override void Serialize(BsonSerializationContext context, BsonSerializationArgs args, HostId value) => context.Writer.WriteString(value.ToString());
+		}
+
+		public sealed class BlobIdBsonSerializer : SerializerBase<BlobId>
+		{
+			/// <inheritdoc/>
+			public override BlobId Deserialize(BsonDeserializationContext context, BsonDeserializationArgs args) => new BlobId(context.Reader.ReadString());
+
+			/// <inheritdoc/>
+			public override void Serialize(BsonSerializationContext context, BsonSerializationArgs args, BlobId value) => context.Writer.WriteString(value.ToString());
+		}
+
+		public sealed class BlobLocatorBsonSerializer : SerializerBase<BlobLocator>
+		{
+			/// <inheritdoc/>
+			public override BlobLocator Deserialize(BsonDeserializationContext context, BsonDeserializationArgs args)
+			{
+				return new BlobLocator(context.Reader.ReadString());
+			}
+
+			/// <inheritdoc/>
+			public override void Serialize(BsonSerializationContext context, BsonSerializationArgs args, BlobLocator value)
+			{
+				context.Writer.WriteString(value.ToString());
+			}
+		}
+
+		public sealed class RefNameBsonSerializer : SerializerBase<RefName>
+		{
+			/// <inheritdoc/>
+			public override RefName Deserialize(BsonDeserializationContext context, BsonDeserializationArgs args)
+			{
+				return new RefName(context.Reader.ReadString());
+			}
+
+			/// <inheritdoc/>
+			public override void Serialize(BsonSerializationContext context, BsonSerializationArgs args, RefName value)
+			{
+				context.Writer.WriteString(value.ToString());
+			}
+		}
+
+		public sealed class IoHashBsonSerializer : SerializerBase<IoHash>
+		{
+			/// <inheritdoc/>
+			public override IoHash Deserialize(BsonDeserializationContext context, BsonDeserializationArgs args)
+			{
+				return IoHash.Parse(context.Reader.ReadString());
+			}
+
+			/// <inheritdoc/>
+			public override void Serialize(BsonSerializationContext context, BsonSerializationArgs args, IoHash value)
+			{
+				context.Writer.WriteString(value.ToString());
+			}
+		}
+
+		public sealed class NamespaceIdBsonSerializer : SerializerBase<NamespaceId>
+		{
+			/// <inheritdoc/>
+			public override NamespaceId Deserialize(BsonDeserializationContext context, BsonDeserializationArgs args)
+			{
+				return new NamespaceId(context.Reader.ReadString());
+			}
+
+			/// <inheritdoc/>
+			public override void Serialize(BsonSerializationContext context, BsonSerializationArgs args, NamespaceId value)
+			{
+				context.Writer.WriteString(value.ToString());
+			}
+		}
+
+		static int s_haveConfiguredMongoDb = 0;
+
+		public static void ConfigureMongoDbClient()
+		{
+			if (Interlocked.CompareExchange(ref s_haveConfiguredMongoDb, 1, 0) == 0)
+			{
+				// Ignore extra elements on deserialized documents
+				ConventionPack conventionPack = new ConventionPack();
+				conventionPack.Add(new IgnoreExtraElementsConvention(true));
+				conventionPack.Add(new EnumRepresentationConvention(BsonType.String));
+				ConventionRegistry.Register("Horde", conventionPack, type => true);
+
+				// Register the custom serializers
+				BsonSerializer.RegisterSerializer(new HostIdBsonSerializer());
+				BsonSerializer.RegisterSerializer(new BlobIdBsonSerializer());
+				BsonSerializer.RegisterSerializer(new BlobLocatorBsonSerializer());
+				BsonSerializer.RegisterSerializer(new RefNameBsonSerializer());
+				BsonSerializer.RegisterSerializer(new IoHashBsonSerializer());
+				BsonSerializer.RegisterSerializer(new NamespaceIdBsonSerializer());
+				BsonSerializer.RegisterSerializer(new ConditionSerializer());
+				BsonSerializer.RegisterSerializationProvider(new BsonSerializationProvider());
+				BsonSerializer.RegisterSerializationProvider(new StringIdBsonSerializationProvider());
+				BsonSerializer.RegisterSerializationProvider(new ObjectIdBsonSerializationProvider());
+			}
+		}
+
+		private static void OnAddHealthChecks(IServiceCollection services)
+		{
+			services.AddHealthChecks()
+				.AddCheck("self", () => HealthCheckResult.Healthy(), tags: new[] { "self" })
+				.AddCheck<RedisService>("redis")
+				.AddCheck<MongoService>("mongo");
+		}
+
+		// This method gets called by the runtime. Use this method to configure the HTTP request pipeline.
+		public static void Configure(IApplicationBuilder app, IWebHostEnvironment env, Microsoft.Extensions.Hosting.IHostApplicationLifetime lifetime, IOptions<ServerSettings> settings)
+		{
+			app.UseForwardedHeaders();
+
+			// Used for allowing auth cookies in combination with OpenID Connect auth (for example, Google Auth did not work with these unset)
+			app.UseCookiePolicy(new CookiePolicyOptions()
+			{
+				MinimumSameSitePolicy = SameSiteMode.None,
+				CheckConsentNeeded = _ => true
+			});
+
+			if (settings.Value.CorsEnabled)
+			{
+				app.UseCors("CorsPolicy");
+			}
+
+			// Enable middleware to serve generated Swagger as a JSON endpoint.
+			app.UseSwagger();
+
+			// Enable serilog request logging
+			app.UseSerilogRequestLogging(options =>
+			{
+				options.GetLevel = GetRequestLoggingLevel;
+				options.EnrichDiagnosticContext = (diagnosticContext, httpContext) =>
+				{
+					diagnosticContext.Set("RemoteIP", httpContext?.Connection?.RemoteIpAddress);
+				};
+			});
+
+			// Enable middleware to serve swagger-ui (HTML, JS, CSS, etc.),
+			// specifying the Swagger JSON endpoint.
+			app.UseSwaggerUI(c =>
+			{
+				c.SwaggerEndpoint("/swagger/v1/swagger.json", "Horde Server API");
+				c.RoutePrefix = "swagger";
+			});
+
+			if (!env.IsDevelopment())
+			{
+				app.UseMiddleware<RequestTrackerMiddleware>();	
+			}
+
+			app.UseDefaultFiles();
+			app.UseStaticFiles();
+
+			DirectoryReference dashboardDir = DirectoryReference.Combine(Program.AppDir, "DashboardApp");
+
+			if (DirectoryReference.Exists(dashboardDir)) 
+			{
+				app.UseSpaStaticFiles();
+			}
+						
+			app.UseRouting();
+
+			app.UseAuthentication();
+			app.UseAuthorization();
+
+			app.UseEndpoints(endpoints =>
+			{
+				endpoints.MapGrpcService<HealthService>();
+				endpoints.MapGrpcService<RpcService>();
+				endpoints.MapGrpcService<JobRpcService>();
+				endpoints.MapGrpcService<LogRpcService>();
+
+				endpoints.MapGrpcReflectionService();
+
+				endpoints.MapControllers();
+
+				endpoints.MapHealthChecks("/health/live");
+			});
+
+			if (DirectoryReference.Exists(dashboardDir)) 
+			{
+				app.UseSpa(spa =>
+				{ 
+					spa.Options.SourcePath = "DashboardApp";        
+				});
+			}
+
+			if (settings.Value.OpenBrowser && RuntimeInformation.IsOSPlatform(OSPlatform.Windows))
+			{
+				lifetime.ApplicationStarted.Register(() => LaunchBrowser(app));
+			}
+		}
+
+		static void LaunchBrowser(IApplicationBuilder app)
+		{
+			IServerAddressesFeature? feature = app.ServerFeatures.Get<IServerAddressesFeature>();
+			if (feature != null && feature.Addresses.Count > 0)
+			{
+				// with a development cert, host will be set by default to localhost, otherwise there will be no host in address
+				string address = feature.Addresses.First().Replace("[::]", System.Net.Dns.GetHostName(), StringComparison.OrdinalIgnoreCase);
+				Process.Start(new ProcessStartInfo { FileName = address, UseShellExecute = true });
+			}
+		}
+
+		static LogEventLevel GetRequestLoggingLevel(HttpContext context, double elapsedMs, Exception ex)
+		{
+			if (context.Request != null && context.Request.Path.HasValue)
+			{
+				string requestPath = context.Request.Path;
+
+				if (requestPath.Equals("/Horde.HordeRpc/QueryServerStateV2", StringComparison.OrdinalIgnoreCase))
+				{
+					return LogEventLevel.Verbose;
+				}
+				if (requestPath.Equals("/Horde.HordeRpc/UpdateSession", StringComparison.OrdinalIgnoreCase))
+				{
+					return LogEventLevel.Verbose;
+				}
+				if (requestPath.Equals("/Horde.HordeRpc/CreateEvents", StringComparison.OrdinalIgnoreCase))
+				{
+					return LogEventLevel.Verbose;
+				}
+				if (requestPath.Equals("/Horde.HordeRpc/WriteOutput", StringComparison.OrdinalIgnoreCase))
+				{
+					return LogEventLevel.Information;
+				}
+				if (requestPath.StartsWith("/Horde.HordeRpc", StringComparison.OrdinalIgnoreCase))
+				{
+					return LogEventLevel.Debug;
+				}
+
+				if (requestPath.StartsWith("/health", StringComparison.OrdinalIgnoreCase))
+				{
+					return LogEventLevel.Debug;
+				}
+
+				if (requestPath.StartsWith("/grpc.health", StringComparison.OrdinalIgnoreCase))
+				{
+					return LogEventLevel.Debug;
+				}
+
+				if (requestPath.StartsWith("/api/v1", StringComparison.OrdinalIgnoreCase))
+				{
+					return LogEventLevel.Debug;
+				}
+
+				if (requestPath.StartsWith("/ugs/api", StringComparison.OrdinalIgnoreCase))
+				{
+					return LogEventLevel.Verbose;
+				}
+			}
+			return LogEventLevel.Information;
+		}
+
+		class HostApplicationLifetime : IHostApplicationLifetime
+		{
+			public CancellationToken ApplicationStarted => CancellationToken.None;
+			public CancellationToken ApplicationStopped => CancellationToken.None;
+			public CancellationToken ApplicationStopping => CancellationToken.None;
+
+			public void StopApplication() { }
+		}
+
+		public static void AddServices(IServiceCollection serviceCollection, IConfiguration configuration)
+		{
+			Startup startup = new Startup(configuration);
+			startup.ConfigureServices(serviceCollection);
+			serviceCollection.AddSingleton<IHostApplicationLifetime, HostApplicationLifetime>();
+		}
+
+		public static IServiceProvider CreateServiceProvider(IConfiguration configuration)
+		{
+			IServiceCollection serviceCollection = new ServiceCollection();
+			AddServices(serviceCollection, configuration);
+			return serviceCollection.BuildServiceProvider();
+		}
+
+		public static ServiceProvider CreateServiceProvider(IConfiguration configuration, ILoggerProvider loggerProvider)
+		{
+			IServiceCollection services = new ServiceCollection();
+			services.AddSingleton(configuration);
+			services.AddSingleton(loggerProvider);
+			AddServices(services, configuration);
+			return services.BuildServiceProvider();
+		}
+	}
+}
diff --git a/Engine/Source/Programs/Horde/Horde.Server/Storage/Backends/AwsStorageBackend.cs b/Engine/Source/Programs/Horde/Horde.Server/Storage/Backends/AwsStorageBackend.cs
new file mode 100644
--- /dev/null
+++ b/Engine/Source/Programs/Horde/Horde.Server/Storage/Backends/AwsStorageBackend.cs
@@ -0,0 +1,547 @@
+// Copyright Epic Games, Inc. All Rights Reserved.
+
+using System;
+using System.Collections.Generic;
+using System.Data.SqlTypes;
+using System.Globalization;
+using System.IO;
+using System.Runtime.CompilerServices;
+using System.Threading;
+using System.Threading.Tasks;
+using Amazon;
+using Amazon.Extensions.NETCore.Setup;
+using Amazon.Runtime;
+using Amazon.S3;
+using Amazon.S3.Model;
+using Google.Protobuf.WellKnownTypes;
+using Horde.Server.Utilities;
+using Microsoft.Extensions.Configuration;
+using Microsoft.Extensions.Logging;
+using OpenTracing;
+using OpenTracing.Util;
+
+namespace Horde.Server.Storage.Backends
+{
+	/// <summary>
+	/// Exception wrapper for S3 requests
+	/// </summary>
+	public sealed class AwsException : Exception
+	{
+		/// <summary>
+		/// Constructor
+		/// </summary>
+		/// <param name="message">Message for the exception</param>
+		/// <param name="innerException">Inner exception data</param>
+		public AwsException(string? message, Exception? innerException)
+			: base(message, innerException)
+		{
+		}
+	}
+
+	/// <summary>
+	/// Credentials to use for AWS
+	/// </summary>
+	public enum AwsCredentialsType
+	{
+		/// <summary>
+		/// Use default credentials from the AWS SDK
+		/// </summary>
+		Default,
+
+		/// <summary>
+		/// Read credentials from the <see cref="IAwsStorageOptions.AwsProfile"/> profile in the AWS config file
+		/// </summary>
+		Profile,
+
+		/// <summary>
+		/// Assume a particular role. Should specify ARN in <see cref="IAwsStorageOptions.AwsRole"/>
+		/// </summary>
+		AssumeRole,
+
+		/// <summary>
+		/// Assume a particular role using the current environment variables.
+		/// </summary>
+		AssumeRoleWebIdentity,
+	}
+
+	/// <summary>
+	/// Options for AWS
+	/// </summary>
+	public interface IAwsStorageOptions
+	{
+		/// <summary>
+		/// Type of credentials to use
+		/// </summary>
+		public AwsCredentialsType? AwsCredentials { get; }
+
+		/// <summary>
+		/// Name of the bucket to use
+		/// </summary>
+		public string? AwsBucketName { get; }
+
+		/// <summary>
+		/// Base path within the bucket 
+		/// </summary>
+		public string? AwsBucketPath { get; }
+
+		/// <summary>
+		/// ARN of a role to assume
+		/// </summary>
+		public string? AwsRole { get; }
+
+		/// <summary>
+		/// The AWS profile to read credentials form
+		/// </summary>
+		public string? AwsProfile { get; }
+
+		/// <summary>
+		/// Region to connect to
+		/// </summary>
+		public string? AwsRegion { get; }
+	}
+
+	/// <summary>
+	/// FileStorage implementation using an s3 bucket
+	/// </summary>
+	public sealed class AwsStorageBackend : IStorageBackendWithRedirects, IDisposable
+	{
+		/// <summary>
+		/// S3 Client
+		/// </summary>
+		private readonly IAmazonS3 _client;
+
+		/// <summary>
+		/// Options for AWs
+		/// </summary>
+		private readonly IAwsStorageOptions _options;
+
+		/// <summary>
+		/// Semaphore for connecting to AWS
+		/// </summary>
+		private readonly SemaphoreSlim _semaphore;
+
+		/// <summary>
+		/// Prefix for objects in the bucket
+		/// </summary>
+		private readonly string _pathPrefix;
+
+		/// <summary>
+		/// Logger interface
+		/// </summary>
+		private readonly ILogger _logger;
+
+		/// <summary>
+		/// Constructor
+		/// </summary>
+		/// <param name="configuration">Global configuration object</param>
+		/// <param name="options">Storage options</param>
+		/// <param name="logger">Logger interface</param>
+		public AwsStorageBackend(IConfiguration configuration, IAwsStorageOptions options, ILogger<AwsStorageBackend> logger)
+		{
+			AWSOptions awsOptions = GetAwsOptions(configuration, options);
+
+			_client = awsOptions.CreateServiceClient<IAmazonS3>();
+			_options = options;
+			_semaphore = new SemaphoreSlim(16);
+			_logger = logger;
+
+			_pathPrefix = (_options.AwsBucketPath ?? String.Empty).TrimEnd('/');
+			if (_pathPrefix.Length > 0)
+			{
+				_pathPrefix += '/';
+			}
+
+			logger.LogInformation("Created AWS storage backend for bucket {BucketName} using credentials {Credentials} {CredentialsStr}", options.AwsBucketName, awsOptions.Credentials.GetType(), awsOptions.Credentials.ToString());
+		}
+
+		/// <summary>
+		/// Gets the AWS options 
+		/// </summary>
+		/// <param name="configuration">Global configuration object</param>
+		/// <param name="options">AWS storage options</param>
+		/// <returns></returns>
+		static AWSOptions GetAwsOptions(IConfiguration configuration, IAwsStorageOptions options)
+		{
+			AWSOptions awsOptions = configuration.GetAWSOptions();
+			if (options.AwsRegion != null)
+			{
+				awsOptions.Region = RegionEndpoint.GetBySystemName(options.AwsRegion);
+			}
+
+			switch (options.AwsCredentials ?? AwsCredentialsType.Default)
+			{
+				case AwsCredentialsType.Default:
+					// Using the fallback credentials from the AWS SDK, it will pick up credentials through a number of default mechanisms.
+					awsOptions.Credentials = FallbackCredentialsFactory.GetCredentials();
+					break;
+				case AwsCredentialsType.Profile:
+					if (options.AwsProfile == null)
+					{
+						throw new AwsException($"Missing {nameof(IAwsStorageOptions.AwsProfile)} setting for configuring {nameof(AwsStorageBackend)}", null);
+					}
+
+					(string accessKey, string secretAccessKey, string secretToken) = AwsHelper.ReadAwsCredentials(options.AwsProfile);
+					awsOptions.Credentials = new Amazon.SecurityToken.Model.Credentials(accessKey, secretAccessKey, secretToken, DateTime.Now + TimeSpan.FromHours(12));
+					break;
+				case AwsCredentialsType.AssumeRole:
+					if(options.AwsRole == null)
+					{
+						throw new AwsException($"Missing {nameof(IAwsStorageOptions.AwsRole)} setting for configuring {nameof(AwsStorageBackend)}", null);
+					}
+					awsOptions.Credentials = new AssumeRoleAWSCredentials(FallbackCredentialsFactory.GetCredentials(), options.AwsRole, "Horde");
+					break;
+				case AwsCredentialsType.AssumeRoleWebIdentity:
+					awsOptions.Credentials = AssumeRoleWithWebIdentityCredentials.FromEnvironmentVariables();
+					break;
+			}
+			return awsOptions;
+		}
+
+		/// <inheritdoc/>
+		public void Dispose()
+		{
+			_client.Dispose();
+			_semaphore.Dispose();
+		}
+
+		class WrappedResponseStream : Stream
+		{
+			readonly IDisposable _semaphore;
+			readonly GetObjectResponse _response;
+			readonly Stream _responseStream;
+
+			public WrappedResponseStream(IDisposable semaphore, GetObjectResponse response)
+			{
+				_semaphore = semaphore;
+				_response = response;
+				_responseStream = response.ResponseStream;
+			}
+
+			public override bool CanRead => true;
+			public override bool CanSeek => false;
+			public override bool CanWrite => false;
+			public override long Length => _responseStream.Length;
+			public override long Position { get => _responseStream.Position; set => throw new NotSupportedException(); }
+
+			public override void Flush() => _responseStream.Flush();
+
+			public override int Read(byte[] buffer, int offset, int count) => _responseStream.Read(buffer, offset, count);
+			public override ValueTask<int> ReadAsync(Memory<byte> buffer, CancellationToken cancellationToken = default) => _responseStream.ReadAsync(buffer, cancellationToken);
+			public override Task<int> ReadAsync(byte[] buffer, int offset, int count, CancellationToken cancellationToken) => _responseStream.ReadAsync(buffer, offset, count, cancellationToken);
+
+			public override long Seek(long offset, SeekOrigin origin) => throw new NotSupportedException();
+			public override void SetLength(long value) => throw new NotSupportedException();
+			public override void Write(byte[] buffer, int offset, int count) => throw new NotSupportedException();
+
+			protected override void Dispose(bool disposing)
+			{
+				base.Dispose(disposing);
+
+				_semaphore.Dispose();
+				_response.Dispose();
+				_responseStream.Dispose();
+			}
+
+			public override async ValueTask DisposeAsync()
+			{
+				await base.DisposeAsync();
+
+				_semaphore.Dispose();
+				_response.Dispose();
+				await _responseStream.DisposeAsync();
+			}
+		}
+
+		string GetFullPath(string path) => _pathPrefix + path;
+
+		/// <inheritdoc/>
+		public Task<Stream?> TryReadAsync(string path, CancellationToken cancellationToken)
+		{
+			return TryReadAsync(path, null, cancellationToken);
+		}
+
+		/// <inheritdoc/>
+		public Task<Stream?> TryReadAsync(string path, int offset, int length, CancellationToken cancellationToken)
+		{
+			return TryReadAsync(path, new ByteRange(offset, offset + length), cancellationToken);
+		}
+
+		async Task<Stream?> TryReadAsync(string path, ByteRange? byteRange, CancellationToken cancellationToken)
+		{
+			using IScope scope = GlobalTracer.Instance.BuildSpan("AwsStorageBackend.ReadAsync").StartActive();
+			scope.Span.SetTag("Path", path);
+			
+			string fullPath = GetFullPath(path);
+
+			IDisposable? semaLock = null;
+			GetObjectResponse? response = null;
+			try
+			{
+				semaLock = await _semaphore.UseWaitAsync(cancellationToken);
+
+				GetObjectRequest newGetRequest = new GetObjectRequest();
+				newGetRequest.BucketName = _options.AwsBucketName;
+				newGetRequest.Key = fullPath;
+				newGetRequest.ByteRange = byteRange;
+
+				response = await _client.GetObjectAsync(newGetRequest, cancellationToken);
+
+				return new WrappedResponseStream(semaLock, response);
+			}
+			catch (Exception ex)
+			{
+				_logger.LogWarning(ex, "Unable to read {Path} from S3", fullPath);
+
+				semaLock?.Dispose();
+				response?.Dispose();
+
+				return null;
+			}
+		}
+
+		/// <inheritdoc/>
+		public ValueTask<Uri?> GetReadRedirectAsync(string path, CancellationToken cancellationToken = default) => new ValueTask<Uri?>(GetPresignedUrl(path, HttpVerb.GET));
+
+		/// <inheritdoc/>
+		public ValueTask<Uri?> GetWriteRedirectAsync(string path, CancellationToken cancellationToken = default) => new ValueTask<Uri?>(GetPresignedUrl(path, HttpVerb.PUT));
+
+		/// <summary>
+		/// Helper method to generate a presigned URL for a request
+		/// </summary>
+		Uri? GetPresignedUrl(string path, HttpVerb verb)
+		{
+			using IScope scope = GlobalTracer.Instance.BuildSpan("AwsStorageBackend.TryGetHttpRedirectUrl").StartActive();
+			scope.Span.SetTag("Path", path);
+
+			string fullPath = GetFullPath(path);
+
+			try
+			{
+				GetPreSignedUrlRequest newGetRequest = new GetPreSignedUrlRequest();
+				newGetRequest.BucketName = _options.AwsBucketName;
+				newGetRequest.Key = fullPath;
+				newGetRequest.Verb = verb;
+				newGetRequest.Expires = DateTime.UtcNow.AddHours(3.0);
+
+				string url = _client.GetPreSignedURL(newGetRequest);
+				return new Uri(url);
+			}
+			catch (Exception ex)
+			{
+				_logger.LogWarning(ex, "Unable to get presigned url for {Path} from S3", fullPath);
+				return null;
+			}
+		}
+
+		/// <inheritdoc/>
+		public async Task WriteAsync(string path, Stream inputStream, CancellationToken cancellationToken)
+		{
+			TimeSpan[] retryTimes =
+			{
+				TimeSpan.FromSeconds(1.0),
+				TimeSpan.FromSeconds(5.0),
+				TimeSpan.FromSeconds(10.0),
+			};
+
+			string fullPath = GetFullPath(path);
+			for (int attempt = 0; ; attempt++)
+			{
+				try
+				{
+					using IDisposable semaLock = await _semaphore.UseWaitAsync(cancellationToken);
+					await WriteInternalAsync(fullPath, inputStream, cancellationToken);
+					_logger.LogDebug("Written data to {Path}", path);
+					break;
+				}
+				catch (Exception ex)
+				{
+					_logger.LogError(ex, "Unable to write data to {Path} ({Attempt}/{AttemptCount})", fullPath, attempt + 1, retryTimes.Length + 1);
+					if (attempt >= retryTimes.Length)
+					{
+						throw new AwsException($"Unable to write to bucket {_options.AwsBucketName} path {fullPath}", ex);
+					}
+				}
+
+				await Task.Delay(retryTimes[attempt], cancellationToken);
+			}
+		}
+
+		/// <inheritdoc/>
+		public async Task WriteInternalAsync(string fullPath, Stream stream, CancellationToken cancellationToken)
+		{
+			using IScope scope = GlobalTracer.Instance.BuildSpan("AwsStorageBackend.WriteInternalAsync").StartActive();
+			scope.Span.SetTag("Path", fullPath);
+			
+			const int MinPartSize = 5 * 1024 * 1024;
+
+			long streamLen = stream.Length;
+			if (streamLen < MinPartSize)
+			{
+				// Read the data into memory. Errors with hash value not matching if we don't do this (?)
+				byte[] buffer = new byte[streamLen];
+				await ReadExactLengthAsync(stream, buffer, (int)streamLen, cancellationToken);
+
+				// Upload it to S3
+				using (MemoryStream inputStream = new MemoryStream(buffer))
+				{
+					PutObjectRequest uploadRequest = new PutObjectRequest();
+					uploadRequest.BucketName = _options.AwsBucketName;
+					uploadRequest.Key = fullPath;
+					uploadRequest.InputStream = inputStream;
+					uploadRequest.Metadata.Add("bytes-written", streamLen.ToString(CultureInfo.InvariantCulture));
+					await _client.PutObjectAsync(uploadRequest, cancellationToken);
+				}
+			}
+			else
+			{
+				// Initiate a multi-part upload
+				InitiateMultipartUploadRequest initiateRequest = new InitiateMultipartUploadRequest();
+				initiateRequest.BucketName = _options.AwsBucketName;
+				initiateRequest.Key = fullPath;
+
+				InitiateMultipartUploadResponse initiateResponse = await _client.InitiateMultipartUploadAsync(initiateRequest, cancellationToken);
+				try
+				{
+					// Buffer for reading the data
+					byte[] buffer = new byte[MinPartSize];
+
+					// Upload all the parts
+					List<PartETag> partTags = new List<PartETag>();
+					for (long streamPos = 0; streamPos < streamLen;)
+					{
+						// Read the next chunk of data into the buffer
+						int bufferLen = (int)Math.Min((long)MinPartSize, streamLen - streamPos);
+						await ReadExactLengthAsync(stream, buffer, bufferLen, cancellationToken);
+						streamPos += bufferLen;
+
+						// Upload the part
+						using (MemoryStream inputStream = new MemoryStream(buffer, 0, bufferLen))
+						{
+							UploadPartRequest partRequest = new UploadPartRequest();
+							partRequest.BucketName = _options.AwsBucketName;
+							partRequest.Key = fullPath;
+							partRequest.UploadId = initiateResponse.UploadId;
+							partRequest.InputStream = inputStream;
+							partRequest.PartSize = bufferLen;
+							partRequest.PartNumber = partTags.Count + 1;
+							partRequest.IsLastPart = (streamPos == streamLen);
+
+							UploadPartResponse partResponse = await _client.UploadPartAsync(partRequest, cancellationToken);
+							partTags.Add(new PartETag(partResponse.PartNumber, partResponse.ETag));
+						}
+					}
+
+					// Mark the upload as complete
+					CompleteMultipartUploadRequest completeRequest = new CompleteMultipartUploadRequest();
+					completeRequest.BucketName = _options.AwsBucketName;
+					completeRequest.Key = fullPath;
+					completeRequest.UploadId = initiateResponse.UploadId;
+					completeRequest.PartETags = partTags;
+					await _client.CompleteMultipartUploadAsync(completeRequest, cancellationToken);
+				}
+				catch
+				{
+					// Abort the upload
+					AbortMultipartUploadRequest abortRequest = new AbortMultipartUploadRequest();
+					abortRequest.BucketName = _options.AwsBucketName;
+					abortRequest.Key = fullPath;
+					abortRequest.UploadId = initiateResponse.UploadId;
+					await _client.AbortMultipartUploadAsync(abortRequest, cancellationToken);
+
+					throw;
+				}
+			}
+		}
+
+		/// <summary>
+		/// Reads data of an exact length into a stream
+		/// </summary>
+		/// <param name="stream">The stream to read from</param>
+		/// <param name="buffer">The buffer to read into</param>
+		/// <param name="length">Length of the data to read</param>
+		/// <param name="cancellationToken">Cancellation token for the operation</param>
+		/// <returns>Async task</returns>
+		static async Task ReadExactLengthAsync(System.IO.Stream stream, byte[] buffer, int length, CancellationToken cancellationToken)
+		{
+			int bufferPos = 0;
+			while (bufferPos < length)
+			{
+				int bytesRead = await stream.ReadAsync(buffer, bufferPos, length - bufferPos, cancellationToken);
+				if (bytesRead == 0)
+				{
+					throw new InvalidOperationException("Unexpected end of stream");
+				}
+				bufferPos += bytesRead;
+			}
+		}
+
+		/// <inheritdoc/>
+		public async Task DeleteAsync(string path, CancellationToken cancellationToken)
+		{
+			using IScope scope = GlobalTracer.Instance.BuildSpan("AwsStorageBackend.DeleteAsync").StartActive();
+			scope.Span.SetTag("Path", path);
+			
+			DeleteObjectRequest newDeleteRequest = new DeleteObjectRequest();
+			newDeleteRequest.BucketName = _options.AwsBucketName;
+			newDeleteRequest.Key = GetFullPath(path);
+			await _client.DeleteObjectAsync(newDeleteRequest, cancellationToken);
+		}
+
+		/// <inheritdoc/>
+		public async Task<bool> ExistsAsync(string path, CancellationToken cancellationToken)
+		{
+			using IScope scope = GlobalTracer.Instance.BuildSpan("AwsStorageBackend.ExistsAsync").StartActive();
+			scope.Span.SetTag("Path", path);
+			
+			try
+			{
+				GetObjectMetadataRequest request = new GetObjectMetadataRequest();
+				request.BucketName = _options.AwsBucketName;
+				request.Key = GetFullPath(path);
+				await _client.GetObjectMetadataAsync(request, cancellationToken);
+				return true;
+			}
+			catch (AmazonS3Exception ex) when (ex.StatusCode == System.Net.HttpStatusCode.NotFound)
+			{
+				return false;
+			}
+		}
+
+		/// <inheritdoc/>
+		public async IAsyncEnumerable<string> EnumerateAsync([EnumeratorCancellation] CancellationToken cancellationToken = default)
+		{
+			ListObjectsV2Request request = new ListObjectsV2Request();
+			request.BucketName = _options.AwsBucketName;
+			if(_pathPrefix.Length > 0)
+			{
+				request.Prefix = _pathPrefix;
+			}
+
+			for (; ; )
+			{
+				ListObjectsV2Response response = await _client.ListObjectsV2Async(request, cancellationToken);
+				foreach (S3Object obj in response.S3Objects)
+				{
+					string path = obj.Key;
+					if (path.StartsWith(_pathPrefix, StringComparison.Ordinal))
+					{
+						yield return path.Substring(_pathPrefix.Length);
+					}
+					else
+					{
+						_logger.LogError("Unexpected object enumerated from {AwsBucketName} - expected object \"{Path}\" to start with \"{AwsBucketPath}\"", _options.AwsBucketName, path, _pathPrefix);
+					}
+				}
+
+				if (response.IsTruncated)
+				{
+					request.ContinuationToken = response.NextContinuationToken;
+				}
+				else
+				{
+					break;
+				}
+			}
+		}
+	}
+}
diff --git a/Engine/Source/Programs/Horde/Horde.Server/Streams/IStream.cs b/Engine/Source/Programs/Horde/Horde.Server/Streams/IStream.cs
new file mode 100644
--- /dev/null
+++ b/Engine/Source/Programs/Horde/Horde.Server/Streams/IStream.cs
@@ -0,0 +1,195 @@
+// Copyright Epic Games, Inc. All Rights Reserved.
+
+using System;
+using System.Collections.Generic;
+using System.Diagnostics.CodeAnalysis;
+using System.Linq;
+using System.Text.RegularExpressions;
+using EpicGames.Core;
+using Horde.Server.Acls;
+using Horde.Server.Agents;
+using Horde.Server.Jobs;
+using Horde.Server.Perforce;
+using Horde.Server.Projects;
+using Horde.Server.Users;
+using Horde.Server.Utilities;
+
+namespace Horde.Server.Streams
+{
+	/// <summary>
+	/// Exception thrown when stream validation fails
+	/// </summary>
+	public class InvalidStreamException : Exception
+	{
+		/// <inheritdoc/>
+		public InvalidStreamException()
+		{
+		}
+
+		/// <inheritdoc/>
+		public InvalidStreamException(string message) : base(message)
+		{
+		}
+
+		/// <inheritdoc/>
+		public InvalidStreamException(string message, Exception innerEx) : base(message, innerEx)
+		{
+		}
+	}
+
+	/// <summary>
+	/// Information about a stream
+	/// </summary>
+	public interface IStream
+	{
+		/// <summary>
+		/// Name of the stream.
+		/// </summary>
+		StreamId Id { get; }
+
+		/// <summary>
+		/// Configuration settings for the stream
+		/// </summary>
+		StreamConfig Config { get; }
+
+		/// <summary>
+		/// List of templates available for this stream
+		/// </summary>
+		IReadOnlyDictionary<TemplateId, ITemplateRef> Templates { get; }
+
+		/// <summary>
+		/// Stream is paused for builds until specified time
+		/// </summary>
+		DateTime? PausedUntil { get; }
+		
+		/// <summary>
+		/// Comment/reason for why the stream was paused
+		/// </summary>
+		string? PauseComment { get; }
+	}
+
+	/// <summary>
+	/// Job template in a stream
+	/// </summary>
+	public interface ITemplateRef
+	{
+		/// <summary>
+		/// The template id
+		/// </summary>
+		TemplateId Id { get; }
+
+		/// <summary>
+		/// Configuration of this template ref
+		/// </summary>
+		TemplateRefConfig Config { get; }
+
+		/// <summary>
+		/// List of schedules for this template
+		/// </summary>
+		ITemplateSchedule? Schedule { get; }
+
+		/// <summary>
+		/// List of template step states
+		/// </summary>
+		IReadOnlyList<ITemplateStep> StepStates { get; }
+	}
+
+	/// <summary>
+	/// Schedule for a template
+	/// </summary>
+	public interface ITemplateSchedule
+	{
+		/// <summary>
+		/// Config for this schedule
+		/// </summary>
+		ScheduleConfig Config { get; }
+
+		/// <summary>
+		/// Last changelist number that this was triggered for
+		/// </summary>
+		int LastTriggerChange { get; }
+
+		/// <summary>
+		/// Gets the last trigger time, in UTC
+		/// </summary>
+		DateTime LastTriggerTimeUtc { get; }
+
+		/// <summary>
+		/// List of jobs that are currently active
+		/// </summary>
+		IReadOnlyList<JobId> ActiveJobs { get; }
+	}
+
+	/// <summary>
+	/// Information about a paused template step
+	/// </summary>
+	public interface ITemplateStep
+	{
+		/// <summary>
+		/// Name of the step
+		/// </summary>
+		string Name { get; }
+
+		/// <summary>
+		/// User who paused the step
+		/// </summary>
+		UserId PausedByUserId { get; }
+
+		/// <summary>
+		/// The UTC time when the step was paused
+		/// </summary>
+		DateTime PauseTimeUtc { get; }
+	}
+
+	/// <summary>
+	/// Extension methods for streams
+	/// </summary>
+	static class StreamExtensions
+	{
+		/// <summary>
+		/// Gets the next trigger time for a schedule
+		/// </summary>
+		/// <param name="schedule"></param>
+		/// <param name="timeZone"></param>
+		/// <returns></returns>
+		public static DateTime? GetNextTriggerTimeUtc(this ITemplateSchedule schedule, TimeZoneInfo timeZone)
+		{
+			return schedule.Config.GetNextTriggerTimeUtc(schedule.LastTriggerTimeUtc, timeZone);
+		}
+
+		/// <summary>
+		/// Converts to a public response object
+		/// </summary>
+		/// <param name="stream">The stream object</param>
+		/// <param name="apiTemplateRefs">The template refs for this stream. Passed separately because they have their own ACL.</param>
+		/// <returns>New response instance</returns>
+		public static GetStreamResponse ToApiResponse(this IStream stream, List<GetTemplateRefResponse> apiTemplateRefs)
+		{
+			return new GetStreamResponse(stream, apiTemplateRefs);
+		}
+
+		/// <summary>
+		/// Converts to an RPC response object
+		/// </summary>
+		/// <param name="streamConfig">The stream config object</param>
+		/// <returns>New response instance</returns>
+		public static HordeCommon.Rpc.GetStreamResponse ToRpcResponse(this StreamConfig streamConfig)
+		{
+			HordeCommon.Rpc.GetStreamResponse response = new HordeCommon.Rpc.GetStreamResponse();
+			response.Name = streamConfig.Name;
+			response.AgentTypes.Add(streamConfig.AgentTypes.ToDictionary(x => x.Key, x => x.Value.ToRpcResponse()));
+			return response;
+		}
+
+		/// <summary>
+		/// Check if stream is paused for new builds
+		/// </summary>
+		/// <param name="stream">The stream object</param>
+		/// <param name="currentTime">Current time (allow tests to pass in a fake clock)</param>
+		/// <returns>If stream is paused</returns>
+		public static bool IsPaused(this IStream stream, DateTime currentTime)
+		{
+			return stream.PausedUntil != null && stream.PausedUntil > currentTime;
+		}
+	}
+}
diff --git a/Engine/Source/Programs/Horde/Horde.Server/Streams/IStreamCollection.cs b/Engine/Source/Programs/Horde/Horde.Server/Streams/IStreamCollection.cs
new file mode 100644
--- /dev/null
+++ b/Engine/Source/Programs/Horde/Horde.Server/Streams/IStreamCollection.cs
@@ -0,0 +1,135 @@
+// Copyright Epic Games, Inc. All Rights Reserved.
+
+using System;
+using System.Collections.Generic;
+using System.Linq;
+using System.Threading.Tasks;
+using Horde.Server.Jobs;
+using Horde.Server.Utilities;
+
+namespace Horde.Server.Streams
+{
+	/// <summary>
+	/// Collection of stream documents
+	/// </summary>
+	public interface IStreamCollection
+	{
+		/// <summary>
+		/// Gets a stream by ID
+		/// </summary>
+		/// <param name="streamConfig">The stream config object</param>
+		/// <returns>The stream document</returns>
+		Task<IStream> GetAsync(StreamConfig streamConfig);
+
+		/// <summary>
+		/// Gets a stream by ID
+		/// </summary>
+		/// <param name="streamConfigs">The stream config object</param>
+		/// <returns>The stream document</returns>
+		Task<List<IStream>> GetAsync(IReadOnlyList<StreamConfig> streamConfigs);
+
+		/// <summary>
+		/// Updates user-facing properties for an existing stream
+		/// </summary>
+		/// <param name="stream">The stream to update</param>
+		/// <param name="newPausedUntil">The new datetime for pausing builds</param>
+		/// <param name="newPauseComment">The reason for pausing</param>
+		/// <returns>The updated stream if successful, null otherwise</returns>
+		Task<IStream?> TryUpdatePauseStateAsync(IStream stream, DateTime? newPausedUntil, string? newPauseComment);
+
+		/// <summary>
+		/// Attempts to update the last trigger time for a schedule
+		/// </summary>
+		/// <param name="stream">The stream to update</param>
+		/// <param name="templateRefId">The template ref id</param>
+		/// <param name="lastTriggerTimeUtc">New last trigger time for the schedule</param>
+		/// <param name="lastTriggerChange">New last trigger changelist for the schedule</param>
+		/// <param name="newActiveJobs">New list of active jobs</param>
+		/// <returns>The updated stream if successful, null otherwise</returns>
+		Task<IStream?> TryUpdateScheduleTriggerAsync(IStream stream, TemplateId templateRefId, DateTime? lastTriggerTimeUtc, int? lastTriggerChange, List<JobId> newActiveJobs);
+
+		/// <summary>
+		/// Attempts to update a stream template ref
+		/// </summary>
+		/// <param name="streamInterface">The stream containing the template ref</param>
+		/// <param name="templateRefId">The template ref to update</param>
+		/// <param name="stepStates">The stream states to update, pass an empty list to clear all step states, otherwise will be a partial update based on included step updates</param>
+		/// <returns></returns>
+		Task<IStream?> TryUpdateTemplateRefAsync(IStream streamInterface, TemplateId templateRefId, List<UpdateStepStateRequest>? stepStates = null);
+	}
+
+	/// <summary>
+	/// Extension methods for <see cref="StreamCollection"/>
+	/// </summary>
+	public static class StreamCollectionExtensions
+	{
+		/// <summary>
+		/// Updates an existing stream
+		/// </summary>
+		/// <param name="streamCollection">The stream collection</param>
+		/// <param name="stream">The stream to update</param>
+		/// <param name="newPausedUntil">The new datetime for pausing builds</param>
+		/// <param name="newPauseComment">The reason for pausing</param>
+		/// <returns>Async task object</returns>
+		public static async Task<IStream?> UpdatePauseStateAsync(this IStreamCollection streamCollection, IStream? stream, DateTime? newPausedUntil = null, string? newPauseComment = null)
+		{
+			for (; stream != null; stream = await streamCollection.GetAsync(stream.Config))
+			{
+				IStream? newStream = await streamCollection.TryUpdatePauseStateAsync(stream, newPausedUntil, newPauseComment);
+				if (newStream != null)
+				{
+					return newStream;
+				}
+			}
+			return null;
+		}
+
+		/// <summary>
+		/// Attempts to update the last trigger time for a schedule
+		/// </summary>
+		/// <param name="streamCollection">The stream collection</param>
+		/// <param name="stream">The stream to update</param>
+		/// <param name="templateRefId">The template ref id</param>
+		/// <param name="lastTriggerTimeUtc"></param>
+		/// <param name="lastTriggerChange"></param>
+		/// <param name="addJobs">Jobs to add</param>
+		/// <param name="removeJobs">Jobs to remove</param>
+		/// <returns>True if the stream was updated</returns>
+		public static async Task<IStream?> UpdateScheduleTriggerAsync(this IStreamCollection streamCollection, IStream stream, TemplateId templateRefId, DateTime? lastTriggerTimeUtc = null, int? lastTriggerChange = null, List<JobId>? addJobs = null, List<JobId>? removeJobs = null)
+		{
+			IStream? newStream = stream;
+			while (newStream != null)
+			{
+				ITemplateRef? templateRef;
+				if (!newStream.Templates.TryGetValue(templateRefId, out templateRef))
+				{
+					break;
+				}
+				if (templateRef.Schedule == null)
+				{
+					break;
+				}
+
+				IEnumerable<JobId> newActiveJobs = templateRef.Schedule.ActiveJobs;
+				if (removeJobs != null)
+				{
+					newActiveJobs = newActiveJobs.Except(removeJobs);
+				}
+				if (addJobs != null)
+				{
+					newActiveJobs = newActiveJobs.Union(addJobs);
+				}
+
+				newStream = await streamCollection.TryUpdateScheduleTriggerAsync(newStream, templateRefId, lastTriggerTimeUtc, lastTriggerChange, newActiveJobs.ToList());
+
+				if (newStream != null)
+				{
+					return newStream;
+				}
+
+				newStream = await streamCollection.GetAsync(stream.Config);
+			}
+			return null;
+		}
+	}
+}
diff --git a/Engine/Source/Programs/Horde/Horde.Server/Streams/StreamCollection.cs b/Engine/Source/Programs/Horde/Horde.Server/Streams/StreamCollection.cs
new file mode 100644
--- /dev/null
+++ b/Engine/Source/Programs/Horde/Horde.Server/Streams/StreamCollection.cs
@@ -0,0 +1,502 @@
+// Copyright Epic Games, Inc. All Rights Reserved.
+
+using System;
+using System.Collections.Generic;
+using System.IO;
+using System.Linq;
+using System.Threading.Tasks;
+using EpicGames.Core;
+using Horde.Server.Jobs;
+using Horde.Server.Projects;
+using Horde.Server.Server;
+using Horde.Server.Users;
+using Horde.Server.Utilities;
+using HordeCommon;
+using MongoDB.Bson.Serialization.Attributes;
+using MongoDB.Driver;
+
+namespace Horde.Server.Streams
+{
+	/// <summary>
+	/// Collection of stream documents
+	/// </summary>
+	public class StreamCollection : IStreamCollection
+	{
+		/// <summary>
+		/// Information about a stream
+		/// </summary>
+		class StreamDoc : IStream
+		{
+			[BsonRequired, BsonId]
+			public StreamId Id { get; set; }
+
+			public Dictionary<TemplateId, TemplateRefDoc> Templates { get; set; } = new Dictionary<TemplateId, TemplateRefDoc>();
+			public DateTime? PausedUntil { get; set; }
+			public string? PauseComment { get; set; }
+
+			public int UpdateIndex { get; set; }
+
+			[BsonConstructor]
+			private StreamDoc()
+			{
+			}
+
+			public StreamDoc(StreamId id)
+			{
+				Id = id;
+			}
+
+			#region IStream Implementation
+
+			[BsonIgnore]
+			public StreamConfig Config { get; private set; } = null!;
+
+			[BsonIgnore]
+			IReadOnlyDictionary<TemplateId, ITemplateRef>? _cachedTemplates;
+
+			IReadOnlyDictionary<TemplateId, ITemplateRef> IStream.Templates
+			{
+				get
+				{
+					_cachedTemplates ??= Templates.ToDictionary(x => x.Key, x => (ITemplateRef)x.Value);
+					return _cachedTemplates;
+				}
+			}
+
+			public bool PostLoad(StreamConfig config, DateTime utcNow)
+			{
+				Config = config;
+
+				bool replaceDocument = false;
+
+				Dictionary<TemplateId, TemplateRefDoc> templateRefDocs = new Dictionary<TemplateId, TemplateRefDoc>(config.Templates.Count);
+				foreach (TemplateRefConfig templateRefConfig in config.Templates)
+				{
+					TemplateRefDoc? templateRefDoc;
+					if (!Templates.TryGetValue(templateRefConfig.Id, out templateRefDoc))
+					{
+						templateRefDoc = new TemplateRefDoc();
+						replaceDocument = true;
+					}
+					replaceDocument |= templateRefDoc.PostLoad(templateRefConfig, utcNow);
+
+					templateRefDocs.Add(templateRefConfig.Id, templateRefDoc);
+				}
+
+				replaceDocument |= Templates.Count != templateRefDocs.Count;
+				Templates = templateRefDocs;
+
+				return replaceDocument;
+			}
+
+			#endregion
+		}
+
+		class TemplateRefDoc : ITemplateRef
+		{
+			[BsonIgnoreIfNull]
+			public TemplateScheduleDoc? Schedule { get; set; }
+
+			[BsonIgnoreIfNull]
+			public List<TemplateStepDoc>? StepStates { get; set; }
+
+			#region ITemplateRef implementation
+
+			[BsonIgnore]
+			public TemplateId Id { get; private set; }
+
+			[BsonIgnore]
+			public TemplateRefConfig Config { get; private set; } = null!;
+
+			ITemplateSchedule? ITemplateRef.Schedule => Schedule;
+			IReadOnlyList<ITemplateStep> ITemplateRef.StepStates => (IReadOnlyList<ITemplateStep>?)StepStates ?? Array.Empty<ITemplateStep>();
+
+			public bool PostLoad(TemplateRefConfig config, DateTime utcNow)
+			{
+				Id = config.Id;
+				Config = config;
+
+				bool replaceDocument = false;
+
+				if (Config.Schedule == null)
+				{
+					Schedule = null;
+					replaceDocument = true;
+				}
+				else if (Schedule == null)
+				{
+					Schedule = new TemplateScheduleDoc();
+					Schedule.LastTriggerTimeUtc = utcNow;
+					replaceDocument = true;
+				}
+
+				Schedule?.PostLoad(this);
+
+				if (StepStates != null)
+				{
+					int count = StepStates.RemoveAll(x => x.PausedByUserId == null);
+					replaceDocument |= count > 0;
+				}
+
+				return replaceDocument;
+			}
+
+			#endregion
+		}
+
+		class TemplateScheduleDoc : ITemplateSchedule
+		{
+			public int LastTriggerChange { get; set; }
+
+			[BsonIgnoreIfNull, Obsolete("Use LastTriggerTimeUtc instead")]
+			public DateTimeOffset? LastTriggerTime { get; set; }
+
+			public DateTime LastTriggerTimeUtc { get; set; }
+			public List<JobId> ActiveJobs { get; set; } = new List<JobId>();
+
+			#region ITemplateSchedule implementation
+
+			[BsonIgnore]
+			TemplateRefDoc? _owner;
+
+			[BsonIgnore]
+			public ScheduleConfig Config => _owner?.Config.Schedule!;
+
+			IReadOnlyList<JobId> ITemplateSchedule.ActiveJobs => ActiveJobs;
+
+			public void PostLoad(TemplateRefDoc owner)
+			{
+				_owner = owner;
+
+#pragma warning disable CS0618 // Type or member is obsolete
+				if (LastTriggerTime.HasValue)
+				{
+					LastTriggerTimeUtc = LastTriggerTime.Value.UtcDateTime;
+					LastTriggerTime = null;
+				}
+#pragma warning restore CS0618 // Type or member is obsolete
+			}
+
+			#endregion
+		}
+
+		class TemplateStepDoc : ITemplateStep
+		{
+			public string Name { get; set; } = String.Empty;
+			public UserId? PausedByUserId { get; set; }
+			public DateTime? PauseTimeUtc { get; set; }
+
+			UserId ITemplateStep.PausedByUserId => PausedByUserId ?? UserId.Empty;
+			DateTime ITemplateStep.PauseTimeUtc => PauseTimeUtc ?? DateTime.MinValue;
+
+			public TemplateStepDoc()
+			{
+			}
+
+			public TemplateStepDoc(string name, UserId pausedByUserId, DateTime pauseTimeUtc)
+			{
+				Name = name;
+				PausedByUserId = pausedByUserId;
+				PauseTimeUtc = pauseTimeUtc;
+			}
+		}
+
+		readonly IMongoCollection<StreamDoc> _streams;
+		readonly IClock _clock;
+
+		/// <summary>
+		/// Constructor
+		/// </summary>
+		/// <param name="mongoService">The database service instance</param>
+		/// <param name="clock"></param>
+		public StreamCollection(MongoService mongoService, IClock clock)
+		{
+			_streams = mongoService.GetCollection<StreamDoc>("Streams");
+			_clock = clock;
+		}
+
+		/// <inheritdoc/>
+		public async Task<IStream> GetAsync(StreamConfig streamConfig) => await GetInternalAsync(streamConfig);
+
+		async Task<StreamDoc> GetInternalAsync(StreamConfig streamConfig)
+		{
+			FindOneAndUpdateOptions<StreamDoc, StreamDoc> options = new FindOneAndUpdateOptions<StreamDoc, StreamDoc>();
+			options.IsUpsert = true;
+			options.ReturnDocument = ReturnDocument.After;
+
+			for (; ; )
+			{
+				StreamDoc stream = await _streams.FindOneAndUpdateAsync<StreamDoc>(x => x.Id == streamConfig.Id, Builders<StreamDoc>.Update.SetOnInsert(x => x.UpdateIndex, 0), options);
+				if (await PostLoadAsync(stream, streamConfig))
+				{
+					return stream;
+				}
+			}
+		}
+
+		/// <inheritdoc/>
+		public async Task<List<IStream>> GetAsync(IReadOnlyList<StreamConfig> streamConfigs) => await GetInternalAsync(streamConfigs).ConvertAllAsync<StreamDoc, IStream>();
+
+		async Task<List<StreamDoc>> GetInternalAsync(IReadOnlyList<StreamConfig> streamConfigs)
+		{
+			FilterDefinition<StreamDoc> filter = Builders<StreamDoc>.Filter.In(x => x.Id, streamConfigs.Select(x => x.Id));
+			List<StreamDoc> matches = await _streams.Find(filter).ToListAsync();
+
+			List<StreamDoc> results = new List<StreamDoc>(streamConfigs.Count);
+			foreach (StreamConfig streamConfig in streamConfigs)
+			{
+				StreamDoc? stream = matches.FirstOrDefault(x => x.Id == streamConfig.Id);
+				if (stream == null)
+				{
+					stream = await GetInternalAsync(streamConfig);
+				}
+				else
+				{
+					stream.PostLoad(streamConfig, _clock.UtcNow);
+				}
+				results.Add(stream);
+			}
+			return results;
+		}
+
+		async ValueTask<bool> PostLoadAsync(StreamDoc streamDoc, StreamConfig streamConfig)
+		{
+			if (streamDoc.PostLoad(streamConfig, _clock.UtcNow))
+			{
+				int updateIndex = streamDoc.UpdateIndex++;
+
+				ReplaceOneResult result = await _streams.ReplaceOneAsync(x => x.Id == streamDoc.Id && x.UpdateIndex == updateIndex, streamDoc);
+				if (result.MatchedCount != 1)
+				{
+					return false;
+				}
+			}
+			return true;
+		}
+
+		/// <inheritdoc/>
+		public async Task<IStream?> TryUpdateTemplateRefAsync(IStream streamInterface, TemplateId templateId, List<UpdateStepStateRequest>? stepStates = null)
+		{
+			StreamDoc stream = (StreamDoc)streamInterface;
+
+			TemplateRefDoc? templateRef;
+			if (!stream.Templates.TryGetValue(templateId, out templateRef))
+			{
+				return null;
+			}
+
+			UpdateDefinitionBuilder<StreamDoc> updateBuilder = Builders<StreamDoc>.Update;
+			List<UpdateDefinition<StreamDoc>> updates = new List<UpdateDefinition<StreamDoc>>();
+
+			Dictionary<TemplateId, TemplateRefDoc> newTemplates = new Dictionary<TemplateId, TemplateRefDoc>(stream.Templates);
+
+			// clear
+			if (stepStates != null && stepStates.Count == 0)
+			{
+				bool hasUpdates = false;
+				foreach (KeyValuePair<TemplateId, TemplateRefDoc> entry in newTemplates)
+				{
+					if (entry.Value.StepStates != null)
+					{
+						hasUpdates = true;
+						entry.Value.StepStates = null;
+					}
+				}
+
+				if (hasUpdates)
+				{
+					updates.Add(updateBuilder.Set(x => x.Templates, newTemplates));
+				}
+			}
+			else if (stepStates != null)
+			{
+				// get currently valid step states
+				List<TemplateStepDoc> newStepStates = templateRef.StepStates?.ToList() ?? new List<TemplateStepDoc>();
+
+				// generate update list
+				foreach (UpdateStepStateRequest updateState in stepStates)
+				{
+					int stateIndex = newStepStates.FindIndex(x => x.Name == updateState.Name);
+
+					UserId? pausedByUserId = updateState.PausedByUserId != null ? UserId.Parse(updateState.PausedByUserId) : null;
+
+					if (stateIndex == -1)
+					{
+						// if this is a new state without anything set, ignore it
+						if (pausedByUserId != null)
+						{
+							newStepStates.Add(new TemplateStepDoc(updateState.Name, pausedByUserId.Value, _clock.UtcNow));
+						}
+					}
+					else
+					{
+						if (pausedByUserId == null)
+						{
+							newStepStates.RemoveAt(stateIndex);
+						}
+						else
+						{
+							newStepStates[stateIndex].PausedByUserId = pausedByUserId.Value;
+						}
+					}
+				}
+
+				if (newStepStates.Count == 0)
+				{
+					templateRef.StepStates = null;
+				}
+				else
+				{
+					templateRef.StepStates = newStepStates;
+				}
+
+				updates.Add(updateBuilder.Set(x => x.Templates, newTemplates));
+			}
+
+			if (updates.Count == 0)
+			{
+				return streamInterface;
+			}
+
+			return await TryUpdateStreamAsync(stream, updateBuilder.Combine(updates));
+
+		}
+
+		/// <inheritdoc/>
+		public async Task<IStream?> TryUpdatePauseStateAsync(IStream streamInterface, DateTime? newPausedUntil, string? newPauseComment)
+		{
+			StreamDoc stream = (StreamDoc)streamInterface;
+
+			UpdateDefinitionBuilder<StreamDoc> updateBuilder = Builders<StreamDoc>.Update;
+
+			List<UpdateDefinition<StreamDoc>> updates = new List<UpdateDefinition<StreamDoc>>();
+			stream.PausedUntil = newPausedUntil;
+			stream.PauseComment = newPauseComment;
+			updates.Add(updateBuilder.Set(x => x.PausedUntil, newPausedUntil));
+			updates.Add(updateBuilder.Set(x => x.PauseComment, newPauseComment));
+
+			return await TryUpdateStreamAsync(stream, updateBuilder.Combine(updates));
+		}
+
+		/// <inheritdoc/>
+		public async Task<IStream?> TryUpdateScheduleTriggerAsync(IStream streamInterface, TemplateId templateId, DateTime? lastTriggerTimeUtc, int? lastTriggerChange, List<JobId> newActiveJobs)
+		{
+			StreamDoc stream = (StreamDoc)streamInterface;
+			TemplateRefDoc template = stream.Templates[templateId];
+			TemplateScheduleDoc schedule = template.Schedule!;
+
+			// Build the updates. MongoDB driver cannot parse TemplateRefId in expression tree; need to specify field name explicitly
+			List<UpdateDefinition<StreamDoc>> updates = new List<UpdateDefinition<StreamDoc>>();
+			if (lastTriggerTimeUtc.HasValue && lastTriggerTimeUtc.Value != schedule.LastTriggerTimeUtc)
+			{
+				FieldDefinition<StreamDoc, DateTime> lastTriggerTimeField = $"{nameof(stream.Templates)}.{templateId}.{nameof(template.Schedule)}.{nameof(schedule.LastTriggerTimeUtc)}";
+				updates.Add(Builders<StreamDoc>.Update.Set(lastTriggerTimeField, lastTriggerTimeUtc.Value));
+				schedule.LastTriggerTimeUtc = lastTriggerTimeUtc.Value;
+			}
+			if (lastTriggerChange.HasValue && lastTriggerChange.Value > schedule.LastTriggerChange)
+			{
+				FieldDefinition<StreamDoc, int> lastTriggerChangeField = $"{nameof(stream.Templates)}.{templateId}.{nameof(template.Schedule)}.{nameof(schedule.LastTriggerChange)}";
+				updates.Add(Builders<StreamDoc>.Update.Set(lastTriggerChangeField, lastTriggerChange.Value));
+				schedule.LastTriggerChange = lastTriggerChange.Value;
+			}
+			if (newActiveJobs != null)
+			{
+				FieldDefinition<StreamDoc, List<JobId>> field = $"{nameof(stream.Templates)}.{templateId}.{nameof(template.Schedule)}.{nameof(schedule.ActiveJobs)}";
+				updates.Add(Builders<StreamDoc>.Update.Set(field, newActiveJobs));
+				schedule.ActiveJobs = newActiveJobs;
+			}
+
+			return (updates.Count == 0)? streamInterface : await TryUpdateStreamAsync(stream, Builders<StreamDoc>.Update.Combine(updates));
+		}
+
+		/// <summary>
+		/// Update a stream
+		/// </summary>
+		/// <param name="stream">The stream to update</param>
+		/// <param name="update">The update definition</param>
+		/// <returns>The updated document, or null the update failed</returns>
+		private async Task<StreamDoc?> TryUpdateStreamAsync(StreamDoc stream, UpdateDefinition<StreamDoc> update)
+		{
+			FilterDefinition<StreamDoc> filter = Builders<StreamDoc>.Filter.Expr(x => x.Id == stream.Id && x.UpdateIndex == stream.UpdateIndex);
+			update = update.Set(x => x.UpdateIndex, stream.UpdateIndex + 1);
+
+			FindOneAndUpdateOptions<StreamDoc> options = new FindOneAndUpdateOptions<StreamDoc> { ReturnDocument = ReturnDocument.After };
+
+			StreamDoc? result = await _streams.FindOneAndUpdateAsync(filter, update, options);
+			if (result != null)
+			{
+				result.PostLoad(stream.Config, _clock.UtcNow);
+			}
+			return result;
+		}
+
+		/// <summary>
+		/// Checks the stream definition for consistency
+		/// </summary>
+		public static void Validate(StreamId streamId, StreamConfig config)
+		{
+			HashSet<TemplateId> remainingTemplates = new HashSet<TemplateId>(config.Templates.Select(x => x.Id));
+
+			// Check the default preflight template is valid
+			if (config.DefaultPreflight != null)
+			{
+				if (config.DefaultPreflight.TemplateId != null && !remainingTemplates.Contains(config.DefaultPreflight.TemplateId.Value))
+				{
+					throw new InvalidStreamException($"Default preflight template was listed as '{config.DefaultPreflight.TemplateId.Value}', but no template was found by that name");
+				}
+			}
+
+			// Check the chained jobs are valid
+			foreach (TemplateRefConfig templateRef in config.Templates)
+			{
+				if (templateRef.ChainedJobs != null)
+				{
+					foreach (ChainedJobTemplateConfig chainedJob in templateRef.ChainedJobs)
+					{
+						if (!remainingTemplates.Contains(chainedJob.TemplateId))
+						{
+							throw new InvalidDataException($"Invalid template ref id '{chainedJob.TemplateId}");
+						}
+					}
+				}
+			}
+
+			HashSet<TemplateId> undefinedTemplates = new();
+			// Check that all the templates are referenced by a tab
+			foreach (JobsTabConfig jobsTab in config.Tabs.OfType<JobsTabConfig>())
+			{
+				if (jobsTab.Templates != null)
+				{
+					remainingTemplates.ExceptWith(jobsTab.Templates);
+					foreach (TemplateId templateId in jobsTab.Templates)
+					{
+						if (config.Templates.Find(x => x.Id == templateId) == null)
+						{
+							undefinedTemplates.Add(templateId);
+						}
+					}
+					
+				}
+			}
+			if (remainingTemplates.Count > 0)
+			{
+				throw new InvalidStreamException(String.Join("\n", remainingTemplates.Select(x => $"Template '{x}' is not listed on any tab for {streamId}")));
+			}
+
+			if (undefinedTemplates.Count > 0)
+			{
+				throw new InvalidStreamException(String.Join("\n", undefinedTemplates.Select(x => $"Template '{x}' is not defined for {streamId}")));
+			}
+
+			// Check that all the agent types reference valid workspace names
+			foreach (KeyValuePair<string, AgentConfig> pair in config.AgentTypes)
+			{
+				string? workspaceTypeName = pair.Value.Workspace;
+				if (workspaceTypeName != null && !config.WorkspaceTypes.ContainsKey(workspaceTypeName))
+				{
+					throw new InvalidStreamException($"Agent type '{pair.Key}' references undefined workspace type '{pair.Value.Workspace}' in {streamId}");
+				}
+			}
+		}
+	}
+}
