Horde: Rename the Horde.Build project to Horde.Server.

diff --git a/Engine/Source/Programs/Horde/Horde.Server/Streams/StreamsController.cs b/Engine/Source/Programs/Horde/Horde.Server/Streams/StreamsController.cs
new file mode 100644
--- /dev/null
+++ b/Engine/Source/Programs/Horde/Horde.Server/Streams/StreamsController.cs
@@ -0,0 +1,398 @@
+// Copyright Epic Games, Inc. All Rights Reserved.
+
+using System;
+using System.Collections.Generic;
+using System.Linq;
+using System.Threading;
+using System.Threading.Tasks;
+using Horde.Server.Acls;
+using Horde.Server.Jobs;
+using Horde.Server.Jobs.Templates;
+using Horde.Server.Perforce;
+using Horde.Server.Projects;
+using Horde.Server.Server;
+using Horde.Server.Users;
+using Horde.Server.Utilities;
+using Microsoft.AspNetCore.Authorization;
+using Microsoft.AspNetCore.Mvc;
+using Microsoft.Extensions.Options;
+using OpenTracing;
+using OpenTracing.Util;
+using TimeZoneConverter;
+
+namespace Horde.Server.Streams
+{
+	/// <summary>
+	/// Controller for the /api/v1/streams endpoint
+	/// </summary>
+	[ApiController]
+	[Authorize]
+	[Route("[controller]")]
+	public class StreamsController : HordeControllerBase
+	{
+		private readonly IStreamCollection _streamCollection;
+		private readonly ICommitService _commitService;
+		private readonly ITemplateCollection _templateCollection;
+		private readonly IJobStepRefCollection _jobStepRefCollection;
+		private readonly IUserCollection _userCollection;
+		private readonly IOptionsSnapshot<GlobalConfig> _globalConfig;
+		private readonly TimeZoneInfo _timeZone;
+
+		/// <summary>
+		/// Constructor
+		/// </summary>
+		public StreamsController(IStreamCollection streamCollection, ICommitService commitService, ITemplateCollection templateCollection, IJobStepRefCollection jobStepRefCollection, IUserCollection userCollection, IOptionsSnapshot<GlobalConfig> globalConfig)
+		{
+			_streamCollection = streamCollection;
+			_commitService = commitService;
+			_templateCollection = templateCollection;
+			_jobStepRefCollection = jobStepRefCollection;
+			_userCollection = userCollection;
+			_globalConfig = globalConfig;
+
+			string? timeZoneName = _globalConfig.Value.ServerSettings.ScheduleTimeZone;
+			_timeZone = (timeZoneName == null) ? TimeZoneInfo.Local : TZConvert.GetTimeZoneInfo(timeZoneName);
+		}
+
+		/// <summary>
+		/// Query all the streams for a particular project.
+		/// </summary>
+		/// <param name="projectIds">Unique id of the project to query</param>
+		/// <param name="filter">Filter for the properties to return</param>
+		/// <returns>Information about all the projects</returns>
+		[HttpGet]
+		[Route("/api/v1/streams")]
+		[ProducesResponseType(typeof(List<GetStreamResponse>), 200)]
+		public async Task<ActionResult<List<object>>> GetStreamsAsync([FromQuery(Name = "ProjectId")] string[] projectIds, [FromQuery] PropertyFilter? filter = null)
+		{
+			ProjectId[] projectIdValues = Array.ConvertAll(projectIds, x => new ProjectId(x));
+
+			List<GetStreamResponse> responses = new List<GetStreamResponse>();
+			foreach (ProjectConfig projectConfig in _globalConfig.Value.Projects)
+			{
+				if (projectIdValues.Length == 0 || projectIdValues.Contains(projectConfig.Id))
+				{
+					foreach (StreamConfig streamConfig in projectConfig.Streams)
+					{
+						if (streamConfig.Authorize(AclAction.ViewStream, User))
+						{
+							IStream stream = await _streamCollection.GetAsync(streamConfig);
+							GetStreamResponse response = await CreateGetStreamResponseAsync(stream);
+							responses.Add(response);
+						}
+					}
+				}
+			}
+
+			return responses.OrderBy(x => x.Id).Select(x => PropertyFilter.Apply(x, filter)).ToList();
+		}
+
+		/// <summary>
+		/// Query all the streams for a particular project.
+		/// </summary>
+		/// <param name="projectIds">Unique id of the project to query</param>
+		/// <param name="filter">Filter for the properties to return</param>
+		/// <returns>Information about all the projects</returns>
+		[HttpGet]
+		[Route("/api/v2/streams")]
+		[ProducesResponseType(typeof(List<GetStreamResponse>), 200)]
+		public async Task<ActionResult<List<object>>> GetStreamsAsyncV2([FromQuery(Name = "ProjectId")] string[] projectIds, [FromQuery] PropertyFilter? filter = null)
+		{
+			ProjectId[] projectIdValues = Array.ConvertAll(projectIds, x => new ProjectId(x));
+
+			List<GetStreamResponseV2> responses = new List<GetStreamResponseV2>();
+			foreach (ProjectConfig projectConfig in _globalConfig.Value.Projects)
+			{
+				if (projectIdValues.Length == 0 || projectIdValues.Contains(projectConfig.Id))
+				{
+					foreach (StreamConfig streamConfig in projectConfig.Streams)
+					{
+						if (streamConfig.Authorize(AclAction.ViewStream, User))
+						{
+							IStream stream = await _streamCollection.GetAsync(streamConfig);
+							GetStreamResponseV2 response = await CreateGetStreamResponseV2Async(stream, false);
+							responses.Add(response);
+						}
+					}
+				}
+			}
+
+			return responses.OrderBy(x => x.Id).Select(x => PropertyFilter.Apply(x, filter)).ToList();
+		}
+
+		/// <summary>
+		/// Retrieve information about a specific stream.
+		/// </summary>
+		/// <param name="streamId">Id of the stream to get information about</param>
+		/// <param name="filter">Filter for the properties to return</param>
+		/// <returns>Information about the requested project</returns>
+		[HttpGet]
+		[Route("/api/v1/streams/{streamId}")]
+		[ProducesResponseType(typeof(GetStreamResponse), 200)]
+		public async Task<ActionResult<object>> GetStreamAsync(StreamId streamId, [FromQuery] PropertyFilter? filter = null)
+		{
+			StreamConfig? streamConfig;
+			if (!_globalConfig.Value.TryGetStream(streamId, out streamConfig))
+			{
+				return NotFound(streamId);
+			}
+			if (!streamConfig.Authorize(AclAction.ViewStream, User))
+			{
+				return Forbid(AclAction.ViewStream, streamId);
+			}
+
+			IStream stream = await _streamCollection.GetAsync(streamConfig);
+			return PropertyFilter.Apply(await CreateGetStreamResponseAsync(stream), filter);
+		}
+
+		/// <summary>
+		/// Retrieve information about a specific stream.
+		/// </summary>
+		/// <param name="streamId">Id of the stream to get information about</param>
+		/// <param name="config">The client's cached config revision. If this matches value matches the current config version of the stream, the config object will be omitted from the response.</param>
+		/// <param name="filter">Filter for the properties to return</param>
+		/// <returns>Information about the requested project</returns>
+		[HttpGet]
+		[Route("/api/v2/streams/{streamId}/config")]
+		[ProducesResponseType(typeof(GetStreamResponseV2), 200)]
+		public async Task<ActionResult<object>> GetStreamAsyncV2(StreamId streamId, [FromQuery] string? config = null, [FromQuery] PropertyFilter? filter = null)
+		{
+			StreamConfig? streamConfig;
+			if (!_globalConfig.Value.TryGetStream(streamId, out streamConfig))
+			{
+				return NotFound(streamId);
+			}
+			if (!streamConfig.Authorize(AclAction.ViewStream, User))
+			{
+				return Forbid(AclAction.ViewStream, streamId);
+			}
+
+			bool includeConfig = !String.Equals(config, streamConfig.Revision, StringComparison.Ordinal);
+			IStream stream = await _streamCollection.GetAsync(streamConfig);
+			return PropertyFilter.Apply(await CreateGetStreamResponseV2Async(stream, includeConfig), filter);
+		}
+
+		/// <summary>
+		/// Create a stream response object, including all the templates
+		/// </summary>
+		/// <param name="stream">Stream to create response for</param>
+		/// <returns>Response object</returns>
+		async Task<GetStreamResponse> CreateGetStreamResponseAsync(IStream stream)
+		{
+			using IScope scope = GlobalTracer.Instance.BuildSpan("CreateGetStreamResponse").StartActive();
+			scope.Span.SetTag("streamId", stream.Id);
+
+			List<GetTemplateRefResponse> apiTemplateRefs = new List<GetTemplateRefResponse>();
+			foreach (KeyValuePair<TemplateId, ITemplateRef> pair in stream.Templates)
+			{
+				using IScope templateScope = GlobalTracer.Instance.BuildSpan("CreateGetStreamResponse.Template").StartActive();
+				templateScope.Span.SetTag("templateName", pair.Value.Config.Name);
+
+				ITemplateRef templateRef = pair.Value;
+				if (templateRef.Config.Authorize(AclAction.ViewTemplate, User))
+				{
+					List<GetTemplateStepStateResponse>? stepStates = null;
+					if (templateRef.StepStates != null)
+					{
+						for (int i = 0; i < templateRef.StepStates.Count; i++)
+						{
+							ITemplateStep state = templateRef.StepStates[i];
+
+							stepStates ??= new List<GetTemplateStepStateResponse>();
+
+							GetThinUserInfoResponse? pausedByUserInfo = new GetThinUserInfoResponse(await _userCollection.GetCachedUserAsync(state.PausedByUserId));
+							stepStates.Add(new GetTemplateStepStateResponse(state, pausedByUserInfo));
+						}
+					}
+
+					ITemplate? template = await _templateCollection.GetOrAddAsync(pair.Value.Config);
+					apiTemplateRefs.Add(new GetTemplateRefResponse(pair.Key, pair.Value, template, stepStates, _timeZone));
+				}
+			}
+
+			return stream.ToApiResponse(apiTemplateRefs);
+		}
+
+		async Task<GetStreamResponseV2> CreateGetStreamResponseV2Async(IStream stream, bool includeConfig)
+		{
+			List<GetTemplateRefResponseV2> templates = new List<GetTemplateRefResponseV2>();
+			foreach (ITemplateRef templateRef in stream.Templates.Values)
+			{
+				ITemplate template = await _templateCollection.GetOrAddAsync(templateRef.Config);
+				templates.Add(new GetTemplateRefResponseV2(template, templateRef));
+			}
+			return new GetStreamResponseV2(stream, includeConfig, templates);
+		}
+
+		/// <summary>
+		/// Gets a list of changes for a stream
+		/// </summary>
+		/// <param name="streamId">The stream id</param>
+		/// <param name="min">The starting changelist number</param>
+		/// <param name="max">The ending changelist number</param>
+		/// <param name="results">Number of results to return</param>
+		/// <param name="tags">Tags to filter the changes returned</param>
+		/// <param name="filter">The filter to apply to the results</param>
+		/// <returns>Http result code</returns>
+		[HttpGet]
+		[Route("/api/v1/streams/{streamId}/changes")]
+		[ProducesResponseType(typeof(List<GetCommitResponse>), 200)]
+		public async Task<ActionResult<List<object>>> GetChangesAsync(StreamId streamId, [FromQuery] int? min = null, [FromQuery] int? max = null, [FromQuery] int results = 50, [FromQuery] string? tags = null, PropertyFilter? filter = null)
+		{
+			StreamConfig? streamConfig;
+			if (!_globalConfig.Value.TryGetStream(streamId, out streamConfig))
+			{
+				return NotFound(streamId);
+			}
+			if (!streamConfig.Authorize(AclAction.ViewChanges, User))
+			{
+				return Forbid(AclAction.ViewChanges, streamId);
+			}
+
+			List<CommitTag>? commitTags = null;
+			if (tags != null)
+			{
+				commitTags = tags.Split(';', StringSplitOptions.RemoveEmptyEntries).Select(x => new CommitTag(x)).ToList();
+			}
+
+			List<ICommit> commits = await _commitService.GetCollection(streamConfig).FindAsync(min, max, results, commitTags).ToListAsync();
+
+			List<GetCommitResponse> responses = new List<GetCommitResponse>();
+			foreach (ICommit commit in commits)
+			{
+				IUser? author = await _userCollection.GetCachedUserAsync(commit.AuthorId);
+				responses.Add(new GetCommitResponse(commit, author!, null, null));
+			}
+			return responses.ConvertAll(x => PropertyFilter.Apply(x, filter));
+		}
+
+		/// <summary>
+		/// Gets a list of changes for a stream
+		/// </summary>
+		/// <param name="streamId">The stream id</param>
+		/// <param name="changeNumber">The changelist number</param>
+		/// <param name="maxFiles">Maximum number of files to return</param>
+		/// <param name="filter">The filter to apply to the results</param>
+		/// <returns>Http result code</returns>
+		[HttpGet]
+		[Route("/api/v1/streams/{streamId}/changes/{changeNumber}")]
+		[ProducesResponseType(typeof(GetCommitResponse), 200)]
+		public async Task<ActionResult<object>> GetChangeDetailsAsync(StreamId streamId, int changeNumber, int maxFiles = 100, PropertyFilter? filter = null)
+		{
+			StreamConfig? streamConfig;
+			if (!_globalConfig.Value.TryGetStream(streamId, out streamConfig))
+			{
+				return NotFound(streamId);
+			}
+			if (!streamConfig.Authorize(AclAction.ViewChanges, User))
+			{
+				return Forbid(AclAction.ViewChanges, streamId);
+			}
+
+			ICommit? changeDetails = await _commitService.GetCollection(streamConfig).GetAsync(changeNumber);
+			if(changeDetails == null)
+			{
+				return NotFound("CL {Change} not found in stream {StreamId}", changeNumber, streamId);
+			}
+
+			IUser? author = await _userCollection.GetCachedUserAsync(changeDetails.AuthorId);
+			IReadOnlyList<CommitTag> tags = await changeDetails.GetTagsAsync(HttpContext.RequestAborted);
+			IReadOnlyList<string> files = await changeDetails.GetFilesAsync(maxFiles, CancellationToken.None);
+
+			return PropertyFilter.Apply(new GetCommitResponse(changeDetails, author!, tags, files), filter);
+		}
+
+		/// <summary>
+		/// Gets the history of a step in the stream
+		/// </summary>
+		/// <param name="streamId">The stream id</param>
+		/// <param name="templateId"></param>
+		/// <param name="step">Name of the step to search for</param>
+		/// <param name="change">Maximum changelist number to return</param>
+		/// <param name="count">Number of results to return</param>
+		/// <param name="filter">The filter to apply to the results</param>
+		/// <returns>Http result code</returns>
+		[HttpGet]
+		[Route("/api/v1/streams/{streamId}/history")]
+		[ProducesResponseType(typeof(List<GetJobStepRefResponse>), 200)]
+		public async Task<ActionResult<List<object>>> GetStepHistoryAsync(StreamId streamId, [FromQuery] string templateId, [FromQuery] string step, [FromQuery] int? change = null, [FromQuery] int count = 10, [FromQuery] PropertyFilter? filter = null)
+		{
+			StreamConfig? streamConfig;
+			if (!_globalConfig.Value.TryGetStream(streamId, out streamConfig))
+			{
+				return NotFound(streamId);
+			}
+			if (!streamConfig.Authorize(AclAction.ViewJob, User))
+			{
+				return Forbid(AclAction.ViewJob, streamId);
+			}
+
+			TemplateId templateIdValue = new TemplateId(templateId);
+
+			List<IJobStepRef> steps = await _jobStepRefCollection.GetStepsForNodeAsync(streamId, templateIdValue, step, change, true, count);
+			return steps.ConvertAll(x => PropertyFilter.Apply(new GetJobStepRefResponse(x), filter));
+		}
+
+		/// <summary>
+		/// Gets a template for a stream
+		/// </summary>
+		/// <param name="streamId">Unique id of the stream to query</param>
+		/// <param name="templateId">Unique id of the template to query</param>
+		/// <param name="filter">Filter for properties to return</param>
+		/// <returns>Information about all the templates</returns>
+		[HttpGet]
+		[Route("/api/v1/streams/{streamId}/templates/{templateId}")]
+		[ProducesResponseType(typeof(List<GetTemplateResponse>), 200)]
+		public async Task<ActionResult<object>> GetTemplateAsync(StreamId streamId, TemplateId templateId, [FromQuery] PropertyFilter? filter = null)
+		{
+			StreamConfig? streamConfig;
+			if (!_globalConfig.Value.TryGetStream(streamId, out streamConfig))
+			{
+				return NotFound(streamId);
+			}
+
+			TemplateRefConfig? templateConfig;
+			if (!streamConfig.TryGetTemplate(templateId, out templateConfig))
+			{
+				return NotFound(streamId, templateId);
+			}
+			if (!templateConfig.Authorize(AclAction.ViewTemplate, User))
+			{
+				return Forbid(AclAction.ViewTemplate, streamId);
+			}
+
+			ITemplate template = await _templateCollection.GetOrAddAsync(templateConfig);
+			return new GetTemplateResponse(template).ApplyFilter(filter);
+		}
+
+		/// <summary>
+		/// Update a stream template ref
+		/// </summary>
+		[HttpPut]
+		[Authorize]
+		[Route("/api/v1/streams/{streamId}/templates/{templateRefId}")]
+		public async Task<ActionResult> UpdateStreamTemplateRefAsync(StreamId streamId, TemplateId templateRefId, [FromBody] UpdateTemplateRefRequest update)
+		{
+			if (!_globalConfig.Value.Authorize(AclAction.AdminWrite, User))
+			{
+				return Forbid(AclAction.AdminWrite);
+			}
+
+			StreamConfig? streamConfig;
+			if (!_globalConfig.Value.TryGetStream(streamId, out streamConfig))
+			{
+				return NotFound(streamId);
+			}
+
+
+			IStream stream = await _streamCollection.GetAsync(streamConfig);
+			if (!stream.Templates.ContainsKey(templateRefId))
+			{
+				return NotFound(streamId, templateRefId);
+			}
+
+			await _streamCollection.TryUpdateTemplateRefAsync(stream, templateRefId, update.StepStates);
+			return Ok();
+		}
+	}
+}
diff --git a/Engine/Source/Programs/Horde/Horde.Server/Tasks/ConformTaskSource.cs b/Engine/Source/Programs/Horde/Horde.Server/Tasks/ConformTaskSource.cs
new file mode 100644
--- /dev/null
+++ b/Engine/Source/Programs/Horde/Horde.Server/Tasks/ConformTaskSource.cs
@@ -0,0 +1,414 @@
+// Copyright Epic Games, Inc. All Rights Reserved.
+
+using System;
+using System.Collections.Generic;
+using System.Linq;
+using System.Threading;
+using System.Threading.Tasks;
+using Google.Protobuf;
+using Google.Protobuf.WellKnownTypes;
+using Horde.Server.Agents;
+using Horde.Server.Agents.Leases;
+using Horde.Server.Agents.Pools;
+using Horde.Server.Jobs;
+using Horde.Server.Logs;
+using Horde.Server.Perforce;
+using Horde.Server.Server;
+using Horde.Server.Utilities;
+using HordeCommon;
+using HordeCommon.Rpc.Tasks;
+using Microsoft.Extensions.Hosting;
+using Microsoft.Extensions.Logging;
+using Microsoft.Extensions.Options;
+
+namespace Horde.Server.Tasks
+{
+	/// <summary>
+	/// Generates tasks telling agents to sync their workspaces
+	/// </summary>
+	public sealed class ConformTaskSource : TaskSourceBase<ConformTask>, IHostedService, IDisposable
+	{
+		/// <inheritdoc/>
+		public override string Type => "Conform";
+
+		/// <inheritdoc/>
+		public override TaskSourceFlags Flags => TaskSourceFlags.AllowWhenDisabled;
+
+		readonly MongoService _mongoService;
+		readonly GlobalsService _globalsService;
+		readonly IAgentCollection _agentCollection;
+		readonly PoolService _poolService;
+		readonly SingletonDocument<ConformList> _conformList;
+		readonly PerforceLoadBalancer _perforceLoadBalancer;
+		readonly ILogFileService _logService;
+		readonly ILogger _logger;
+		readonly IOptionsMonitor<ServerSettings> _settings;
+		readonly IOptionsMonitor<GlobalConfig> _globalConfig;
+		readonly ITicker _tickConformList;
+
+		/// <summary>
+		/// Constructor
+		/// </summary>
+		public ConformTaskSource(MongoService mongoService, GlobalsService globalsService, IAgentCollection agentCollection, PoolService poolService, ILogFileService logService, PerforceLoadBalancer perforceLoadBalancer, IClock clock, IOptionsMonitor<ServerSettings> settings, IOptionsMonitor<GlobalConfig> globalConfig, ILogger<ConformTaskSource> logger)
+		{
+			_mongoService = mongoService;
+			_globalsService = globalsService;
+			_agentCollection = agentCollection;
+			_poolService = poolService;
+			_conformList = new SingletonDocument<ConformList>(mongoService);
+			_perforceLoadBalancer = perforceLoadBalancer;
+			_logService = logService;
+			_settings = settings;
+			_globalConfig = globalConfig;
+			_logger = logger;
+			_tickConformList = clock.AddSharedTicker<ConformTaskSource>(TimeSpan.FromMinutes(1.0), CleanConformListAsync, logger);
+
+			OnLeaseStartedProperties.Add(nameof(ConformTask.LogId), x => LogId.Parse(x.LogId));
+		}
+
+		/// <inheritdoc/>
+		public Task StartAsync(CancellationToken cancellationToken) => _tickConformList.StartAsync();
+
+		/// <inheritdoc/>
+		public Task StopAsync(CancellationToken cancellationToken) => _tickConformList.StopAsync();
+
+		/// <inheritdoc/>
+		public void Dispose() => _tickConformList.Dispose();
+
+		/// <summary>
+		/// Clean up the conform list of any outdated entries
+		/// </summary>
+		/// <param name="cancellationToken">Cancellation token</param>
+		/// <returns>Async task</returns>
+		async ValueTask CleanConformListAsync(CancellationToken cancellationToken)
+		{
+			DateTime utcNow = DateTime.UtcNow;
+			DateTime lastCheckTimeUtc = utcNow - TimeSpan.FromMinutes(30.0);
+
+			// Get the current state of the conform list
+			ConformList list = await _conformList.GetAsync();
+
+			// Update any leases that are older than LastCheckTimeUtc
+			Dictionary<LeaseId, bool> removeLeases = new Dictionary<LeaseId, bool>();
+			foreach (ConformListEntry entry in Enumerable.Concat(list.Entries, list.Servers.SelectMany(x => x.Entries)))
+			{
+				if (entry.LastCheckTimeUtc < lastCheckTimeUtc)
+				{
+					IAgent? agent = await _agentCollection.GetAsync(entry.AgentId);
+
+					bool remove = false;
+					if (agent == null || !agent.Leases.Any(x => x.Id == entry.LeaseId))
+					{
+						_logger.LogWarning("Removing invalid lease from conform list: {AgentId} lease {LeaseId}", entry.AgentId, entry.LeaseId);
+						remove = true;
+					}
+
+					removeLeases[entry.LeaseId] = remove;
+				}
+			}
+
+			// If there's anything to change, update the list
+			if (removeLeases.Count > 0)
+			{
+				await _conformList.UpdateAsync(list => UpdateConformList(list, utcNow, removeLeases));
+			}
+		}
+
+		/// <summary>
+		/// Remove items from the conform list, and update timestamps for them
+		/// </summary>
+		/// <param name="list">The list to update</param>
+		/// <param name="utcNow">Current time</param>
+		/// <param name="removeLeases">List of leases to update. Entries with values set to true will be removed, entries with values set to false will have their timestamp updated.</param>
+		static void UpdateConformList(ConformList list, DateTime utcNow, Dictionary<LeaseId, bool> removeLeases)
+		{
+			UpdateConformList(list.Entries, utcNow, removeLeases);
+			foreach (ConformListServer server in list.Servers)
+			{
+				UpdateConformList(server.Entries, utcNow, removeLeases);
+			}
+		}
+
+		/// <summary>
+		/// Remove items from the conform list, and update timestamps for them
+		/// </summary>
+		/// <param name="entries">The list to update</param>
+		/// <param name="utcNow">Current time</param>
+		/// <param name="removeLeases">List of leases to update. Entries with values set to true will be removed, entries with values set to false will have their timestamp updated.</param>
+		static void UpdateConformList(List<ConformListEntry> entries, DateTime utcNow, Dictionary<LeaseId, bool> removeLeases)
+		{
+			for (int idx = 0; idx < entries.Count; idx++)
+			{
+				bool remove;
+				if (removeLeases.TryGetValue(entries[idx].LeaseId, out remove))
+				{
+					if (remove)
+					{
+						entries.RemoveAt(idx--);
+					}
+					else if (entries[idx].LastCheckTimeUtc < utcNow)
+					{
+						entries[idx].LastCheckTimeUtc = utcNow;
+					}
+				}
+			}
+		}
+
+		/// <inheritdoc/>
+		public override async Task<Task<AgentLease?>> AssignLeaseAsync(IAgent agent, CancellationToken cancellationToken)
+		{
+			if (!_settings.CurrentValue.EnableConformTasks)
+			{
+				return Skip(cancellationToken);
+			}
+
+			DateTime utcNow = DateTime.UtcNow;
+			if (!await IsConformPendingAsync(agent, utcNow))
+			{
+				return Skip(cancellationToken);
+			}
+
+			if (agent.Leases.Count == 0)
+			{
+				ConformTask task = new ConformTask();
+				if (await GetWorkspacesAsync(agent, task.Workspaces))
+				{
+					LeaseId leaseId = LeaseId.GenerateNewId();
+					if (await AllocateConformLeaseAsync(agent.Id, task.Workspaces, leaseId))
+					{
+						ILogFile log = await _logService.CreateLogFileAsync(JobId.Empty, leaseId, agent.SessionId, LogType.Json, useNewStorageBackend: false, cancellationToken: cancellationToken);
+						task.LogId = log.Id.ToString();
+						task.RemoveUntrackedFiles = agent.RequestFullConform;
+
+						byte[] payload = Any.Pack(task).ToByteArray();
+
+						return Lease(new AgentLease(leaseId, "Updating workspaces", null, null, log.Id, LeaseState.Pending, null, true, payload));
+					}
+				}
+			}
+
+			if (agent.RequestConform || agent.RequestFullConform)
+			{
+				return await DrainAsync(cancellationToken);
+			}
+			else
+			{
+				return Skip(cancellationToken);
+			}
+		}
+
+		/// <inheritdoc/>
+		public override Task CancelLeaseAsync(IAgent agent, LeaseId leaseId, ConformTask payload)
+		{
+			return ReleaseConformLeaseAsync(leaseId);
+		}
+
+		/// <inheritdoc/>
+		public async Task<bool> GetWorkspacesAsync(IAgent agent, IList<HordeCommon.Rpc.Messages.AgentWorkspace> workspaces)
+		{
+			GlobalConfig globalConfig = _globalConfig.CurrentValue;
+
+			HashSet<AgentWorkspace> conformWorkspaces = await _poolService.GetWorkspacesAsync(agent, DateTime.UtcNow, globalConfig);
+			foreach (AgentWorkspace conformWorkspace in conformWorkspaces)
+			{
+				PerforceCluster? cluster = globalConfig.FindPerforceCluster(conformWorkspace.Cluster);
+				if (cluster == null || !await agent.TryAddWorkspaceMessage(conformWorkspace, cluster, _perforceLoadBalancer, workspaces))
+				{
+					return false;
+				}
+			}
+
+			return true;
+		}
+
+		/// <inheritdoc/>
+		public override Task OnLeaseFinishedAsync(IAgent agent, LeaseId leaseId, ConformTask payload, LeaseOutcome outcome, ReadOnlyMemory<byte> output, ILogger logger)
+		{
+			return ReleaseConformLeaseAsync(leaseId);
+		}
+
+		/// <summary>
+		/// Atempt to allocate a conform resource for the given lease
+		/// </summary>
+		/// <param name="agentId">The agent id</param>
+		/// <param name="workspaces">List of workspaces that are required</param>
+		/// <param name="leaseId">The lease id</param>
+		/// <returns>True if the resource was allocated, false otherwise</returns>
+		private async Task<bool> AllocateConformLeaseAsync(AgentId agentId, IEnumerable<HordeCommon.Rpc.Messages.AgentWorkspace> workspaces, LeaseId leaseId)
+		{
+			GlobalConfig globalConfig = _globalConfig.CurrentValue;
+			for (; ; )
+			{
+				ConformList currentValue = await _conformList.GetAsync();
+				if (globalConfig.MaxConformCount != 0 && currentValue.Entries.Count + currentValue.Servers.Sum(x => x.Entries.Count) >= globalConfig.MaxConformCount)
+				{
+					return false;
+				}
+
+				HashSet<string> servers = new HashSet<string>(StringComparer.OrdinalIgnoreCase);
+				foreach (HordeCommon.Rpc.Messages.AgentWorkspace workspace in workspaces)
+				{
+					if (servers.Add(workspace.ServerAndPort))
+					{
+						PerforceCluster? cluster = globalConfig.FindPerforceCluster(workspace.Cluster);
+						if (cluster == null)
+						{
+							_logger.LogWarning("Unable to find perforce cluster '{Cluster}' for conform", workspace.Cluster);
+							return false;
+						}
+
+						PerforceServer? server = cluster.Servers.FirstOrDefault(x => x.ServerAndPort.Equals(workspace.BaseServerAndPort, StringComparison.Ordinal));
+						if (server == null)
+						{
+							_logger.LogWarning("Unable to find perforce server '{Server}' in cluster '{Cluster}' for conform", workspace.BaseServerAndPort, workspace.Cluster);
+							return false;
+						}
+
+						ConformListServer? conformServer = currentValue.Servers.FirstOrDefault(x => x.Cluster.Equals(workspace.Cluster, StringComparison.Ordinal) && x.ServerAndPort.Equals(workspace.ServerAndPort, StringComparison.Ordinal));
+						if (conformServer == null)
+						{
+							conformServer = new ConformListServer { Cluster = workspace.Cluster, ServerAndPort = workspace.ServerAndPort };
+							currentValue.Servers.Add(conformServer);
+						}
+
+						if (server.MaxConformCount != 0 && conformServer.Entries.Count > server.MaxConformCount)
+						{
+							return false;
+						}
+
+						ConformListEntry entry = new ConformListEntry();
+						entry.AgentId = agentId;
+						entry.LeaseId = leaseId;
+						entry.LastCheckTimeUtc = DateTime.UtcNow;
+						conformServer.Entries.Add(entry);
+					}
+				}
+
+				if (await _conformList.TryUpdateAsync(currentValue))
+				{
+					_logger.LogInformation("Added conform lease {LeaseId}", leaseId);
+					return true;
+				}
+			}
+		}
+
+		/// <summary>
+		/// Terminate a conform lease
+		/// </summary>
+		/// <param name="leaseId">The lease id</param>
+		/// <returns>Async task</returns>
+		public async Task ReleaseConformLeaseAsync(LeaseId leaseId)
+		{
+			for (; ; )
+			{
+				ConformList currentValue = await _conformList.GetAsync();
+				if (currentValue.Entries.RemoveAll(x => x.LeaseId == leaseId) + currentValue.Servers.Sum(x => x.Entries.RemoveAll(x => x.LeaseId == leaseId)) == 0)
+				{
+					_logger.LogInformation("Conform lease {LeaseId} is not in singelton", leaseId);
+					break;
+				}
+				if (await _conformList.TryUpdateAsync(currentValue))
+				{
+					_logger.LogInformation("Removed conform lease {LeaseId}", leaseId);
+					break;
+				}
+			}
+		}
+
+		/// <summary>
+		/// Determine if an agent should be conformed
+		/// </summary>
+		/// <param name="agent">The agent to test</param>
+		/// <param name="utcNow">Current time</param>
+		/// <returns>True if the agent should be conformed</returns>
+		private async Task<bool> IsConformPendingAsync(IAgent agent, DateTime utcNow)
+		{
+			GlobalConfig globalConfig = _globalConfig.CurrentValue;
+
+			// If a conform was manually requested, allow it to run even if the agent is disabled
+			if (agent.RequestConform || agent.RequestFullConform)
+			{
+				return !IsConformCoolDownPeriod(agent, utcNow);
+			}
+
+			// Otherwise only run if the agent is enabled
+			if (agent.Enabled)
+			{
+				// If we've attempted (and failed) a conform, run again after a certain time
+				if (agent.ConformAttemptCount.HasValue)
+				{
+					return !IsConformCoolDownPeriod(agent, utcNow);
+				}
+
+				// Get the current pools for the agent
+				List<IPool> pools = await _poolService.GetCachedPoolsAsync(agent, DateTime.UtcNow - TimeSpan.FromMinutes(2.0));
+
+				TimeSpan? conformInterval = null;
+				foreach(IPool pool in pools)
+				{
+					TimeSpan interval = pool.ConformInterval ?? TimeSpan.FromDays(1.0);
+					if (interval > TimeSpan.Zero && (conformInterval == null || interval < conformInterval.Value))
+					{
+						conformInterval = interval;
+					}
+				}
+
+				// If there is no conform interval, early out
+				if(conformInterval == null)
+				{
+					return false;
+				}
+
+				// Always run a conform every 24h
+				if (utcNow > agent.LastConformTime + conformInterval.Value)
+				{
+					return true;
+				}
+
+				// Check if the workspaces have changed (first check against a cached list of workspaces, then an accurate one)
+				HashSet<AgentWorkspace> workspaces = await _poolService.GetWorkspacesAsync(agent, utcNow - TimeSpan.FromSeconds(30.0), globalConfig);
+				if (!workspaces.SetEquals(agent.Workspaces))
+				{
+					workspaces = await _poolService.GetWorkspacesAsync(agent, utcNow, globalConfig);
+					if (!workspaces.SetEquals(agent.Workspaces))
+					{
+						return true;
+					}
+				}
+			}
+			return false;
+		}
+
+		/// <summary>
+		/// Determines if the given agent is in a cooldown period, waiting to retry after a failed conform
+		/// </summary>
+		/// <param name="agent">The agent to test</param>
+		/// <param name="utcNow">Current time</param>
+		/// <returns></returns>
+		private static bool IsConformCoolDownPeriod(IAgent agent, DateTime utcNow)
+		{
+			if (!agent.ConformAttemptCount.HasValue)
+			{
+				return false;
+			}
+
+			TimeSpan retryTime;
+			switch (agent.ConformAttemptCount.Value)
+			{
+				case 1:
+					retryTime = TimeSpan.FromMinutes(5.0);
+					break;
+				case 2:
+					retryTime = TimeSpan.FromMinutes(20.0);
+					break;
+				case 3:
+					retryTime = TimeSpan.FromHours(1.0);
+					break;
+				default:
+					retryTime = TimeSpan.FromHours(6.0);
+					break;
+			}
+			return utcNow < agent.LastConformTime + retryTime;
+		}
+	}
+}
diff --git a/Engine/Source/Programs/Horde/Horde.Server/Tasks/RestartTaskSource.cs b/Engine/Source/Programs/Horde/Horde.Server/Tasks/RestartTaskSource.cs
new file mode 100644
--- /dev/null
+++ b/Engine/Source/Programs/Horde/Horde.Server/Tasks/RestartTaskSource.cs
@@ -0,0 +1,54 @@
+// Copyright Epic Games, Inc. All Rights Reserved.
+
+using System.Threading;
+using System.Threading.Tasks;
+using Google.Protobuf;
+using Google.Protobuf.WellKnownTypes;
+using Horde.Server.Agents;
+using Horde.Server.Agents.Leases;
+using Horde.Server.Jobs;
+using Horde.Server.Logs;
+using Horde.Server.Utilities;
+using HordeCommon;
+using HordeCommon.Rpc.Tasks;
+
+namespace Horde.Server.Tasks
+{
+	class RestartTaskSource : TaskSourceBase<RestartTask>
+	{
+		public override string Type => "Restart";
+
+		public override TaskSourceFlags Flags => TaskSourceFlags.AllowWhenDisabled | TaskSourceFlags.AllowDuringDowntime;
+
+		readonly ILogFileService _logService;
+
+		public RestartTaskSource(ILogFileService logService)
+		{
+			_logService = logService;
+
+			OnLeaseStartedProperties.Add(nameof(RestartTask.LogId), x => LogId.Parse(x.LogId));
+		}
+
+		public override async Task<Task<AgentLease?>> AssignLeaseAsync(IAgent agent, CancellationToken cancellationToken)
+		{
+			if (!agent.RequestRestart)
+			{
+				return Skip(cancellationToken);
+			}
+			if (agent.Leases.Count > 0)
+			{
+				return await DrainAsync(cancellationToken);
+			}
+
+			LeaseId leaseId = LeaseId.GenerateNewId();
+			ILogFile log = await _logService.CreateLogFileAsync(JobId.Empty, leaseId, agent.SessionId, LogType.Json, useNewStorageBackend: false, cancellationToken: cancellationToken);
+
+			RestartTask task = new RestartTask();
+			task.LogId = log.Id.ToString();
+
+			byte[] payload = Any.Pack(task).ToByteArray();
+
+			return Lease(new AgentLease(leaseId, "Restart", null, null, log.Id, LeaseState.Pending, null, true, payload));
+		}
+	}
+}
diff --git a/Engine/Source/Programs/Horde/Horde.Server/Tasks/ShutdownTaskSource.cs b/Engine/Source/Programs/Horde/Horde.Server/Tasks/ShutdownTaskSource.cs
new file mode 100644
--- /dev/null
+++ b/Engine/Source/Programs/Horde/Horde.Server/Tasks/ShutdownTaskSource.cs
@@ -0,0 +1,53 @@
+// Copyright Epic Games, Inc. All Rights Reserved.
+
+using System.Threading;
+using System.Threading.Tasks;
+using Google.Protobuf;
+using Google.Protobuf.WellKnownTypes;
+using Horde.Server.Agents;
+using Horde.Server.Agents.Leases;
+using Horde.Server.Jobs;
+using Horde.Server.Logs;
+using Horde.Server.Utilities;
+using HordeCommon;
+using HordeCommon.Rpc.Tasks;
+
+namespace Horde.Server.Tasks
+{
+	class ShutdownTaskSource : TaskSourceBase<ShutdownTask>
+	{
+		public override string Type => "Shutdown";
+
+		public override TaskSourceFlags Flags => TaskSourceFlags.AllowWhenDisabled | TaskSourceFlags.AllowDuringDowntime;
+
+		readonly ILogFileService _logService;
+
+		public ShutdownTaskSource(ILogFileService logService)
+		{
+			_logService = logService;
+			OnLeaseStartedProperties.Add(nameof(ShutdownTask.LogId), x => LogId.Parse(x.LogId));
+		}
+
+		public override async Task<Task<AgentLease?>> AssignLeaseAsync(IAgent agent, CancellationToken cancellationToken)
+		{
+			if (!agent.RequestShutdown)
+			{
+				return Skip(cancellationToken);
+			}
+			if (agent.Leases.Count > 0)
+			{
+				return await DrainAsync(cancellationToken);
+			}
+
+			LeaseId leaseId = LeaseId.GenerateNewId();
+			ILogFile log = await _logService.CreateLogFileAsync(JobId.Empty, leaseId, agent.SessionId, LogType.Json, useNewStorageBackend: false, cancellationToken: cancellationToken);
+
+			ShutdownTask task = new ShutdownTask();
+			task.LogId = log.Id.ToString();
+
+			byte[] payload = Any.Pack(task).ToByteArray();
+
+			return Lease(new AgentLease(leaseId, "Shutdown", null, null, log.Id, LeaseState.Pending, null, true, payload));
+		}
+	}
+}
diff --git a/Engine/Source/Programs/Horde/Horde.Server/Tasks/UpgradeTaskSource.cs b/Engine/Source/Programs/Horde/Horde.Server/Tasks/UpgradeTaskSource.cs
new file mode 100644
--- /dev/null
+++ b/Engine/Source/Programs/Horde/Horde.Server/Tasks/UpgradeTaskSource.cs
@@ -0,0 +1,110 @@
+// Copyright Epic Games, Inc. All Rights Reserved.
+
+using System;
+using System.Text;
+using System.Threading;
+using System.Threading.Tasks;
+using EpicGames.Core;
+using Google.Protobuf;
+using Google.Protobuf.WellKnownTypes;
+using Horde.Server.Agents;
+using Horde.Server.Agents.Leases;
+using Horde.Server.Jobs;
+using Horde.Server.Logs;
+using Horde.Server.Server;
+using Horde.Server.Tools;
+using HordeCommon;
+using HordeCommon.Rpc.Tasks;
+using Microsoft.Extensions.Options;
+
+namespace Horde.Server.Tasks
+{
+	class UpgradeTaskSource : TaskSourceBase<UpgradeTask>
+	{
+		public override string Type => "Upgrade";
+
+		public override TaskSourceFlags Flags => TaskSourceFlags.AllowWhenDisabled | TaskSourceFlags.AllowDuringDowntime;
+
+		readonly IToolCollection _toolCollection;
+		readonly ILogFileService _logService;
+		readonly IOptionsMonitor<GlobalConfig> _globalConfig;
+		readonly IOptions<ServerSettings> _serverSettings;
+		readonly IClock _clock;
+
+		public UpgradeTaskSource(IToolCollection toolCollection, ILogFileService logService, IOptionsMonitor<GlobalConfig> globalConfig, IOptions<ServerSettings> serverSettings, IClock clock)
+		{
+			_toolCollection = toolCollection;
+			_logService = logService;
+			_globalConfig = globalConfig;
+			_serverSettings = serverSettings;
+			_clock = clock;
+
+			OnLeaseStartedProperties.Add(nameof(UpgradeTask.LogId), x => LogId.Parse(x.LogId));
+		}
+
+		public override async Task<Task<AgentLease?>> AssignLeaseAsync(IAgent agent, CancellationToken cancellationToken)
+		{
+			if (!_serverSettings.Value.EnableUpgradeTasks)
+			{
+				return Skip(cancellationToken);
+			}
+
+			(ITool, IToolDeployment)? required = await GetRequiredSoftwareVersion(agent);
+			if (required == null)
+			{
+				return Skip(cancellationToken);
+			}
+
+			(ITool tool, IToolDeployment deployment) = required.Value; 
+			if (agent.Version == deployment.Version)
+			{
+				return Skip(cancellationToken);
+			}
+
+			string softwareId = $"{tool.Id}:{deployment.Version}";
+			if (agent.Leases.Count > 0 || !(agent.LastUpgradeTime == null || agent.LastUpgradeTime.Value + TimeSpan.FromMinutes(5.0) < _clock.UtcNow || agent.LastUpgradeVersion != softwareId))
+			{
+				return await DrainAsync(cancellationToken);
+			}
+
+			LeaseId leaseId = LeaseId.GenerateNewId();
+			ILogFile logFile = await _logService.CreateLogFileAsync(JobId.Empty, leaseId, agent.SessionId, LogType.Json, useNewStorageBackend: false, cancellationToken: cancellationToken);
+
+			UpgradeTask task = new UpgradeTask();
+			task.SoftwareId = softwareId;
+			task.LogId = logFile.Id.ToString();
+
+			byte[] payload = Any.Pack(task).ToByteArray();
+			return Lease(new AgentLease(leaseId, $"Upgrade to {tool.Id} {deployment.Version}", null, null, logFile.Id, LeaseState.Pending, null, true, payload));
+		}
+
+		/// <summary>
+		/// Determines the client software version that should be installed on an agent
+		/// </summary>
+		/// <param name="agent">The agent instance</param>
+		/// <returns>Unique id of the client version this agent should be running</returns>
+		public async Task<(ITool, IToolDeployment)?> GetRequiredSoftwareVersion(IAgent agent)
+		{
+			GlobalConfig globalConfig = _globalConfig.CurrentValue;
+
+			ToolId toolId = agent.GetSoftwareToolId(globalConfig);
+
+			ITool? tool = await _toolCollection.GetAsync(toolId, globalConfig);
+			if (tool == null)
+			{
+				return null;
+			}
+
+			uint value = BuzHash.Add(0, Encoding.UTF8.GetBytes(agent.Id.ToString()));
+			double phase = (value % 10000) / 10000.0;
+
+			IToolDeployment? deployment = tool.GetCurrentDeployment(phase, _clock.UtcNow);
+			if (deployment == null)
+			{
+				return null;
+			}
+
+			return (tool, deployment);
+		}
+	}
+}
diff --git a/Engine/Source/Programs/Horde/Horde.Server/Tools/ToolConfig.cs b/Engine/Source/Programs/Horde/Horde.Server/Tools/ToolConfig.cs
new file mode 100644
--- /dev/null
+++ b/Engine/Source/Programs/Horde/Horde.Server/Tools/ToolConfig.cs
@@ -0,0 +1,108 @@
+// Copyright Epic Games, Inc. All Rights Reserved.
+
+using System;
+using System.ComponentModel.DataAnnotations;
+using System.Diagnostics;
+using System.Security.Claims;
+using System.Text.Json.Serialization;
+using Horde.Server.Acls;
+using Horde.Server.Server;
+using Horde.Server.Utilities;
+
+namespace Horde.Server.Tools
+{
+	/// <summary>
+	/// Options for configuring a tool
+	/// </summary>
+	[DebuggerDisplay("{Id}")]
+	public class ToolConfig
+	{
+		/// <summary>
+		/// The global config containing this tool
+		/// </summary>
+		[JsonIgnore]
+		public GlobalConfig GlobalConfig { get; private set; } = null!;
+
+		/// <summary>
+		/// Unique identifier for the tool
+		/// </summary>
+		[Required]
+		public ToolId Id { get; set; }
+
+		/// <summary>
+		/// Name of the tool
+		/// </summary>
+		[Required]
+		public string Name { get; set; }
+
+		/// <summary>
+		/// Description for the tool
+		/// </summary>
+		public string Description { get; set; }
+
+		/// <summary>
+		/// Whether this tool should be exposed for download on a public endpoint without authentication
+		/// </summary>
+		public bool Public { get; set; }
+
+		/// <summary>
+		/// Permissions for the tool
+		/// </summary>
+		public AclConfig? Acl { get; set; }
+
+		/// <summary>
+		/// Default constructor for serialization
+		/// </summary>
+		public ToolConfig()
+		{
+			Name = String.Empty;
+			Description = String.Empty;
+		}
+
+		/// <summary>
+		/// Constructor
+		/// </summary>
+		public ToolConfig(ToolId id)
+		{
+			Id = id;
+			Name = id.ToString();
+			Description = String.Empty;
+		}
+
+		/// <summary>
+		/// Called after the config has been read
+		/// </summary>
+		/// <param name="globalConfig">Parent GlobalConfig object</param>
+		public void PostLoad(GlobalConfig globalConfig)
+		{
+			GlobalConfig = globalConfig;
+		}
+
+		/// <summary>
+		/// Authorizes a user to perform a given action
+		/// </summary>
+		/// <param name="action">The action being performed</param>
+		/// <param name="user">The principal to validate</param>
+		public bool Authorize(AclAction action, ClaimsPrincipal user)
+		{
+			return Acl?.Authorize(action, user) ?? GlobalConfig.Authorize(action, user);
+		}
+	}
+
+	/// <summary>
+	/// Options for a new deployment
+	/// </summary>
+	public class ToolDeploymentConfig
+	{
+		/// <inheritdoc cref="IToolDeployment.Version"/>
+		public string Version { get; set; } = "Unknown";
+
+		/// <inheritdoc cref="IToolDeployment.Duration"/>
+		public TimeSpan Duration { get; set; }
+
+		/// <summary>
+		/// Whether to create the deployment in a paused state
+		/// </summary>
+		public bool CreatePaused { get; set; }
+	}
+}
diff --git a/Engine/Source/Programs/Horde/Horde.Server/Tools/ToolDeploymentId.cs b/Engine/Source/Programs/Horde/Horde.Server/Tools/ToolDeploymentId.cs
new file mode 100644
--- /dev/null
+++ b/Engine/Source/Programs/Horde/Horde.Server/Tools/ToolDeploymentId.cs
@@ -0,0 +1,39 @@
+// Copyright Epic Games, Inc. All Rights Reserved.
+
+using System.ComponentModel;
+using EpicGames.Core;
+using Horde.Server.Utilities;
+using MongoDB.Bson;
+
+namespace Horde.Server.Tools
+{
+	/// <summary>
+	/// Identifier for a tool deployment
+	/// </summary>
+	/// <param name="Id">Id to construct from</param>
+	[JsonSchemaString]
+	[TypeConverter(typeof(ObjectIdTypeConverter<ToolDeploymentId, ToolDeploymentIdConverter>))]
+	[ObjectIdConverter(typeof(ToolDeploymentIdConverter))]
+	public record struct ToolDeploymentId(ObjectId Id)
+	{
+		/// <summary>
+		/// Creates a new <see cref="ToolDeploymentId"/>
+		/// </summary>
+		public static ToolDeploymentId GenerateNewId() => new ToolDeploymentId(ObjectId.GenerateNewId());
+
+		/// <inheritdoc/>
+		public override string ToString() => Id.ToString();
+	}
+
+	/// <summary>
+	/// Converter to and from <see cref="ObjectId"/> instances.
+	/// </summary>
+	class ToolDeploymentIdConverter : ObjectIdConverter<ToolDeploymentId>
+	{
+		/// <inheritdoc/>
+		public override ToolDeploymentId FromObjectId(ObjectId id) => new ToolDeploymentId(id);
+
+		/// <inheritdoc/>
+		public override ObjectId ToObjectId(ToolDeploymentId value) => value.Id;
+	}
+}
diff --git a/Engine/Source/Programs/Horde/Horde.Server/Ugs/UgsController.cs b/Engine/Source/Programs/Horde/Horde.Server/Ugs/UgsController.cs
new file mode 100644
--- /dev/null
+++ b/Engine/Source/Programs/Horde/Horde.Server/Ugs/UgsController.cs
@@ -0,0 +1,393 @@
+// Copyright Epic Games, Inc. All Rights Reserved.
+
+using System;
+using System.Collections.Generic;
+using System.Linq;
+using System.Threading;
+using System.Threading.Tasks;
+using Horde.Server.Users;
+using Horde.Server.Logs;
+using Horde.Server.Issues;
+using Horde.Server.Utilities;
+using Microsoft.AspNetCore.Mvc;
+using Microsoft.Extensions.Logging;
+using Microsoft.Extensions.Options;
+using MongoDB.Bson;
+
+namespace Horde.Server.Ugs
+{
+	/// <summary>
+	/// Controller for the /api/v1/issues endpoint
+	/// </summary>
+	[ApiController]
+	[Route("[controller]")]
+	public sealed class UgsController : ControllerBase
+	{
+		/// <summary>
+		/// Singleton instance of the issue service
+		/// </summary>
+		private readonly IssueService _issueService;
+
+		/// <summary>
+		/// Collection of metadata documents
+		/// </summary>
+		private readonly IUgsMetadataCollection _ugsMetadataCollection;
+
+		/// <summary>
+		/// Collection of users
+		/// </summary>
+		private readonly IUserCollection _userCollection;
+
+		/// <summary>
+		/// The log file service
+		/// </summary>
+		private readonly ILogFileService _logFileService;
+
+		/// <summary>
+		/// Server settings
+		/// </summary>
+		private readonly ServerSettings _settings;
+
+		/// <summary>
+		/// Logger 
+		/// </summary>
+		readonly ILogger _logger;
+
+		/// <summary>
+		/// Constructor
+		/// </summary>
+		public UgsController(IssueService issueService, IUgsMetadataCollection ugsMetadataCollection, IUserCollection userCollection, ILogFileService logFileService, IOptionsMonitor<ServerSettings> optionsMonitor, ILogger<UgsController> logger)
+		{
+			_issueService = issueService;
+			_ugsMetadataCollection = ugsMetadataCollection;
+			_userCollection = userCollection;
+			_logFileService = logFileService;
+			_settings = optionsMonitor.CurrentValue;
+			_logger = logger;
+		}
+
+		/// <summary>
+		/// Gets the latest version info
+		/// </summary>
+		/// <returns>Result code</returns>
+		[HttpGet]
+		[Route("/ugs/api/latest")]
+		public ActionResult<object> GetLatest()
+		{
+			return new { Version = 2 };
+		}
+
+		/// <summary>
+		/// Adds new metadata to the database
+		/// </summary>
+		/// <param name="request">Request object</param>
+		/// <returns>Result code</returns>
+		[HttpPost]
+		[Route("/ugs/api/metadata")]
+		public async Task<ActionResult> AddMetadataAsync(AddUgsMetadataRequest request)
+		{
+			IUgsMetadata metadata = await _ugsMetadataCollection.FindOrAddAsync(request.Stream, request.Change, request.Project);
+			if (request.Synced != null || request.Vote != null || request.Investigating != null || request.Starred != null || request.Comment != null)
+			{
+				if (request.UserName == null)
+				{
+					return BadRequest("Missing UserName field on request body");
+				}
+				metadata = await _ugsMetadataCollection.UpdateUserAsync(metadata, request.UserName, request.Synced, request.Vote, request.Investigating, request.Starred, request.Comment);
+			}
+			if (request.Badges != null)
+			{
+				foreach (AddUgsBadgeRequest badge in request.Badges)
+				{
+					metadata = await _ugsMetadataCollection.UpdateBadgeAsync(metadata, badge.Name, badge.Url, badge.State);
+				}
+			}
+			return Ok();
+		}
+
+		/// <summary>
+		/// Searches for metadata updates
+		/// </summary>
+		/// <param name="stream">THe stream to search for</param>
+		/// <param name="minChange">Minimum changelist number</param>
+		/// <param name="maxChange">Maximum changelist number</param>
+		/// <param name="project">The project identifiers to search for</param>
+		/// <param name="sequence">Last sequence number</param>
+		/// <returns>List of metadata updates</returns>
+		[HttpGet]
+		[Route("/ugs/api/metadata")]
+		public async Task<GetUgsMetadataListResponse> FindMetadataAsync([FromQuery] string stream, [FromQuery] int minChange, [FromQuery] int? maxChange = null, [FromQuery] string? project = null, [FromQuery] long? sequence = null)
+		{
+			List<IUgsMetadata> metadataList = await _ugsMetadataCollection.FindAsync(stream, minChange, maxChange, sequence);
+		
+			GetUgsMetadataListResponse response = new GetUgsMetadataListResponse();
+			if(sequence != null)
+			{
+				response.SequenceNumber = sequence.Value;
+			}
+
+			foreach (IUgsMetadata metadata in metadataList)
+			{
+				if (metadata.UpdateTicks > response.SequenceNumber)
+				{
+					response.SequenceNumber = metadata.UpdateTicks;
+				}
+				if (String.IsNullOrEmpty(metadata.Project) || metadata.Project.Equals(project, StringComparison.OrdinalIgnoreCase))
+				{
+					response.Items.Add(new GetUgsMetadataResponse(metadata));
+				}
+			}
+			return response;
+		}
+
+		/// <summary>
+		/// Retrieve information about open issues
+		/// </summary>
+		/// <param name="user"></param>
+		/// <param name="includeResolved">Whether to include resolved issues</param>
+		/// <param name="maxResults">Maximum number of results to return</param>
+		/// <returns>List of matching agents</returns>
+		[HttpGet]
+		[Route("/ugs/api/issues")]
+		[ProducesResponseType(typeof(GetUgsIssueResponse), 200)]
+		public async Task<ActionResult<List<GetUgsIssueResponse>>> GetIssuesAsync([FromQuery] string? user = null, [FromQuery] bool includeResolved = false, [FromQuery] int maxResults = 100)
+		{
+			IUser? userInfo = (user != null) ? await _userCollection.FindUserByLoginAsync(user) : null;
+
+			List<GetUgsIssueResponse> responses = new List<GetUgsIssueResponse>();
+			if (includeResolved)
+			{
+				List<IIssue> issues = await _issueService.Collection.FindIssuesAsync(null, resolved: null, count: maxResults);
+				foreach(IIssue issue in issues)
+				{
+					IIssueDetails details = await _issueService.GetIssueDetailsAsync(issue);
+					bool notify = userInfo != null && details.Suspects.Any(x => x.AuthorId == userInfo.Id);
+					responses.Add(await CreateIssueResponseAsync(details, notify));
+				}
+			}
+			else
+			{
+				foreach (IIssueDetails cachedOpenIssue in _issueService.CachedOpenIssues)
+				{
+					if (responses.Count >= maxResults)
+					{
+						break;
+					}
+
+					if (cachedOpenIssue.ShowNotifications())
+					{
+						bool notify = userInfo != null && cachedOpenIssue.IncludeForUser(userInfo.Id);
+						responses.Add(await CreateIssueResponseAsync(cachedOpenIssue, notify));
+					}
+				}
+			}
+			return responses;
+		}
+
+		/// <summary>
+		/// Retrieve information about a specific issue
+		/// </summary>
+		/// <param name="issueId">Id of the issue to get information about</param>
+		/// <param name="filter">Filter for the properties to return</param>
+		/// <returns>List of matching agents</returns>
+		[HttpGet]
+		[Route("/ugs/api/issues/{issueId}")]
+		[ProducesResponseType(typeof(GetUgsIssueBuildResponse), 200)]
+		public async Task<ActionResult<object>> GetIssueAsync(int issueId, [FromQuery] PropertyFilter? filter = null)
+		{
+			IIssueDetails? issue = await _issueService.GetIssueDetailsAsync(issueId);
+			if (issue == null)
+			{
+				return NotFound();
+			}
+
+			return PropertyFilter.Apply(await CreateIssueResponseAsync(issue, false), filter);
+		}
+
+		/// <summary>
+		/// Retrieve information about builds for a specific issue
+		/// </summary>
+		/// <param name="issueId">Id of the issue to get information about</param>
+		/// <param name="filter">Filter for the properties to return</param>
+		/// <returns>List of matching agents</returns>
+		[HttpGet]
+		[Route("/ugs/api/issues/{issueId}/builds")]
+		[ProducesResponseType(typeof(List<GetUgsIssueBuildResponse>), 200)]
+		public async Task<ActionResult<List<object>>> GetIssueBuildsAsync(int issueId, [FromQuery] PropertyFilter? filter = null)
+		{
+			IIssueDetails? issue = await _issueService.GetCachedIssueDetailsAsync(issueId);
+			if (issue == null)
+			{
+				return NotFound();
+			}
+
+			List<GetUgsIssueBuildResponse> responses = new List<GetUgsIssueBuildResponse>();
+			foreach (IIssueSpan span in issue.Spans)
+			{
+				if (span.LastSuccess != null)
+				{
+					responses.Add(CreateBuildResponse(span, span.LastSuccess, IssueBuildOutcome.Success));
+				}
+				foreach (IIssueStep step in issue.Steps)
+				{
+					responses.Add(CreateBuildResponse(span, step, IssueBuildOutcome.Error));
+				}
+				if (span.NextSuccess != null)
+				{
+					responses.Add(CreateBuildResponse(span, span.NextSuccess, IssueBuildOutcome.Success));
+				}
+			}
+
+			return responses.ConvertAll(x => PropertyFilter.Apply(x, filter));
+		}
+
+		/// <summary>
+		/// Retrieve information about builds for a specific issue
+		/// </summary>
+		/// <param name="issueId">Id of the issue to get information about</param>
+		/// <param name="cancellationToken">Cancellation token for the request</param>
+		/// <returns>List of matching agents</returns>
+		[HttpGet]
+		[Route("/ugs/api/issues/{issueId}/diagnostics")]
+		[ProducesResponseType(typeof(List<GetUgsIssueDiagnosticResponse>), 200)]
+		public async Task<ActionResult<List<GetUgsIssueDiagnosticResponse>>> GetIssueDiagnosticsAsync(int issueId, CancellationToken cancellationToken)
+		{
+			List<GetUgsIssueDiagnosticResponse> diagnostics = new List<GetUgsIssueDiagnosticResponse>();
+
+			Dictionary<LogId, ILogFile?> logFiles = new Dictionary<LogId, ILogFile?>();
+
+			List<IIssueSpan> spans = await _issueService.Collection.FindSpansAsync(issueId);
+			List<ILogEvent> events = await _logFileService.FindEventsForSpansAsync(spans.Select(x => x.Id), null, 0, count: 10, cancellationToken);
+
+			foreach (ILogEvent logEvent in events)
+			{
+				ILogFile? logFile;
+				if(!logFiles.TryGetValue(logEvent.LogId, out logFile))
+				{
+					logFile = await _logFileService.GetLogFileAsync(logEvent.LogId, cancellationToken);
+					logFiles.Add(logEvent.LogId, logFile);
+				}
+				if (logFile != null)
+				{
+					ILogEventData eventData = await _logFileService.GetEventDataAsync(logFile, logEvent.LineIndex, logEvent.LineCount, cancellationToken);
+					long buildId = logEvent.LogId.GetHashCode();
+					Uri url = new Uri(_settings.DashboardUrl, $"log/{logEvent.LogId}?lineindex={logEvent.LineIndex}");
+
+					GetUgsIssueDiagnosticResponse diagnostic = new GetUgsIssueDiagnosticResponse(buildId, eventData.Message, url);
+					diagnostics.Add(diagnostic);
+				}
+			}
+
+			return diagnostics;
+		}
+
+		/// <summary>
+		/// Gets the URL for a failing step in the
+		/// </summary>
+		/// <param name="details">The issue to get a URL for</param>
+		/// <param name="notify">Whether to show notifications for this issue</param>
+		/// <returns>The issue response</returns>
+		async Task<GetUgsIssueResponse> CreateIssueResponseAsync(IIssueDetails details, bool notify)
+		{
+			Uri? buildUrl = GetIssueBuildUrl(details);
+
+			IUser? owner = details.Issue.OwnerId.HasValue ? await _userCollection.GetCachedUserAsync(details.Issue.OwnerId.Value) : null;
+			IUser? nominatedBy = details.Issue.NominatedById.HasValue ? await _userCollection.GetCachedUserAsync(details.Issue.NominatedById.Value) : null;
+
+			return new GetUgsIssueResponse(details, owner, nominatedBy, notify, buildUrl);
+		}
+
+		/// <summary>
+		/// Gets the URL for a failing step in the given issue
+		/// </summary>
+		/// <param name="issue">The issue details</param>
+		/// <returns>The build URL</returns>
+		Uri? GetIssueBuildUrl(IIssueDetails issue)
+		{
+			HashSet<ObjectId> unresolvedSpans = new HashSet<ObjectId>(issue.Spans.Where(x => x.NextSuccess == null).Select(x => x.Id));
+
+			IIssueStep? step = issue.Steps.OrderByDescending(x => unresolvedSpans.Contains(x.SpanId)).ThenByDescending(x => x.Change).FirstOrDefault();
+			if (step == null)
+			{
+				return null;
+			}
+
+			return new Uri(_settings.DashboardUrl, $"job/{step.JobId}?step={step.StepId}");
+		}
+
+		/// <summary>
+		/// Creates the response for a particular build
+		/// </summary>
+		/// <param name="span">Span containing the step</param>
+		/// <param name="step">The step to describe</param>
+		/// <param name="outcome">Outcome of this step</param>
+		/// <returns>Response object</returns>
+		GetUgsIssueBuildResponse CreateBuildResponse(IIssueSpan span, IIssueStep step, IssueBuildOutcome outcome)
+		{
+			GetUgsIssueBuildResponse response = new GetUgsIssueBuildResponse(span.StreamName, step.Change, outcome);
+			response.Id = step.LogId.GetHashCode();
+			response.JobName = $"{step.JobName}: {span.NodeName}";
+			response.JobUrl = new Uri(_settings.DashboardUrl, $"job/{step.JobId}");
+			response.JobStepName = span.NodeName;
+			response.JobStepUrl = new Uri(_settings.DashboardUrl, $"job/{step.JobId}?step={step.StepId}");
+			response.ErrorUrl = response.JobStepUrl;
+			return response;
+		}
+
+		/// <summary>
+		/// Update an issue
+		/// </summary>
+		/// <param name="issueId">Id of the issue to get information about</param>
+		/// <param name="request">The update information</param>
+		/// <returns>List of matching agents</returns>
+		[HttpPut]
+		[Route("/ugs/api/issues/{issueId}")]
+		public async Task<ActionResult> UpdateIssueAsync(int issueId, [FromBody] UpdateUgsIssueRequest request)
+		{
+			UserId? newOwnerId = null;
+			if (!String.IsNullOrEmpty(request.Owner))
+			{
+				newOwnerId = (await _userCollection.FindOrAddUserByLoginAsync(request.Owner))?.Id;
+			}
+
+			UserId? newNominatedById = null;
+			if (!String.IsNullOrEmpty(request.NominatedBy))
+			{
+				newNominatedById = (await _userCollection.FindOrAddUserByLoginAsync(request.NominatedBy))?.Id;
+			}
+
+			UserId? newDeclinedById = null;
+			if (!String.IsNullOrEmpty(request.DeclinedBy))
+			{
+				newDeclinedById = (await _userCollection.FindOrAddUserByLoginAsync(request.DeclinedBy))?.Id;
+			}
+
+			UserId? newResolvedById = null;
+			if (!String.IsNullOrEmpty(request.ResolvedBy))
+			{
+				newResolvedById = (await _userCollection.FindOrAddUserByLoginAsync(request.ResolvedBy))?.Id;
+			}
+			if (newResolvedById == null && request.Resolved.HasValue)
+			{
+				newResolvedById = request.Resolved.Value ? IIssue.ResolvedByUnknownId : UserId.Empty;
+			}
+
+			if (!await _issueService.UpdateIssueAsync(issueId, ownerId: newOwnerId, nominatedById: newNominatedById, acknowledged: request.Acknowledged, declinedById: newDeclinedById, fixChange: request.FixChange, resolvedById: newResolvedById, initiatedById: User.GetUserId()))
+			{
+				return NotFound();
+			}
+			return Ok();
+		}
+
+		/// <summary>
+		/// Post information about net core installation. 
+		/// </summary>
+		[HttpPost]
+		[Route("/ugs/api/netcore")]
+		public ActionResult<object> PostNetCoreInfo(string? user = null, string? machine = null, bool netCore = false)
+		{
+			_logger.LogInformation("NetCore: User={User}, Machine={Machine}, NetCore={NetCore}", user, machine, netCore);
+			return new { };
+		}
+	}
+}
diff --git a/Engine/Source/Programs/Horde/Horde.Server/Users/IServiceAccountCollection.cs b/Engine/Source/Programs/Horde/Horde.Server/Users/IServiceAccountCollection.cs
new file mode 100644
--- /dev/null
+++ b/Engine/Source/Programs/Horde/Horde.Server/Users/IServiceAccountCollection.cs
@@ -0,0 +1,54 @@
+// Copyright Epic Games, Inc. All Rights Reserved.
+
+using System.Collections.Generic;
+using System.Threading.Tasks;
+using MongoDB.Bson;
+
+namespace Horde.Server.Users
+{
+	/// <summary>
+	/// Interface for a collection of service account documents
+	/// </summary>
+	public interface IServiceAccountCollection
+	{
+		/// <summary>
+		/// Adds a new service account to the collection
+		/// </summary>
+		/// <param name="secretToken">Secret token used for identifying this service account</param>
+		/// <param name="claims">Dictionary of claims to assign</param>
+		/// <param name="description">Description of the account</param>
+		Task<IServiceAccount> AddAsync(string secretToken, List<string> claims, string description);
+
+		/// <summary>
+		/// Get service account via ID
+		/// </summary>
+		/// <param name="id">The unique service account id</param>
+		/// <returns>The service account</returns>
+		Task<IServiceAccount?> GetAsync(ObjectId id);
+		
+		/// <summary>
+		/// Get service account via secret token
+		/// </summary>
+		/// <param name="secretToken">Secret token to use for searching</param>
+		/// <returns>The service account</returns>
+		Task<IServiceAccount?> GetBySecretTokenAsync(string secretToken);
+
+		/// <summary>
+		/// Update a service account from the collection
+		/// </summary>
+		/// <param name="id">The service account</param>
+		/// <param name="secretToken">If set, secret token will be set</param>
+		/// <param name="claims">If set, claims will be set</param>
+		/// <param name="enabled">If set, enabled flag will be set</param>
+		/// <param name="description">If set, description will be set</param>
+		/// <returns>Async task</returns>
+		Task UpdateAsync(ObjectId id, string? secretToken, List<string>? claims, bool? enabled, string? description);
+
+		/// <summary>
+		/// Delete a service account from the collection
+		/// </summary>
+		/// <param name="id">The service account</param>
+		/// <returns>Async task</returns>
+		Task DeleteAsync(ObjectId id);
+	}
+}
diff --git a/Engine/Source/Programs/Horde/Horde.Server/Users/IUserCollection.cs b/Engine/Source/Programs/Horde/Horde.Server/Users/IUserCollection.cs
new file mode 100644
--- /dev/null
+++ b/Engine/Source/Programs/Horde/Horde.Server/Users/IUserCollection.cs
@@ -0,0 +1,122 @@
+// Copyright Epic Games, Inc. All Rights Reserved.
+
+using System.Collections.Generic;
+using System.Security.Claims;
+using System.Threading.Tasks;
+using Horde.Server.Jobs;
+using Horde.Server.Utilities;
+using MongoDB.Bson;
+
+namespace Horde.Server.Users
+{
+	/// <summary>
+	/// Manages user documents
+	/// </summary>
+	public interface IUserCollection
+	{
+		/// <summary>
+		/// Gets a user by unique id
+		/// </summary>
+		/// <param name="id">Id of the user</param>
+		/// <returns>The user information</returns>
+		Task<IUser?> GetUserAsync(UserId id);
+
+		/// <summary>
+		/// Gets a cached user by unique id
+		/// </summary>
+		/// <param name="id">Id of the user</param>
+		/// <returns>The user information</returns>
+		ValueTask<IUser?> GetCachedUserAsync(UserId? id);
+
+		/// <summary>
+		/// Gets a user by unique id
+		/// </summary>
+		/// <param name="ids">Ids of the users</param>
+		/// <param name="nameRegex">Name regex to match for the users</param>
+		/// <param name="index">Maximum number of results</param>
+		/// <param name="count">Number of results to return</param>
+		/// <returns>The user information</returns>
+		Task<List<IUser>> FindUsersAsync(IEnumerable<UserId>? ids = null, string? nameRegex = null, int? index = null, int? count = null);
+
+		/// <summary>
+		/// Gets a user by login
+		/// </summary>
+		/// <param name="login">Login for the user</param>
+		/// <returns>The user information</returns>
+		Task<IUser?> FindUserByLoginAsync(string login);
+
+		/// <summary>
+		/// Gets a user by email address
+		/// </summary>
+		/// <param name="email">Email for the user</param>
+		/// <returns>The user information</returns>
+		Task<IUser?> FindUserByEmailAsync(string email);
+
+		/// <summary>
+		/// Find or add a user with the given claims. Claims will be updated if the user exists.
+		/// </summary>
+		/// <param name="login">Login id of the user</param>
+		/// <param name="name">Full name of the user</param>
+		/// <param name="email">Email address of the user</param>
+		/// <returns>The user document</returns>
+		Task<IUser> FindOrAddUserByLoginAsync(string login, string? name = null, string? email = null);
+
+		/// <summary>
+		/// Gets the claims for a user
+		/// </summary>
+		/// <param name="userId"></param>
+		/// <returns></returns>
+		Task<IUserClaims> GetClaimsAsync(UserId userId);
+
+		/// <summary>
+		/// Update the claims for a user
+		/// </summary>
+		/// <param name="userId"></param>
+		/// <param name="claims"></param>
+		/// <returns></returns>
+		Task UpdateClaimsAsync(UserId userId, IEnumerable<IUserClaim> claims);
+
+		/// <summary>
+		/// Get settings for a user
+		/// </summary>
+		/// <param name="userId"></param>
+		/// <returns></returns>
+		Task<IUserSettings> GetSettingsAsync(UserId userId);
+
+		/// <summary>
+		/// Update a user
+		/// </summary>
+		/// <param name="userId">The user to update</param>
+		/// <param name="enableExperimentalFeatures"></param>
+		/// <param name="dashboardSettings">Opaque settings object for the dashboard</param>
+		/// <param name="addPinnedJobIds"></param>
+		/// <param name="removePinnedJobIds"></param>
+		/// <returns>Updated user object</returns>
+		Task UpdateSettingsAsync(UserId userId, bool? enableExperimentalFeatures = null, BsonValue? dashboardSettings = null, IEnumerable<JobId>? addPinnedJobIds = null, IEnumerable<JobId>? removePinnedJobIds = null);
+	}
+
+	/// <summary>
+	/// Extension methods for the user collect
+	/// </summary>
+	static class UserCollectionExtensions
+	{
+		/// <summary>
+		/// Gets a particular user info from the collection
+		/// </summary>
+		/// <param name="userCollection"></param>
+		/// <param name="principal"></param>
+		/// <returns></returns>
+		public static Task<IUser?> GetUserAsync(this IUserCollection userCollection, ClaimsPrincipal principal)
+		{
+			UserId? userId = principal.GetUserId();
+			if (userId != null)
+			{
+				return userCollection.GetUserAsync(userId.Value);
+			}
+			else
+			{
+				return Task.FromResult<IUser?>(null);
+			}
+		}
+	}
+}
diff --git a/Engine/Source/Programs/Horde/Horde.Server/Users/ServiceAccountCollection.cs b/Engine/Source/Programs/Horde/Horde.Server/Users/ServiceAccountCollection.cs
new file mode 100644
--- /dev/null
+++ b/Engine/Source/Programs/Horde/Horde.Server/Users/ServiceAccountCollection.cs
@@ -0,0 +1,164 @@
+// Copyright Epic Games, Inc. All Rights Reserved.
+
+using System;
+using System.Collections.Generic;
+using System.Linq;
+using System.Threading.Tasks;
+using Horde.Server.Server;
+using MongoDB.Bson;
+using MongoDB.Bson.Serialization.Attributes;
+using MongoDB.Driver;
+
+namespace Horde.Server.Users
+{
+	/// <summary>
+	/// Collection of service account documents
+	/// </summary>
+	public class ServiceAccountCollection : IServiceAccountCollection
+	{
+		/// <summary>
+		/// Concrete implementation of IServiceAccount
+		/// </summary>
+		class ServiceAccountDocument : IServiceAccount
+		{
+			public const string ClaimSeparator = "###";
+			
+			[BsonRequired, BsonId]
+			public ObjectId Id { get; set; }
+
+			[BsonRequired]
+			public string SecretToken { get; set; } = "empty";
+
+			[BsonRequired]
+			public List<string> Claims { get; set; } = new List<string>();
+			
+			public bool Enabled { get; set; }
+			public string Description { get; set; } = "empty";
+
+			[BsonConstructor]
+			private ServiceAccountDocument()
+			{
+			}
+
+			public ServiceAccountDocument(ObjectId id, string secretToken, List<string> claims, bool enabled, string description)
+			{
+				Id = id;
+				SecretToken = secretToken;
+				Claims = claims;
+				Enabled = enabled;
+				Description = description;
+			}
+			
+			/// <inheritdoc/>
+			public void AddClaim(string type, string value)
+			{
+				Claims.Add(type + ClaimSeparator + value);
+			}
+
+			/// <inheritdoc/>
+			public IReadOnlyList<(string Type, string Value)> GetClaims()
+			{
+				return Claims.Select(x =>
+				{
+					string[] split = x.Split(ClaimSeparator);
+					return (split[0], split[1]);
+				}).ToList();
+			}
+
+			protected bool Equals(ServiceAccountDocument other)
+			{
+				bool areClaimsEqual = !Claims.Except(other.Claims).Any();
+
+                return Id.Equals(other.Id) && SecretToken == other.SecretToken && areClaimsEqual && Enabled == other.Enabled && Description == other.Description;
+			}
+
+			public override bool Equals(object? obj)
+			{
+				if (obj is null)
+				{
+					return false;
+				}
+				if (ReferenceEquals(this, obj))
+				{
+					return true;
+				}
+				if (obj.GetType() != GetType())
+				{
+					return false;
+				}
+				return Equals((ServiceAccountDocument) obj);
+			}
+
+			public override int GetHashCode()
+			{
+				return HashCode.Combine(Id, SecretToken, Claims, Enabled, Description);
+			}
+		}
+
+		/// <summary>
+		/// Collection of session documents
+		/// </summary>
+		private readonly IMongoCollection<ServiceAccountDocument> _serviceAccounts;
+
+		/// <summary>
+		/// Constructor
+		/// </summary>
+		/// <param name="mongoService">The database service</param>
+		public ServiceAccountCollection(MongoService mongoService)
+		{
+			_serviceAccounts = mongoService.GetCollection<ServiceAccountDocument>("ServiceAccounts", keys => keys.Ascending(x => x.SecretToken));
+		}
+
+		/// <inheritdoc/>
+		public async Task<IServiceAccount> AddAsync(string secretToken, List<string> claims, string description)
+		{
+			ServiceAccountDocument newSession = new ServiceAccountDocument(ObjectId.GenerateNewId(), secretToken, claims, true, description);
+			await _serviceAccounts.InsertOneAsync(newSession);
+			return newSession;
+		}
+
+		/// <inheritdoc/>
+		public async Task<IServiceAccount?> GetAsync(ObjectId id)
+		{
+			return await _serviceAccounts.Find(x => x.Id == id).FirstOrDefaultAsync();
+		}
+
+		/// <inheritdoc/>
+		public async Task<IServiceAccount?> GetBySecretTokenAsync(string secretToken)
+		{
+			return await _serviceAccounts.Find(x => x.SecretToken == secretToken).FirstOrDefaultAsync();
+		}
+
+		/// <inheritdoc/>
+		public Task UpdateAsync(ObjectId id, string? secretToken, List<string>? claims, bool? enabled, string? description)
+		{
+			UpdateDefinitionBuilder<ServiceAccountDocument> update = Builders<ServiceAccountDocument>.Update;
+			List<UpdateDefinition<ServiceAccountDocument>> updates = new List<UpdateDefinition<ServiceAccountDocument>>();
+
+			if (secretToken != null)
+			{
+				updates.Add(update.Set(x => x.SecretToken, secretToken));
+			}
+			if (claims != null)
+			{
+				updates.Add(update.Set(x => x.Claims, claims));
+			}
+			if (enabled != null)
+			{
+				updates.Add(update.Set(x => x.Enabled, enabled));
+			}
+			if (description != null)
+			{
+				updates.Add(update.Set(x => x.Description, description));
+			}
+
+			return _serviceAccounts.FindOneAndUpdateAsync(x => x.Id == id, update.Combine(updates));
+		}
+
+		/// <inheritdoc/>
+		public Task DeleteAsync(ObjectId sessionId)
+		{
+			return _serviceAccounts.DeleteOneAsync(x => x.Id == sessionId);
+		}
+	}
+}
diff --git a/Engine/Source/Programs/Horde/Horde.Server/Users/UserCollection.cs b/Engine/Source/Programs/Horde/Horde.Server/Users/UserCollection.cs
new file mode 100644
--- /dev/null
+++ b/Engine/Source/Programs/Horde/Horde.Server/Users/UserCollection.cs
@@ -0,0 +1,210 @@
+// Copyright Epic Games, Inc. All Rights Reserved.
+
+using System;
+using System.Collections.Generic;
+using System.Linq;
+using System.Security.Claims;
+using System.Threading.Tasks;
+using Horde.Server.Jobs;
+using Horde.Server.Server;
+using Horde.Server.Utilities;
+using MongoDB.Bson;
+using MongoDB.Bson.Serialization;
+using MongoDB.Bson.Serialization.Attributes;
+using MongoDB.Driver;
+
+namespace Horde.Server.Users
+{
+	/// <summary>
+	/// Manages user documents
+	/// </summary>
+	class UserCollectionV1 : IUserCollection
+	{
+		class UserDocument : IUser, IUserClaims, IUserSettings
+		{
+			public UserId Id { get; set; }
+
+			public ClaimDocument PrimaryClaim { get; set; } = null!;
+			public List<ClaimDocument> Claims { get; set; } = new List<ClaimDocument>();
+
+			[BsonDefaultValue(false), BsonIgnoreIfDefault]
+			public bool EnableExperimentalFeatures { get; set; }
+
+			public BsonValue DashboardSettings { get; set; } = BsonNull.Value;
+			public List<JobId> PinnedJobIds { get; set; } = new List<JobId>();
+
+			string IUser.Name => Claims.FirstOrDefault(x => String.Equals(x.Type, "name", StringComparison.Ordinal))?.Value ?? PrimaryClaim.Value;
+			string IUser.Login => Claims.FirstOrDefault(x => String.Equals(x.Type, ClaimTypes.Name, StringComparison.Ordinal))?.Value ?? PrimaryClaim.Value;
+			string? IUser.Email => Claims.FirstOrDefault(x => String.Equals(x.Type, ClaimTypes.Email, StringComparison.Ordinal))?.Value;
+
+			UserId IUserClaims.UserId => Id;
+			IReadOnlyList<IUserClaim> IUserClaims.Claims => Claims;
+
+			UserId IUserSettings.UserId => Id;
+			IReadOnlyList<JobId> IUserSettings.PinnedJobIds => PinnedJobIds;
+		}
+
+		class ClaimDocument : IUserClaim
+		{
+			public string Type { get; set; }
+			public string Value { get; set; }
+
+			private ClaimDocument()
+			{
+				Type = null!;
+				Value = null!;
+			}
+
+			public ClaimDocument(string type, string value)
+			{
+				Type = type;
+				Value = value;
+			}
+
+			public ClaimDocument(IUserClaim other)
+			{
+				Type = other.Type;
+				Value = other.Value;
+			}
+		}
+
+		readonly IMongoCollection<UserDocument> _users;
+
+		/// <summary>
+		/// Static constructor
+		/// </summary>
+		static UserCollectionV1()
+		{
+			BsonSerializer.RegisterDiscriminatorConvention(typeof(UserDocument), NullDiscriminatorConvention.Instance);
+			BsonSerializer.RegisterDiscriminatorConvention(typeof(ClaimDocument), NullDiscriminatorConvention.Instance);
+		}
+
+		/// <summary>
+		/// Constructor
+		/// </summary>
+		/// <param name="mongoService"></param>
+		public UserCollectionV1(MongoService mongoService)
+		{
+			_users = mongoService.GetCollection<UserDocument>("Users", keys => keys.Ascending(x => x.PrimaryClaim), unique: true);
+		}
+
+		/// <inheritdoc/>
+		public async Task<IUser?> GetUserAsync(UserId id)
+		{
+			return await _users.Find(x => x.Id == id).FirstOrDefaultAsync();
+		}
+
+		/// <inheritdoc/>
+		public async ValueTask<IUser?> GetCachedUserAsync(UserId? id)
+		{
+			if (id == null)
+			{
+				return null;
+			}
+			else
+			{
+				return await GetUserAsync(id.Value);
+			}
+		}
+
+		/// <inheritdoc/>
+		public async Task<List<IUser>> FindUsersAsync(IEnumerable<UserId>? ids, string? nameRegex, int? index, int? count)
+		{
+			FilterDefinition<UserDocument> filter = Builders<UserDocument>.Filter.In(x => x.Id, ids);
+			return await _users.Find(filter).Range(index, count).ToListAsync<UserDocument, IUser>();
+		}
+
+		/// <inheritdoc/>
+		public async Task<IUser?> FindUserByLoginAsync(string login)
+		{
+			ClaimDocument primaryClaim = new ClaimDocument(ClaimTypes.Name, login);
+			return await _users.Find(x => x.PrimaryClaim == primaryClaim).FirstOrDefaultAsync();
+		}
+
+		/// <inheritdoc/>
+		public Task<IUser?> FindUserByEmailAsync(string email)
+		{
+			return Task.FromResult<IUser?>(null);
+		}
+
+		/// <inheritdoc/>
+		public async Task<IUser> FindOrAddUserByLoginAsync(string login, string? name, string? email)
+		{
+			ClaimDocument newPrimaryClaim = new ClaimDocument(ClaimTypes.Name, login);
+			UpdateDefinition<UserDocument> update = Builders<UserDocument>.Update.SetOnInsert(x => x.Id, UserId.GenerateNewId());
+			return await _users.FindOneAndUpdateAsync<UserDocument>(x => x.PrimaryClaim == newPrimaryClaim, update, new FindOneAndUpdateOptions<UserDocument> { IsUpsert = true, ReturnDocument = ReturnDocument.After });
+		}
+
+		/// <inheritdoc/>
+		public async Task<IUserClaims> GetClaimsAsync(UserId userId)
+		{
+			return await _users.Find(x => x.Id == userId).FirstOrDefaultAsync() ?? new UserDocument { Id = userId };
+		}
+
+		/// <inheritdoc/>
+		public async Task UpdateClaimsAsync(UserId userId, IEnumerable<IUserClaim> claims)
+		{
+			List<ClaimDocument> newClaims = claims.Select(x => new ClaimDocument(x)).ToList();
+			await _users.FindOneAndUpdateAsync(x => x.Id == userId, Builders<UserDocument>.Update.Set(x => x.Claims, newClaims));
+		}
+
+		/// <inheritdoc/>
+		public async Task<IUserSettings> GetSettingsAsync(UserId userId)
+		{
+			return await _users.Find(x => x.Id == userId).FirstOrDefaultAsync() ?? new UserDocument { Id = userId };
+		}
+
+		/// <inheritdoc/>
+		public async Task UpdateSettingsAsync(UserId userId, bool? enableExperimentalFeatures, BsonValue? dashboardSettings = null, IEnumerable<JobId>? addPinnedJobIds = null, IEnumerable<JobId>? removePinnedJobIds = null)
+		{
+			if (addPinnedJobIds != null)
+			{
+				foreach (JobId pinnedJobId in addPinnedJobIds)
+				{
+					FilterDefinition<UserDocument> filter = Builders<UserDocument>.Filter.Eq(x => x.Id, userId) & Builders<UserDocument>.Filter.AnyNin<JobId>(x => x.PinnedJobIds, new[] { pinnedJobId });
+					UpdateDefinition<UserDocument> update = Builders<UserDocument>.Update.PushEach(x => x.PinnedJobIds, new[] { pinnedJobId }, -50);
+					await _users.UpdateOneAsync(filter, update);
+				}
+			}
+
+			List<UpdateDefinition<UserDocument>> updates = new List<UpdateDefinition<UserDocument>>();
+			if (enableExperimentalFeatures != null)
+			{
+				updates.Add(Builders<UserDocument>.Update.SetOrUnsetNull(x => x.EnableExperimentalFeatures, enableExperimentalFeatures));
+			}
+			if (dashboardSettings != null)
+			{
+				updates.Add(Builders<UserDocument>.Update.Set(x => x.DashboardSettings, dashboardSettings));
+			}
+			if (removePinnedJobIds != null && removePinnedJobIds.Any())
+			{
+				updates.Add(Builders<UserDocument>.Update.PullAll(x => x.PinnedJobIds, removePinnedJobIds));
+			}
+			if (updates.Count > 0)
+			{
+				await _users.UpdateOneAsync<UserDocument>(x => x.Id == userId, Builders<UserDocument>.Update.Combine(updates));
+			}
+		}
+
+		/// <summary>
+		/// Enumerate all the documents in this collection
+		/// </summary>
+		/// <returns></returns>
+		public async IAsyncEnumerable<(IUser, IUserClaims, IUserSettings)> EnumerateDocumentsAsync()
+		{
+			using (IAsyncCursor<UserDocument> cursor = await _users.Find(FilterDefinition<UserDocument>.Empty).ToCursorAsync())
+			{
+				while (await cursor.MoveNextAsync())
+				{
+					foreach (UserDocument document in cursor.Current)
+					{
+						if (document.Claims.Count > 0)
+						{
+							yield return (document, document, document);
+						}
+					}
+				}
+			}
+		}
+	}
+}
diff --git a/Engine/Source/Programs/Horde/Horde.Server/Users/UserController.cs b/Engine/Source/Programs/Horde/Horde.Server/Users/UserController.cs
new file mode 100644
--- /dev/null
+++ b/Engine/Source/Programs/Horde/Horde.Server/Users/UserController.cs
@@ -0,0 +1,80 @@
+// Copyright Epic Games, Inc. All Rights Reserved.
+
+using System.Collections.Generic;
+using System.Linq;
+using System.Threading.Tasks;
+using Horde.Server.Jobs;
+using Horde.Server.Utilities;
+using Microsoft.AspNetCore.Authorization;
+using Microsoft.AspNetCore.Mvc;
+
+namespace Horde.Server.Users
+{
+	/// <summary>
+	/// Controller for the /api/v1/user endpoint
+	/// </summary>
+	[ApiController]
+	[Authorize]
+	[Route("[controller]")]
+	public class UserController : ControllerBase
+	{
+		/// <summary>
+		/// The user collection instance
+		/// </summary>
+		IUserCollection UserCollection { get; set; }
+
+		/// <summary>
+		/// The avatar service
+		/// </summary>
+		IAvatarService? AvatarService { get; set; }
+
+		/// <summary>
+		/// Constructor
+		/// </summary>
+		/// <param name="userCollection"></param>
+		/// <param name="avatarService"></param>
+		public UserController(IUserCollection userCollection, IAvatarService? avatarService)
+		{
+			UserCollection = userCollection;
+			AvatarService = avatarService;
+		}
+
+		/// <summary>
+		/// Gets information about the logged in user
+		/// </summary>
+		/// <returns>Http result code</returns>
+		[HttpGet]
+		[Route("/api/v1/user")]
+		[ProducesResponseType(typeof(List<GetUserResponse>), 200)]
+		public async Task<ActionResult<object>> GetUserAsync([FromQuery] PropertyFilter? filter = null)
+		{
+			IUser? internalUser = await UserCollection.GetUserAsync(User);
+			if (internalUser == null)
+			{
+				return NotFound();
+			}
+
+			IAvatar? avatar = (AvatarService == null)? (IAvatar?)null : await AvatarService.GetAvatarAsync(internalUser);
+			IUserClaims claims = await UserCollection.GetClaimsAsync(internalUser.Id);
+			IUserSettings settings = await UserCollection.GetSettingsAsync(internalUser.Id);
+			return PropertyFilter.Apply(new GetUserResponse(internalUser, avatar, claims, settings), filter);
+		}
+
+		/// <summary>
+		/// Updates the logged in user
+		/// </summary>
+		/// <returns>Http result code</returns>
+		[HttpPut]
+		[Route("/api/v1/user")]
+		public async Task<ActionResult> UpdateUserAsync(UpdateUserRequest request)
+		{
+			UserId? userId = User.GetUserId();
+			if(userId == null)
+			{
+				return BadRequest("Current user does not have a registered profile");
+			}
+			await UserCollection.UpdateSettingsAsync(userId.Value, request.EnableExperimentalFeatures, request.DashboardSettings?.ToBsonValue(), request.AddPinnedJobIds?.Select(x => JobId.Parse(x)), request.RemovePinnedJobIds?.Select(x => JobId.Parse(x)));
+			return Ok();
+		}
+	}
+}
diff --git a/Engine/Source/Programs/Horde/Horde.Server/Users/UserId.cs b/Engine/Source/Programs/Horde/Horde.Server/Users/UserId.cs
new file mode 100644
--- /dev/null
+++ b/Engine/Source/Programs/Horde/Horde.Server/Users/UserId.cs
@@ -0,0 +1,67 @@
+// Copyright Epic Games, Inc. All Rights Reserved.
+
+using System.ComponentModel;
+using EpicGames.Core;
+using Horde.Server.Utilities;
+using MongoDB.Bson;
+
+namespace Horde.Server.Users
+{
+	/// <summary>
+	/// Identifier for a user
+	/// </summary>
+	/// <param name="Id">Id to construct from</param>
+	[JsonSchemaString]
+	[TypeConverter(typeof(ObjectIdTypeConverter<UserId, UserIdConverter>))]
+	[ObjectIdConverter(typeof(UserIdConverter))]
+	public record struct UserId(ObjectId Id)
+	{
+		/// <summary>
+		/// Constant value for empty user id
+		/// </summary>
+		public static UserId Empty { get; } = new UserId(ObjectId.Empty);
+
+		/// <summary>
+		/// Special user id for an anonymous administrator
+		/// </summary>
+		public static UserId Anonymous { get; } = UserId.Parse("63f7d3525119b9aa4c0f035a");
+
+		/// <summary>
+		/// Creates a new <see cref="UserId"/>
+		/// </summary>
+		public static UserId GenerateNewId() => new UserId(ObjectId.GenerateNewId());
+
+		/// <inheritdoc cref="ObjectId.Parse(System.String)"/>
+		public static UserId Parse(string text) => new UserId(ObjectId.Parse(text));
+
+		/// <inheritdoc cref="ObjectId.TryParse(System.String, out ObjectId)"/>
+		public static bool TryParse(string text, out UserId id)
+		{
+			if (ObjectId.TryParse(text, out ObjectId objectId))
+			{
+				id = new UserId(objectId);
+				return true;
+			}
+			else
+			{
+				id = default;
+				return false;
+			}
+		}
+
+		/// <inheritdoc/>
+		public override string ToString() => Id.ToString();
+	}
+
+	/// <summary>
+	/// Converter to and from <see cref="ObjectId"/> instances.
+	/// </summary>
+	class UserIdConverter : ObjectIdConverter<UserId>
+	{
+		/// <inheritdoc/>
+		public override UserId FromObjectId(ObjectId id) => new UserId(id);
+
+		/// <inheritdoc/>
+		public override ObjectId ToObjectId(UserId value) => value.Id;
+	}
+}
diff --git a/Engine/Source/Programs/Horde/Horde.Server/Users/UserMessages.cs b/Engine/Source/Programs/Horde/Horde.Server/Users/UserMessages.cs
new file mode 100644
--- /dev/null
+++ b/Engine/Source/Programs/Horde/Horde.Server/Users/UserMessages.cs
@@ -0,0 +1,171 @@
+// Copyright Epic Games, Inc. All Rights Reserved.
+
+using System;
+using System.Collections.Generic;
+using System.Linq;
+using System.Text.Json;
+using EpicGames.Core;
+using MongoDB.Bson;
+
+namespace Horde.Server.Users
+{
+	/// <summary>
+	/// Response describing the current user
+	/// </summary>
+	public class GetUserResponse
+	{
+		/// <summary>
+		/// Id of the user
+		/// </summary>
+		public string Id { get; set; }
+
+		/// <summary>
+		/// Name of the user
+		/// </summary>
+		public string Name { get; set; }
+
+		/// <summary>
+		/// Avatar image URL (24px)
+		/// </summary>
+		public string? Image24 { get; set; }
+
+		/// <summary>
+		/// Avatar image URL (32px)
+		/// </summary>
+		public string? Image32 { get; set; }
+
+		/// <summary>
+		/// Avatar image URL (48px)
+		/// </summary>
+		public string? Image48 { get; set; }
+
+		/// <summary>
+		/// Avatar image URL (72px)
+		/// </summary>
+		public string? Image72 { get; set; }
+
+		/// <summary>
+		/// Email of the user
+		/// </summary>
+		public string? Email { get; set; }
+
+		/// <summary>
+		/// Claims for the user
+		/// </summary>
+		public List<UserClaim>? Claims { get; set; }
+
+		/// <summary>
+		/// Whether to enable experimental features for this user
+		/// </summary>
+		public bool? EnableExperimentalFeatures { get; set; }
+
+		/// <summary>
+		/// Settings for the dashboard
+		/// </summary>
+		public object? DashboardSettings { get; set; }
+
+		/// <summary>
+		/// List of pinned job ids
+		/// </summary>
+		public List<string>? PinnedJobIds { get; set; }
+
+		/// <summary>
+		/// Constructor
+		/// </summary>
+		public GetUserResponse(IUser user, IAvatar? avatar, IUserClaims? claims, IUserSettings? settings)
+		{
+			Id = user.Id.ToString();
+			Name = user.Name;
+			Email = user.Email;
+
+			Image24 = avatar?.Image24;
+			Image32 = avatar?.Image32;
+			Image48 = avatar?.Image48;
+			Image72 = avatar?.Image72;
+						
+			Claims = claims?.Claims.Select(x => new UserClaim(x)).ToList();
+
+			if (settings != null)
+			{
+				EnableExperimentalFeatures = settings.EnableExperimentalFeatures;
+				
+				DashboardSettings = BsonTypeMapper.MapToDotNetValue(settings.DashboardSettings);
+				PinnedJobIds = settings.PinnedJobIds.ConvertAll(x => x.ToString());
+			}
+		}
+	}
+
+	/// <summary>
+	/// Basic information about a user. May be embedded in other responses.
+	/// </summary>
+	public class GetThinUserInfoResponse
+	{
+		/// <summary>
+		/// Id of the user
+		/// </summary>
+		public string Id { get; set; }
+
+		/// <summary>
+		/// Name of the user
+		/// </summary>
+		public string Name { get; set; }
+
+		/// <summary>
+		/// The user's email address
+		/// </summary>
+		public string? Email { get; set; }
+
+		/// <summary>
+		/// The user login [DEPRECATED]
+		/// </summary>
+		internal string? Login { get; set; }
+
+		/// <summary>
+		/// Constructor
+		/// </summary>
+		/// <param name="user"></param>
+		public GetThinUserInfoResponse(IUser? user)
+		{
+			if (user == null)
+			{
+				Id = String.Empty;
+				Name = "(Unknown)";
+				Email = null;
+				Login = null;
+			}
+			else
+			{
+				Id = user.Id.ToString();
+				Name = user.Name;
+				Email = user.Email;
+				Login = user.Login;
+			}
+		}
+	}
+
+	/// <summary>
+	/// Request to update settings for a user
+	/// </summary>
+	public class UpdateUserRequest
+	{
+		/// <summary>
+		/// Whether to enable experimental features for this user
+		/// </summary>
+		public bool? EnableExperimentalFeatures { get; set; }
+
+		/// <summary>
+		/// New dashboard settings
+		/// </summary>
+		public JsonElement? DashboardSettings { get; set; }
+
+		/// <summary>
+		/// Job ids to add to the pinned list
+		/// </summary>
+		public List<string>? AddPinnedJobIds { get; set; }
+
+		/// <summary>
+		/// Jobs ids to remove from the pinned list
+		/// </summary>
+		public List<string>? RemovePinnedJobIds { get; set; }
+	}
+}
diff --git a/Engine/Source/Programs/Horde/Horde.Server/Users/UsersController.cs b/Engine/Source/Programs/Horde/Horde.Server/Users/UsersController.cs
new file mode 100644
--- /dev/null
+++ b/Engine/Source/Programs/Horde/Horde.Server/Users/UsersController.cs
@@ -0,0 +1,117 @@
+// Copyright Epic Games, Inc. All Rights Reserved.
+
+using System;
+using System.Collections.Generic;
+using System.Linq;
+using System.Threading.Tasks;
+using Horde.Server.Utilities;
+using Microsoft.AspNetCore.Authorization;
+using Microsoft.AspNetCore.Mvc;
+
+namespace Horde.Server.Users
+{
+	/// <summary>
+	/// Controller for the /api/v1/users endpoint
+	/// </summary>
+	[ApiController]
+	[Authorize]
+	[Route("[controller]")]
+	public class UsersController : HordeControllerBase
+	{
+		/// <summary>
+		/// The user collection instance
+		/// </summary>
+		IUserCollection UserCollection { get; set; }
+
+		/// <summary>
+		/// The avatar service
+		/// </summary>
+		IAvatarService? AvatarService { get; set; }
+
+		/// <summary>
+		/// Constructor
+		/// </summary>
+		/// <param name="userCollection"></param>
+		/// <param name="avatarService"></param>
+		public UsersController(IUserCollection userCollection, IAvatarService? avatarService)
+		{
+			UserCollection = userCollection;
+			AvatarService = avatarService;
+		}
+
+		/// <summary>
+		/// Gets information about a user by id, specify "current" for id to get the currently logged in user
+		/// </summary>
+		/// <returns>Http result code</returns>
+		[HttpGet]
+		[Route("/api/v1/users/{id}")]
+		[ProducesResponseType(typeof(List<GetUserResponse>), 200)]
+		public async Task<ActionResult<object>> GetUserAsync(string id, [FromQuery] PropertyFilter? filter = null)
+		{
+			UserId? userId = ParseUserId(id);
+			if(userId == null)
+			{
+				return BadRequest("Invalid user id '{Id}'", id);
+			}
+
+			IUser? user = await UserCollection.GetUserAsync(userId.Value);
+			if (user == null)
+			{
+				return NotFound(userId.Value);
+			}
+
+			IAvatar? avatar = (AvatarService == null) ? (IAvatar?)null : await AvatarService.GetAvatarAsync(user);
+			IUserClaims? claims = await UserCollection.GetClaimsAsync(user.Id);
+			IUserSettings? settings = await UserCollection.GetSettingsAsync(user.Id);
+			return PropertyFilter.Apply(new GetUserResponse(user, avatar, claims, settings), filter);
+		}
+
+		/// <summary>
+		/// Gets a list of users
+		/// </summary>
+		/// <returns>List of user responses</returns>
+		[HttpGet]
+		[Route("/api/v1/users")]
+		[ProducesResponseType(typeof(List<GetUserResponse>), 200)]
+		public async Task<ActionResult<List<GetUserResponse>>> FindUsersAsync(
+			[FromQuery] string[]? ids = null,
+			[FromQuery] string? nameRegex = null,
+			[FromQuery] int index = 0,
+			[FromQuery] int count = 100,
+			[FromQuery] bool includeClaims = false,
+			[FromQuery] bool includeAvatar = false)
+		{
+
+			UserId[]? userIds = null;
+			if (ids != null && ids.Length > 0)
+			{
+				userIds = ids.Select(x => UserId.Parse(x)).ToArray();
+			}
+
+			List<IUser> users = await UserCollection.FindUsersAsync(userIds, nameRegex, index, count);
+
+			List<GetUserResponse> response = new List<GetUserResponse>();
+			foreach (IUser user in users)
+			{
+				IAvatar? avatar = (AvatarService == null || !includeAvatar) ? (IAvatar?)null : await AvatarService.GetAvatarAsync(user);
+				IUserClaims? claims = (!includeClaims) ? null : await UserCollection.GetClaimsAsync(user.Id);
+				response.Add(new GetUserResponse(user, avatar, claims, null));
+			}
+
+			return response;
+		}
+
+		UserId? ParseUserId(string id)
+		{
+			if (id.Equals("current", StringComparison.OrdinalIgnoreCase))
+			{
+				return User.GetUserId();
+			}
+			else if(UserId.TryParse(id, out UserId result))
+			{
+				return result;
+			}
+			return null;
+		}
+	}
+}
diff --git a/Engine/Source/Programs/Horde/Horde.Server/Utilities/ArtifactChunkStream.cs b/Engine/Source/Programs/Horde/Horde.Server/Utilities/ArtifactChunkStream.cs
new file mode 100644
--- /dev/null
+++ b/Engine/Source/Programs/Horde/Horde.Server/Utilities/ArtifactChunkStream.cs
@@ -0,0 +1,128 @@
+// Copyright Epic Games, Inc. All Rights Reserved.
+
+using System;
+using System.IO;
+using System.Threading;
+using System.Threading.Tasks;
+using Grpc.Core;
+using HordeCommon.Rpc;
+
+namespace Horde.Server.Utilities
+{
+	/// <summary>
+	/// Stream that takes in rpc chunks from an upload and copies them?
+	/// </summary>
+	public class ArtifactChunkStream : Stream
+	{
+		/// <summary>
+		/// The grpc client reader. Should probably be templatized
+		/// </summary>
+		readonly IAsyncStreamReader<UploadArtifactRequest> _reader;
+
+		/// <summary>
+		/// Position within the stream
+		/// </summary>
+		long _streamPosition;
+
+		/// <summary>
+		/// The length of the stream
+		/// </summary>
+		readonly long _streamLength;
+
+		/// <summary>
+		/// The current request being read from
+		/// </summary>
+		UploadArtifactRequest? _request;
+
+		/// <summary>
+		/// Position within the current request
+		/// </summary>
+		int _requestPos;
+
+		/// <summary>
+		/// Constructor
+		/// </summary>
+		/// <param name="reader">the grpc reader</param>
+		/// <param name="length">filesize reported by the client</param>
+		public ArtifactChunkStream(IAsyncStreamReader<UploadArtifactRequest> reader, long length)
+		{
+			_reader = reader;
+			_streamLength = length;
+		}
+
+		/// <inheritdoc/>
+		public override long Position
+		{
+			get => _streamPosition;
+			set => throw new NotSupportedException();
+		}
+
+		/// <inheritdoc/>
+		public override long Length => _streamLength;
+
+		/// <inheritdoc/>
+		public override void SetLength(long value) => throw new NotSupportedException();
+
+		/// <inheritdoc/>
+		public override bool CanSeek => false;
+
+		/// <inheritdoc/>
+		public override long Seek(long offset, SeekOrigin origin) => throw new NotSupportedException();
+
+		/// <inheritdoc/>
+		public override bool CanWrite => false;
+
+		/// <inheritdoc/>
+		public override void Write(byte[] buffer, int offset, int count) => throw new NotSupportedException();
+
+		/// <inheritdoc/>
+		public override bool CanRead => true;
+
+		/// <inheritdoc/>
+		public override int Read(byte[] buffer, int offset, int count)
+		{
+			return ReadAsync(buffer, offset, count, CancellationToken.None).Result;
+		}
+
+		/// <inheritdoc/>
+		public override async Task<int> ReadAsync(byte[] buffer, int offset, int count, CancellationToken cancellationToken)
+		{
+			return await ReadAsync(buffer.AsMemory(offset, count), cancellationToken);
+		}
+
+		/// <inheritdoc/>
+		public override async ValueTask<int> ReadAsync(Memory<byte> buffer, CancellationToken cancellationToken)
+		{
+			int bytesRead = 0;
+			while (bytesRead < buffer.Length && _streamPosition < _streamLength)
+			{
+				if (_request == null || _requestPos == _request.Data.Length)
+				{
+					// Read the next request object
+					if (!await _reader.MoveNext())
+					{
+						throw new EndOfStreamException("Unexpected end of stream while reading artifact");
+					}
+
+					_request = _reader.Current;
+					_requestPos = 0;
+				}
+				else
+				{
+					// Copy data from the current request object
+					int numBytesToCopy = Math.Min(buffer.Length - bytesRead, _request.Data.Length - _requestPos);
+					_request.Data.Span.Slice(_requestPos, numBytesToCopy).CopyTo(buffer.Slice(bytesRead, numBytesToCopy).Span);
+					_requestPos += numBytesToCopy;
+					bytesRead += numBytesToCopy;
+					_streamPosition += numBytesToCopy;
+				}
+			}
+			return bytesRead;
+		}
+
+		/// <inheritdoc/>
+		public override void Flush()
+		{
+		}
+	}
+}
diff --git a/Engine/Source/Programs/Horde/Horde.Server/Utilities/CompressionExtensions.cs b/Engine/Source/Programs/Horde/Horde.Server/Utilities/CompressionExtensions.cs
new file mode 100644
--- /dev/null
+++ b/Engine/Source/Programs/Horde/Horde.Server/Utilities/CompressionExtensions.cs
@@ -0,0 +1,74 @@
+// Copyright Epic Games, Inc. All Rights Reserved.
+
+using System;
+using System.Buffers.Binary;
+using System.Globalization;
+using System.IO;
+using EpicGames.Core;
+using ICSharpCode.SharpZipLib.BZip2;
+using OpenTracing;
+using OpenTracing.Util;
+
+namespace Horde.Server.Utilities
+{
+	/// <summary>
+	/// Extension methods for compression
+	/// </summary>
+	static class CompressionExtensions
+	{
+		/// <summary>
+		/// Compress a block of data with bzip2
+		/// </summary>
+		/// <returns>The compressed data</returns>
+		public static byte[] CompressBzip2(this ReadOnlyMemory<byte> memory)
+		{
+			using IScope scope = GlobalTracer.Instance.BuildSpan("CompressBzip2").StartActive();
+			scope.Span.SetTag("DecompressedSize", memory.Length.ToString(CultureInfo.InvariantCulture));
+
+			byte[] compressedData;
+			using (MemoryStream stream = new MemoryStream())
+			{
+				byte[] decompressedSize = new byte[4];
+				BinaryPrimitives.WriteInt32LittleEndian(decompressedSize.AsSpan(), memory.Length);
+				stream.Write(decompressedSize.AsSpan());
+
+				using (BZip2OutputStream compressedStream = new BZip2OutputStream(stream))
+				{
+					compressedStream.Write(memory.Span);
+				}
+
+				compressedData = stream.ToArray();
+			}
+
+			scope.Span.SetTag("CompressedSize", compressedData.Length.ToString(CultureInfo.InvariantCulture));
+			return compressedData;
+		}
+
+		/// <summary>
+		/// Decompress the data
+		/// </summary>
+		/// <returns>The decompressed data</returns>
+		public static byte[] DecompressBzip2(this ReadOnlyMemory<byte> memory)
+		{
+			int decompressedSize = BinaryPrimitives.ReadInt32LittleEndian(memory.Span);
+
+			using IScope scope = GlobalTracer.Instance.BuildSpan("DecompressBzip2").StartActive();
+			scope.Span.SetTag("CompressedSize", memory.Length.ToString(CultureInfo.InvariantCulture));
+			scope.Span.SetTag("DecompressedSize", decompressedSize.ToString(CultureInfo.InvariantCulture));
+
+			byte[] data = new byte[decompressedSize];
+			using (ReadOnlyMemoryStream stream = new ReadOnlyMemoryStream(memory.Slice(4)))
+			{
+				using (BZip2InputStream decompressedStream = new BZip2InputStream(stream))
+				{
+					int readSize = decompressedStream.Read(data.AsSpan());
+					if (readSize != data.Length)
+					{
+						throw new InvalidDataException($"Compressed data is too short (expected {data.Length} bytes, got {readSize} bytes)");
+					}
+				}
+			}
+			return data;
+		}
+	}
+}
