Horde: Rename the Horde.Build project to Horde.Server.

diff --git a/Engine/Source/Programs/Horde/Horde.Server/Utilities/HordeClaimTypes.cs b/Engine/Source/Programs/Horde/Horde.Server/Utilities/HordeClaimTypes.cs
new file mode 100644
--- /dev/null
+++ b/Engine/Source/Programs/Horde/Horde.Server/Utilities/HordeClaimTypes.cs
@@ -0,0 +1,158 @@
+// Copyright Epic Games, Inc. All Rights Reserved.
+
+using System.Linq;
+using System.Security.Claims;
+using Horde.Server.Users;
+
+namespace Horde.Server.Utilities
+{
+	/// <summary>
+	/// Claim types that are specific to horde
+	/// </summary>
+	public static class HordeClaimTypes
+	{
+		/// <summary>
+		/// Base URI for all Horde claims.
+		/// </summary>
+		const string Prefix = "http://epicgames.com/ue/horde/";
+
+		/// <summary>
+		/// Claim for a particular role.
+		/// </summary>
+		public const string Role = Prefix + "role";
+
+		/// <summary>
+		/// Claim for a well known internal role (eg. admin)
+		/// </summary>
+		public const string InternalRole = Prefix + "internal-role";
+
+		/// <summary>
+		/// Claim for a particular session
+		/// </summary>
+		public const string AgentId = Prefix + "agent";
+
+		/// <summary>
+		/// Claim for a particular session
+		/// </summary>
+		public const string AgentSessionId = Prefix + "session";
+
+		/// <summary>
+		/// User name for horde
+		/// </summary>
+		public const string User = Prefix + "user";
+
+		/// <summary>
+		/// Unique id of the horde user (see <see cref="IUserCollection"/>)
+		/// </summary>
+		public const string UserId = Prefix + "user-id-v3";
+
+		/// <summary>
+		/// Claim for downloading artifacts from a particular job
+		/// </summary>
+		public const string JobArtifacts = Prefix + "job-artifacts";
+
+		/// <summary>
+		/// Claim for the Perforce username
+		/// </summary>
+		public const string PerforceUser = Prefix + "perforce-user";
+
+		/// <summary>
+		/// Claim for the external issue username
+		/// </summary>
+		public const string ExternalIssueUser = Prefix + "external-issue-user";
+
+		/// <summary>
+		/// Claim identifying an agent as executing a particular lease
+		/// </summary>
+		public const string Lease = Prefix + "lease";
+
+		/// <summary>
+		/// Claim allowing a certain namespace and/or path to be read from. See <see cref="WriteNamespace"/>.
+		/// </summary>
+		public const string ReadNamespace = Prefix + "read-namespace";
+
+		/// <summary>
+		/// Claim allowing a certain namespace and/or path to be written to. Value may be a namespace name, or a subpath (eg. ns:a/b will only allow blobs and refs with an a/b prefix to be written).
+		/// </summary>
+		public const string WriteNamespace = Prefix + "write-namespace";
+
+		/// <summary>
+		/// Claim allowing a certain ref to be written
+		/// </summary>
+		public const string WriteRef = Prefix + "write-ref";
+	}
+
+	/// <summary>
+	/// Extension methods for getting Horde claims from a principal
+	/// </summary>
+	public static class HordeClaimExtensions
+	{
+		/// <summary>
+		/// Gets the Horde user id from a principal
+		/// </summary>
+		/// <param name="principal"></param>
+		/// <returns></returns>
+		public static UserId? GetUserId(this ClaimsPrincipal principal)
+		{
+			string? idValue = principal.FindFirstValue(HordeClaimTypes.UserId);
+			if(idValue == null)
+			{
+				return null;
+			}
+			else
+			{
+				return UserId.Parse(idValue);
+			}
+		}
+
+		/// <summary>
+		/// Gets the Horde user name from a principal
+		/// </summary>
+		/// <param name="principal"></param>
+		/// <returns></returns>
+		public static string? GetUserName(this ClaimsPrincipal principal)
+		{
+			return principal.FindFirstValue(HordeClaimTypes.User) ?? principal.FindFirstValue(ClaimTypes.Name);
+		}
+
+		/// <summary>
+		/// Get the email for the given user
+		/// </summary>
+		/// <param name="user">The user to query the email for</param>
+		/// <returns>The user's email address or null if not found</returns>
+		public static string? GetEmail(this ClaimsPrincipal user)
+		{
+			return user.Claims.FirstOrDefault(x => x.Type == ClaimTypes.Email)?.Value;
+		}
+
+		/// <summary>
+		/// Gets the perforce username for the given principal
+		/// </summary>
+		/// <param name="user">The principal to get the Perforce user for</param>
+		/// <returns>Perforce user name</returns>
+		public static string? GetPerforceUser(this ClaimsPrincipal user)
+		{
+			Claim? claim = user.FindFirst(HordeClaimTypes.PerforceUser);
+			if (claim == null)
+			{
+				return null;
+			}
+			return claim.Value;
+		}
+
+		/// <summary>
+		/// Gets the external issue username for the given principal
+		/// </summary>
+		/// <param name="user">The principal to get the external issue user for</param>
+		/// <returns>Jira user name</returns>
+		public static string? GetExternalIssueUser(this ClaimsPrincipal user)
+		{
+			Claim? claim = user.FindFirst(HordeClaimTypes.ExternalIssueUser);
+			if (claim == null)
+			{
+				return null;
+			}
+			return claim.Value;
+		}
+	}
+}
diff --git a/Engine/Source/Programs/Horde/Horde.Server/Utilities/MongoTracingCollection.cs b/Engine/Source/Programs/Horde/Horde.Server/Utilities/MongoTracingCollection.cs
new file mode 100644
--- /dev/null
+++ b/Engine/Source/Programs/Horde/Horde.Server/Utilities/MongoTracingCollection.cs
@@ -0,0 +1,771 @@
+// Copyright Epic Games, Inc. All Rights Reserved.
+
+using System.Collections.Generic;
+using System.Threading;
+using System.Threading.Tasks;
+using MongoDB.Bson;
+using MongoDB.Bson.Serialization;
+using MongoDB.Driver;
+using OpenTracing;
+using OpenTracing.Util;
+
+namespace Horde.Server.Utiltiies
+{
+	/// <summary>
+	/// Wrap a IMongoCollection with trace scopes
+	///
+	/// Will capture the entire invocation of MongoDB calls, including serialization.
+	/// The other command logging (not in this file) only deals with queries sent to the server at the protocol level.
+	/// </summary>
+	/// <typeparam name="T">A MongoDB document</typeparam>
+	public class MongoTracingCollection<T> : IMongoCollection<T>
+	{
+		private readonly IMongoCollection<T> _collection;
+		private const string Prefix = "MongoCollection.";
+
+		/// <summary>
+		/// Constructor
+		/// </summary>
+		/// <param name="collection">Collection to wrap with tracing</param>
+		public MongoTracingCollection(IMongoCollection<T> collection)
+		{
+			_collection = collection;
+		}
+
+#pragma warning disable CS0618		
+		/// <inheritdoc />
+		public IAsyncCursor<TResult> Aggregate<TResult>(PipelineDefinition<T, TResult> pipeline, AggregateOptions options = null!,
+			CancellationToken  cancellationToken = new CancellationToken())
+		{
+			return _collection.Aggregate(pipeline, options,  cancellationToken);
+		}
+
+		/// <inheritdoc />
+		public IAsyncCursor<TResult> Aggregate<TResult>(IClientSessionHandle session, PipelineDefinition<T, TResult> pipeline,
+			AggregateOptions options = null!, CancellationToken  cancellationToken = new CancellationToken())
+		{
+			return _collection.Aggregate(session, pipeline, options,  cancellationToken);
+		}
+
+		/// <inheritdoc />
+		public async Task<IAsyncCursor<TResult>> AggregateAsync<TResult>(PipelineDefinition<T, TResult> pipeline, AggregateOptions options = null!,
+			CancellationToken  cancellationToken = new CancellationToken())
+		{
+			using IScope scope = GlobalTracer.Instance.BuildSpan(Prefix + "AggregateAsync").StartActive();
+			scope.Span.SetTag("CollectionName", _collection.CollectionNamespace.CollectionName);
+			return await _collection.AggregateAsync(pipeline, options,  cancellationToken);
+		}
+
+		/// <inheritdoc />
+		public async Task<IAsyncCursor<TResult>> AggregateAsync<TResult>(IClientSessionHandle session, PipelineDefinition<T, TResult> pipeline, AggregateOptions options = null!,
+			CancellationToken  cancellationToken = new CancellationToken())
+		{
+			using IScope scope = GlobalTracer.Instance.BuildSpan(Prefix + "AggregateAsync").StartActive();
+			scope.Span.SetTag("CollectionName", _collection.CollectionNamespace.CollectionName);
+			return await _collection.AggregateAsync(session, pipeline, options, cancellationToken);
+		}
+
+		/// <inheritdoc />
+		public void AggregateToCollection<TResult>(PipelineDefinition<T, TResult> pipeline, AggregateOptions? options = null, CancellationToken cancellationToken = default)
+		{
+			using IScope scope = GlobalTracer.Instance.BuildSpan(Prefix + "AggregateToCollection").StartActive();
+			scope.Span.SetTag("CollectionName", _collection.CollectionNamespace.CollectionName);
+			_collection.AggregateToCollection(pipeline, options, cancellationToken);
+		}
+
+		/// <inheritdoc />
+		public void AggregateToCollection<TResult>(IClientSessionHandle session, PipelineDefinition<T, TResult> pipeline, AggregateOptions? options = null, CancellationToken cancellationToken = default)
+		{
+			using IScope scope = GlobalTracer.Instance.BuildSpan(Prefix + "AggregateToCollection").StartActive();
+			scope.Span.SetTag("CollectionName", _collection.CollectionNamespace.CollectionName);
+			_collection.AggregateToCollection(session, pipeline, options, cancellationToken);
+		}
+
+		/// <inheritdoc />
+		public async Task AggregateToCollectionAsync<TResult>(PipelineDefinition<T, TResult> pipeline, AggregateOptions? options = null, CancellationToken cancellationToken = default)
+		{
+			using IScope scope = GlobalTracer.Instance.BuildSpan(Prefix + "AggregateToCollectionAsync").StartActive();
+			scope.Span.SetTag("CollectionName", _collection.CollectionNamespace.CollectionName);
+			await _collection.AggregateToCollectionAsync(pipeline, options, cancellationToken);
+		}
+
+		/// <inheritdoc />
+		public async Task AggregateToCollectionAsync<TResult>(IClientSessionHandle session, PipelineDefinition<T, TResult> pipeline, AggregateOptions? options = null, CancellationToken cancellationToken = default)
+		{
+			using IScope scope = GlobalTracer.Instance.BuildSpan(Prefix + "AggregateToCollectionAsync").StartActive();
+			scope.Span.SetTag("CollectionName", _collection.CollectionNamespace.CollectionName);
+			await _collection.AggregateToCollectionAsync(session, pipeline, options, cancellationToken);
+		}
+
+		/// <inheritdoc />
+		public BulkWriteResult<T> BulkWrite(IEnumerable<WriteModel<T>> requests, BulkWriteOptions options = null!,
+			CancellationToken  cancellationToken = new CancellationToken())
+		{
+			return _collection.BulkWrite(requests, options,  cancellationToken);
+		}
+
+		/// <inheritdoc />
+		public BulkWriteResult<T> BulkWrite(IClientSessionHandle session, IEnumerable<WriteModel<T>> requests, BulkWriteOptions options = null!,
+			CancellationToken  cancellationToken = new CancellationToken())
+		{
+			return _collection.BulkWrite(session, requests, options,  cancellationToken);
+		}
+
+		/// <inheritdoc />
+		public async Task<BulkWriteResult<T>> BulkWriteAsync(IEnumerable<WriteModel<T>> requests, BulkWriteOptions options = null!,
+			CancellationToken  cancellationToken = new CancellationToken())
+		{
+			using IScope scope = GlobalTracer.Instance.BuildSpan(Prefix + "BulkWriteAsync").StartActive();
+			scope.Span.SetTag("CollectionName", _collection.CollectionNamespace.CollectionName);
+			return await _collection.BulkWriteAsync(requests, options,  cancellationToken);
+		}
+
+		/// <inheritdoc />
+		public async Task<BulkWriteResult<T>> BulkWriteAsync(IClientSessionHandle session, IEnumerable<WriteModel<T>> requests, BulkWriteOptions options = null!,
+			CancellationToken  cancellationToken = new CancellationToken())
+		{
+			using IScope scope = GlobalTracer.Instance.BuildSpan(Prefix + "BulkWriteAsync").StartActive();
+			scope.Span.SetTag("CollectionName", _collection.CollectionNamespace.CollectionName);
+			return await _collection.BulkWriteAsync(session, requests, options,  cancellationToken);
+		}
+
+		/// <inheritdoc />
+		public long Count(FilterDefinition<T> filter, CountOptions options = null!,
+			CancellationToken  cancellationToken = new CancellationToken())
+		{
+			return _collection.Count(filter, options,  cancellationToken);
+		}
+
+		/// <inheritdoc />
+		public long Count(IClientSessionHandle session, FilterDefinition<T> filter, CountOptions options = null!,
+			CancellationToken  cancellationToken = new CancellationToken())
+		{
+			return _collection.Count(session, filter, options,  cancellationToken);
+		}
+
+		/// <inheritdoc />
+		public async Task<long> CountAsync(FilterDefinition<T> filter, CountOptions options = null!,
+			CancellationToken  cancellationToken = new CancellationToken())
+		{
+			using IScope scope = GlobalTracer.Instance.BuildSpan(Prefix + "CountAsync").StartActive();
+			scope.Span.SetTag("CollectionName", _collection.CollectionNamespace.CollectionName);
+			return await _collection.CountAsync(filter, options,  cancellationToken);
+		}
+
+		/// <inheritdoc />
+		public async Task<long> CountAsync(IClientSessionHandle session, FilterDefinition<T> filter, CountOptions options = null!,
+			CancellationToken  cancellationToken = new CancellationToken())
+		{
+			using IScope scope = GlobalTracer.Instance.BuildSpan(Prefix + "CountAsync").StartActive();
+			scope.Span.SetTag("CollectionName", _collection.CollectionNamespace.CollectionName);
+			return await _collection.CountAsync(session, filter, options,  cancellationToken);
+		}
+
+		/// <inheritdoc />
+		public long CountDocuments(FilterDefinition<T> filter, CountOptions options = null!,
+			CancellationToken  cancellationToken = new CancellationToken())
+		{
+			return _collection.CountDocuments(filter, options,  cancellationToken);
+		}
+
+		/// <inheritdoc />
+		public long CountDocuments(IClientSessionHandle session, FilterDefinition<T> filter, CountOptions options = null!,
+			CancellationToken  cancellationToken = new CancellationToken())
+		{
+			using IScope scope = GlobalTracer.Instance.BuildSpan(Prefix + "CountDocumentsAsync").StartActive();
+			scope.Span.SetTag("CollectionName", _collection.CollectionNamespace.CollectionName);
+			return _collection.CountDocuments(session, filter, options,  cancellationToken);
+		}
+
+		/// <inheritdoc />
+		public async Task<long> CountDocumentsAsync(FilterDefinition<T> filter, CountOptions options = null!,
+			CancellationToken  cancellationToken = new CancellationToken())
+		{
+			using IScope scope = GlobalTracer.Instance.BuildSpan(Prefix + "CountDocumentsAsync").StartActive();
+			scope.Span.SetTag("CollectionName", _collection.CollectionNamespace.CollectionName);
+			return await _collection.CountDocumentsAsync(filter, options,  cancellationToken);
+		}
+
+		/// <inheritdoc />
+		public async Task<long> CountDocumentsAsync(IClientSessionHandle session, FilterDefinition<T> filter, CountOptions options = null!,
+			CancellationToken  cancellationToken = new CancellationToken())
+		{
+			using IScope scope = GlobalTracer.Instance.BuildSpan(Prefix + "CountDocumentsAsync").StartActive();
+			scope.Span.SetTag("CollectionName", _collection.CollectionNamespace.CollectionName);
+			return await _collection.CountDocumentsAsync(session, filter, options,  cancellationToken);
+		}
+
+		/// <inheritdoc />
+		public DeleteResult DeleteMany(FilterDefinition<T> filter, CancellationToken  cancellationToken = new CancellationToken())
+		{
+			return _collection.DeleteMany(filter,  cancellationToken);
+		}
+
+		/// <inheritdoc />
+		public DeleteResult DeleteMany(FilterDefinition<T> filter, DeleteOptions options,
+			CancellationToken  cancellationToken = new CancellationToken())
+		{
+			return _collection.DeleteMany(filter, options,  cancellationToken);
+		}
+
+		/// <inheritdoc />
+		public DeleteResult DeleteMany(IClientSessionHandle session, FilterDefinition<T> filter, DeleteOptions options = null!,
+			CancellationToken  cancellationToken = new CancellationToken())
+		{
+			return _collection.DeleteMany(session, filter, options,  cancellationToken);
+		}
+
+		/// <inheritdoc />
+		public async Task<DeleteResult> DeleteManyAsync(FilterDefinition<T> filter, CancellationToken  cancellationToken = new CancellationToken())
+		{
+			using IScope scope = GlobalTracer.Instance.BuildSpan(Prefix + "DeleteManyAsync").StartActive();
+			scope.Span.SetTag("CollectionName", _collection.CollectionNamespace.CollectionName);
+			return await _collection.DeleteManyAsync(filter,  cancellationToken);
+		}
+
+		/// <inheritdoc />
+		public async Task<DeleteResult> DeleteManyAsync(FilterDefinition<T> filter, DeleteOptions options,
+			CancellationToken  cancellationToken = new CancellationToken())
+		{
+			using IScope scope = GlobalTracer.Instance.BuildSpan(Prefix + "DeleteManyAsync").StartActive();
+			scope.Span.SetTag("CollectionName", _collection.CollectionNamespace.CollectionName);
+			return await _collection.DeleteManyAsync(filter, options,  cancellationToken);
+		}
+
+		/// <inheritdoc />
+		public async Task<DeleteResult> DeleteManyAsync(IClientSessionHandle session, FilterDefinition<T> filter, DeleteOptions options = null!,
+			CancellationToken  cancellationToken = new CancellationToken())
+		{
+			using IScope scope = GlobalTracer.Instance.BuildSpan(Prefix + "DeleteManyAsync").StartActive();
+			scope.Span.SetTag("CollectionName", _collection.CollectionNamespace.CollectionName);
+			return await _collection.DeleteManyAsync(session, filter, options,  cancellationToken);
+		}
+
+		/// <inheritdoc />
+		public DeleteResult DeleteOne(FilterDefinition<T> filter, CancellationToken  cancellationToken = new CancellationToken())
+		{
+			return _collection.DeleteOne(filter,  cancellationToken);
+		}
+
+		/// <inheritdoc />
+		public DeleteResult DeleteOne(FilterDefinition<T> filter, DeleteOptions options,
+			CancellationToken  cancellationToken = new CancellationToken())
+		{
+			return _collection.DeleteOne(filter, options,  cancellationToken);
+		}
+
+		/// <inheritdoc />
+		public DeleteResult DeleteOne(IClientSessionHandle session, FilterDefinition<T> filter, DeleteOptions options = null!,
+			CancellationToken  cancellationToken = new CancellationToken())
+		{
+			return _collection.DeleteOne(session, filter, options,  cancellationToken);
+		}
+
+		/// <inheritdoc />
+		public async Task<DeleteResult> DeleteOneAsync(FilterDefinition<T> filter, CancellationToken  cancellationToken = new CancellationToken())
+		{
+			using IScope scope = GlobalTracer.Instance.BuildSpan(Prefix + "DeleteOneAsync").StartActive();
+			scope.Span.SetTag("CollectionName", _collection.CollectionNamespace.CollectionName);
+			return await _collection.DeleteOneAsync(filter,  cancellationToken);
+		}
+
+		/// <inheritdoc />
+		public async Task<DeleteResult> DeleteOneAsync(FilterDefinition<T> filter, DeleteOptions options,
+			CancellationToken  cancellationToken = new CancellationToken())
+		{
+			using IScope scope = GlobalTracer.Instance.BuildSpan(Prefix + "DeleteOneAsync").StartActive();
+			scope.Span.SetTag("CollectionName", _collection.CollectionNamespace.CollectionName);
+			return await _collection.DeleteOneAsync(filter, options,  cancellationToken);
+		}
+
+		/// <inheritdoc />
+		public Task<DeleteResult> DeleteOneAsync(IClientSessionHandle session, FilterDefinition<T> filter, DeleteOptions options = null!,
+			CancellationToken  cancellationToken = new CancellationToken())
+		{
+			using IScope scope = GlobalTracer.Instance.BuildSpan(Prefix + "DeleteOneAsync").StartActive();
+			scope.Span.SetTag("CollectionName", _collection.CollectionNamespace.CollectionName);
+			return _collection.DeleteOneAsync(session, filter, options,  cancellationToken);
+		}
+
+		/// <inheritdoc />
+		public IAsyncCursor<TField> Distinct<TField>(FieldDefinition<T, TField> field, FilterDefinition<T> filter, DistinctOptions options = null!,
+			CancellationToken  cancellationToken = new CancellationToken())
+		{
+			return _collection.Distinct(field, filter, options,  cancellationToken);
+		}
+
+		/// <inheritdoc />
+		public IAsyncCursor<TField> Distinct<TField>(IClientSessionHandle session, FieldDefinition<T, TField> field, FilterDefinition<T> filter,
+			DistinctOptions options = null!, CancellationToken  cancellationToken = new CancellationToken())
+		{
+			return _collection.Distinct(session, field, filter, options,  cancellationToken);
+		}
+
+		/// <inheritdoc />
+		public async Task<IAsyncCursor<TField>> DistinctAsync<TField>(FieldDefinition<T, TField> field, FilterDefinition<T> filter, DistinctOptions options = null!,
+			CancellationToken  cancellationToken = new CancellationToken())
+		{
+			using IScope scope = GlobalTracer.Instance.BuildSpan(Prefix + "DistinctAsync").StartActive();
+			scope.Span.SetTag("CollectionName", _collection.CollectionNamespace.CollectionName);
+			return await _collection.DistinctAsync(field, filter, options,  cancellationToken);
+		}
+
+		/// <inheritdoc />
+		public async Task<IAsyncCursor<TField>> DistinctAsync<TField>(IClientSessionHandle session, FieldDefinition<T, TField> field, FilterDefinition<T> filter,
+			DistinctOptions options = null!, CancellationToken  cancellationToken = new CancellationToken())
+		{
+			using IScope scope = GlobalTracer.Instance.BuildSpan(Prefix + "DistinctAsync").StartActive();
+			scope.Span.SetTag("CollectionName", _collection.CollectionNamespace.CollectionName);
+			return await _collection.DistinctAsync(session, field, filter, options,  cancellationToken);
+		}
+
+		/// <inheritdoc />
+		public long EstimatedDocumentCount(EstimatedDocumentCountOptions options = null!,
+			CancellationToken  cancellationToken = new CancellationToken())
+		{
+			return _collection.EstimatedDocumentCount(options,  cancellationToken);
+		}
+
+		/// <inheritdoc />
+		public async Task<long> EstimatedDocumentCountAsync(EstimatedDocumentCountOptions options = null!,
+			CancellationToken  cancellationToken = new CancellationToken())
+		{
+			using IScope scope = GlobalTracer.Instance.BuildSpan(Prefix + "EstimatedDocumentCountAsync").StartActive();
+			scope.Span.SetTag("CollectionName", _collection.CollectionNamespace.CollectionName);
+			return await _collection.EstimatedDocumentCountAsync(options,  cancellationToken);
+		}
+
+		/// <inheritdoc />
+		public IAsyncCursor<TProjection> FindSync<TProjection>(FilterDefinition<T> filter, FindOptions<T, TProjection> options = null!,
+			CancellationToken  cancellationToken = new CancellationToken())
+		{
+			return _collection.FindSync(filter, options,  cancellationToken);
+		}
+
+		/// <inheritdoc />
+		public IAsyncCursor<TProjection> FindSync<TProjection>(IClientSessionHandle session, FilterDefinition<T> filter, FindOptions<T, TProjection> options = null!,
+			CancellationToken  cancellationToken = new CancellationToken())
+		{
+			return _collection.FindSync(session, filter, options,  cancellationToken);
+		}
+
+		/// <inheritdoc />
+		public async Task<IAsyncCursor<TProjection>> FindAsync<TProjection>(FilterDefinition<T> filter, FindOptions<T, TProjection> options = null!,
+			CancellationToken  cancellationToken = new CancellationToken())
+		{
+			using IScope scope = GlobalTracer.Instance.BuildSpan(Prefix + "FindAsync").StartActive();
+			scope.Span.SetTag("CollectionName", _collection.CollectionNamespace.CollectionName);
+			return await _collection.FindAsync(filter, options,  cancellationToken);
+		}
+
+		/// <inheritdoc />
+		public async Task<IAsyncCursor<TProjection>> FindAsync<TProjection>(IClientSessionHandle session, FilterDefinition<T> filter, FindOptions<T, TProjection> options = null!,
+			CancellationToken  cancellationToken = new CancellationToken())
+		{
+			using IScope scope = GlobalTracer.Instance.BuildSpan(Prefix + "FindAsync").StartActive();
+			scope.Span.SetTag("CollectionName", _collection.CollectionNamespace.CollectionName);
+			return await _collection.FindAsync(session, filter, options,  cancellationToken);
+		}
+
+		/// <inheritdoc />
+		public TProjection FindOneAndDelete<TProjection>(FilterDefinition<T> filter, FindOneAndDeleteOptions<T, TProjection> options = null!,
+			CancellationToken  cancellationToken = new CancellationToken())
+		{
+			return _collection.FindOneAndDelete(filter, options,  cancellationToken);
+		}
+
+		/// <inheritdoc />
+		public TProjection FindOneAndDelete<TProjection>(IClientSessionHandle session, FilterDefinition<T> filter,
+			FindOneAndDeleteOptions<T, TProjection> options = null!, CancellationToken  cancellationToken = new CancellationToken())
+		{
+			return _collection.FindOneAndDelete(session, filter, options,  cancellationToken);
+		}
+
+		/// <inheritdoc />
+		public async Task<TProjection> FindOneAndDeleteAsync<TProjection>(FilterDefinition<T> filter, FindOneAndDeleteOptions<T, TProjection> options = null!,
+			CancellationToken  cancellationToken = new CancellationToken())
+		{
+			using IScope scope = GlobalTracer.Instance.BuildSpan(Prefix + "FindOneAndDeleteAsync").StartActive();
+			scope.Span.SetTag("CollectionName", _collection.CollectionNamespace.CollectionName);
+			return await _collection.FindOneAndDeleteAsync(filter, options,  cancellationToken);
+		}
+
+		/// <inheritdoc />
+		public async Task<TProjection> FindOneAndDeleteAsync<TProjection>(IClientSessionHandle session, FilterDefinition<T> filter,
+			FindOneAndDeleteOptions<T, TProjection> options = null!, CancellationToken  cancellationToken = new CancellationToken())
+		{
+			using IScope scope = GlobalTracer.Instance.BuildSpan(Prefix + "FindOneAndDeleteAsync").StartActive();
+			scope.Span.SetTag("CollectionName", _collection.CollectionNamespace.CollectionName);
+			return await _collection.FindOneAndDeleteAsync(session, filter, options,  cancellationToken);
+		}
+
+		/// <inheritdoc />
+		public TProjection FindOneAndReplace<TProjection>(FilterDefinition<T> filter, T replacement,
+			FindOneAndReplaceOptions<T, TProjection> options = null!, CancellationToken  cancellationToken = new CancellationToken())
+		{
+			return _collection.FindOneAndReplace(filter, replacement, options,  cancellationToken);
+		}
+
+		/// <inheritdoc />
+		public TProjection FindOneAndReplace<TProjection>(IClientSessionHandle session, FilterDefinition<T> filter, T replacement,
+			FindOneAndReplaceOptions<T, TProjection> options = null!, CancellationToken  cancellationToken = new CancellationToken())
+		{
+			return _collection.FindOneAndReplace(session, filter, replacement, options,  cancellationToken);
+		}
+
+		/// <inheritdoc />
+		public async Task<TProjection> FindOneAndReplaceAsync<TProjection>(FilterDefinition<T> filter, T replacement,
+			FindOneAndReplaceOptions<T, TProjection> options = null!, CancellationToken  cancellationToken = new CancellationToken())
+		{
+			using IScope scope = GlobalTracer.Instance.BuildSpan(Prefix + "FindOneAndReplaceAsync").StartActive();
+			scope.Span.SetTag("CollectionName", _collection.CollectionNamespace.CollectionName);
+			return await _collection.FindOneAndReplaceAsync(filter, replacement, options,  cancellationToken);
+		}
+
+		/// <inheritdoc />
+		public async Task<TProjection> FindOneAndReplaceAsync<TProjection>(IClientSessionHandle session, FilterDefinition<T> filter, T replacement,
+			FindOneAndReplaceOptions<T, TProjection> options = null!, CancellationToken  cancellationToken = new CancellationToken())
+		{
+			using IScope scope = GlobalTracer.Instance.BuildSpan(Prefix + "FindOneAndReplaceAsync").StartActive();
+			scope.Span.SetTag("CollectionName", _collection.CollectionNamespace.CollectionName);
+			return await _collection.FindOneAndReplaceAsync(session, filter, replacement, options,  cancellationToken);
+		}
+
+		/// <inheritdoc />
+		public TProjection FindOneAndUpdate<TProjection>(FilterDefinition<T> filter, UpdateDefinition<T> update,
+			FindOneAndUpdateOptions<T, TProjection> options = null!, CancellationToken  cancellationToken = new CancellationToken())
+		{
+			return _collection.FindOneAndUpdate(filter, update, options,  cancellationToken);
+		}
+
+		/// <inheritdoc />
+		public TProjection FindOneAndUpdate<TProjection>(IClientSessionHandle session, FilterDefinition<T> filter,
+			UpdateDefinition<T> update, FindOneAndUpdateOptions<T, TProjection> options = null!,
+			CancellationToken  cancellationToken = new CancellationToken())
+		{
+			return _collection.FindOneAndUpdate(session, filter, update, options,  cancellationToken);
+		}
+
+		/// <inheritdoc />
+		public async Task<TProjection> FindOneAndUpdateAsync<TProjection>(FilterDefinition<T> filter, UpdateDefinition<T> update,
+			FindOneAndUpdateOptions<T, TProjection> options = null!, CancellationToken  cancellationToken = new CancellationToken())
+		{
+			using IScope scope = GlobalTracer.Instance.BuildSpan(Prefix + "FindOneAndUpdateAsync").StartActive();
+			scope.Span.SetTag("CollectionName", _collection.CollectionNamespace.CollectionName);
+			return await _collection.FindOneAndUpdateAsync(filter, update, options,  cancellationToken);
+		}
+
+		/// <inheritdoc />
+		public async Task<TProjection> FindOneAndUpdateAsync<TProjection>(IClientSessionHandle session, FilterDefinition<T> filter, UpdateDefinition<T> update,
+			FindOneAndUpdateOptions<T, TProjection> options = null!, CancellationToken  cancellationToken = new CancellationToken())
+		{
+			using IScope scope = GlobalTracer.Instance.BuildSpan(Prefix + "FindOneAndUpdateAsync").StartActive();
+			scope.Span.SetTag("CollectionName", _collection.CollectionNamespace.CollectionName);
+			return await _collection.FindOneAndUpdateAsync(session, filter, update, options,  cancellationToken);
+		}
+
+		/// <inheritdoc />
+		public void InsertOne( T document, InsertOneOptions options = null!,
+			CancellationToken  cancellationToken = new CancellationToken())
+		{
+			_collection.InsertOne(document, options,  cancellationToken);
+		}
+
+		/// <inheritdoc />
+		public void InsertOne(IClientSessionHandle session, T document, InsertOneOptions options = null!,
+			CancellationToken  cancellationToken = new CancellationToken())
+		{
+			_collection.InsertOne(session, document, options,  cancellationToken);
+		}
+
+		/// <inheritdoc />
+		public async Task InsertOneAsync( T document, CancellationToken cancellationToken)
+		{
+			using IScope scope = GlobalTracer.Instance.BuildSpan(Prefix + "InsertOneAsync").StartActive();
+			scope.Span.SetTag("CollectionName", _collection.CollectionNamespace.CollectionName);
+			await _collection.InsertOneAsync(document, cancellationToken);
+		}
+
+		/// <inheritdoc />
+		public async Task InsertOneAsync( T document, InsertOneOptions options = null!,
+			CancellationToken  cancellationToken = new CancellationToken())
+		{
+			using IScope scope = GlobalTracer.Instance.BuildSpan(Prefix + "InsertOneAsync").StartActive();
+			scope.Span.SetTag("CollectionName", _collection.CollectionNamespace.CollectionName);
+			await _collection.InsertOneAsync(document, options,  cancellationToken);
+		}
+
+		/// <inheritdoc />
+		public async Task InsertOneAsync(IClientSessionHandle session, T document, InsertOneOptions options = null!,
+			CancellationToken cancellationToken = new CancellationToken())
+		{
+			using IScope scope = GlobalTracer.Instance.BuildSpan(Prefix + "InsertOneAsync").StartActive();
+			scope.Span.SetTag("CollectionName", _collection.CollectionNamespace.CollectionName);
+			await _collection.InsertOneAsync(session, document, options,  cancellationToken);
+		}
+
+		/// <inheritdoc />
+		public void InsertMany(IEnumerable<T> documents, InsertManyOptions options = null!,
+			CancellationToken  cancellationToken = new CancellationToken())
+		{
+			_collection.InsertMany(documents, options,  cancellationToken);
+		}
+
+		/// <inheritdoc />
+		public void InsertMany(IClientSessionHandle session, IEnumerable<T> documents, InsertManyOptions options = null!,
+			CancellationToken  cancellationToken = new CancellationToken())
+		{
+			_collection.InsertMany(session, documents, options,  cancellationToken);
+		}
+
+		/// <inheritdoc />
+		public async Task InsertManyAsync(IEnumerable<T> documents, InsertManyOptions options = null!,
+			CancellationToken  cancellationToken = new CancellationToken())
+		{
+			using IScope scope = GlobalTracer.Instance.BuildSpan(Prefix + "InsertManyAsync").StartActive();
+			scope.Span.SetTag("CollectionName", _collection.CollectionNamespace.CollectionName);
+			await _collection.InsertManyAsync(documents, options,  cancellationToken);
+		}
+
+		/// <inheritdoc />
+		public async Task InsertManyAsync(IClientSessionHandle session, IEnumerable<T> documents, InsertManyOptions options = null!,
+			CancellationToken  cancellationToken = new CancellationToken())
+		{
+			using IScope scope = GlobalTracer.Instance.BuildSpan(Prefix + "InsertManyAsync").StartActive();
+			scope.Span.SetTag("CollectionName", _collection.CollectionNamespace.CollectionName);
+			await _collection.InsertManyAsync(session, documents, options,  cancellationToken);
+		}
+
+		/// <inheritdoc />
+		public IAsyncCursor<TResult> MapReduce<TResult>(BsonJavaScript map, BsonJavaScript reduce, MapReduceOptions<T, TResult> options = null!,
+			CancellationToken  cancellationToken = new CancellationToken())
+		{
+			return _collection.MapReduce(map, reduce, options,  cancellationToken);
+		}
+
+		/// <inheritdoc />
+		public IAsyncCursor<TResult> MapReduce<TResult>(IClientSessionHandle session, BsonJavaScript map, BsonJavaScript reduce,
+			MapReduceOptions<T, TResult> options = null!, CancellationToken  cancellationToken = new CancellationToken())
+		{
+			return _collection.MapReduce(session, map, reduce, options,  cancellationToken);
+		}
+
+		/// <inheritdoc />
+		public async Task<IAsyncCursor<TResult>> MapReduceAsync<TResult>(BsonJavaScript map, BsonJavaScript reduce, MapReduceOptions<T, TResult> options = null!,
+			CancellationToken  cancellationToken = new CancellationToken())
+		{
+			using IScope scope = GlobalTracer.Instance.BuildSpan(Prefix + "MapReduceAsync").StartActive();
+			scope.Span.SetTag("CollectionName", _collection.CollectionNamespace.CollectionName);
+			return await _collection.MapReduceAsync(map, reduce, options,  cancellationToken);
+		}
+
+		/// <inheritdoc />
+		public async Task<IAsyncCursor<TResult>> MapReduceAsync<TResult>(IClientSessionHandle session, BsonJavaScript map, BsonJavaScript reduce,
+			MapReduceOptions<T, TResult> options = null!, CancellationToken  cancellationToken = new CancellationToken())
+		{
+			using IScope scope = GlobalTracer.Instance.BuildSpan(Prefix + "MapReduceAsync").StartActive();
+			scope.Span.SetTag("CollectionName", _collection.CollectionNamespace.CollectionName);
+			return await _collection.MapReduceAsync(session, map, reduce, options,  cancellationToken);
+		}
+
+		/// <inheritdoc />
+		public IFilteredMongoCollection<TDerivedDocument> OfType<TDerivedDocument>() where TDerivedDocument : T
+		{
+			return _collection.OfType<TDerivedDocument>();
+		}
+
+		/// <inheritdoc />
+		public ReplaceOneResult ReplaceOne(FilterDefinition<T> filter, T replacement, ReplaceOptions options = null!,
+			CancellationToken  cancellationToken = new CancellationToken())
+		{
+			return _collection.ReplaceOne(filter, replacement, options,  cancellationToken);
+		}
+
+		/// <inheritdoc />
+		public ReplaceOneResult ReplaceOne(FilterDefinition<T> filter, T replacement, UpdateOptions options,
+			CancellationToken  cancellationToken = new CancellationToken())
+		{
+			return _collection.ReplaceOne(filter, replacement, options,  cancellationToken);
+		}
+
+		/// <inheritdoc />
+		public ReplaceOneResult ReplaceOne(IClientSessionHandle session, FilterDefinition<T> filter, T replacement,
+			ReplaceOptions options = null!, CancellationToken  cancellationToken = new CancellationToken())
+		{
+			return _collection.ReplaceOne(session, filter, replacement, options,  cancellationToken);
+		}
+
+		/// <inheritdoc />
+		public ReplaceOneResult ReplaceOne(IClientSessionHandle session, FilterDefinition<T> filter, T replacement, UpdateOptions options,
+			CancellationToken  cancellationToken = new CancellationToken())
+		{
+			return _collection.ReplaceOne(session, filter, replacement, options,  cancellationToken);
+		}
+
+		/// <inheritdoc />
+		public async Task<ReplaceOneResult> ReplaceOneAsync(FilterDefinition<T> filter, T replacement, ReplaceOptions options = null!,
+			CancellationToken  cancellationToken = new CancellationToken())
+		{
+			using IScope scope = GlobalTracer.Instance.BuildSpan(Prefix + "ReplaceOneAsync").StartActive();
+			scope.Span.SetTag("CollectionName", _collection.CollectionNamespace.CollectionName);
+			return await _collection.ReplaceOneAsync(filter, replacement, options,  cancellationToken);
+		}
+
+		/// <inheritdoc />
+		public async Task<ReplaceOneResult> ReplaceOneAsync(FilterDefinition<T> filter, T replacement, UpdateOptions options,
+			CancellationToken  cancellationToken = new CancellationToken())
+		{
+			using IScope scope = GlobalTracer.Instance.BuildSpan(Prefix + "ReplaceOneAsync").StartActive();
+			scope.Span.SetTag("CollectionName", _collection.CollectionNamespace.CollectionName);
+			return await _collection.ReplaceOneAsync(filter, replacement, options,  cancellationToken);
+		}
+
+		/// <inheritdoc />
+		public async Task<ReplaceOneResult> ReplaceOneAsync(IClientSessionHandle session, FilterDefinition<T> filter, T replacement,
+			ReplaceOptions options = null!, CancellationToken  cancellationToken = new CancellationToken())
+		{
+			using IScope scope = GlobalTracer.Instance.BuildSpan(Prefix + "ReplaceOneAsync").StartActive();
+			scope.Span.SetTag("CollectionName", _collection.CollectionNamespace.CollectionName);
+			return await _collection.ReplaceOneAsync(session, filter, replacement, options,  cancellationToken);
+		}
+
+		/// <inheritdoc />
+		public async Task<ReplaceOneResult> ReplaceOneAsync(IClientSessionHandle session, FilterDefinition<T> filter, T replacement, UpdateOptions options,
+			CancellationToken  cancellationToken = new CancellationToken())
+		{
+			using IScope scope = GlobalTracer.Instance.BuildSpan(Prefix + "ReplaceOneAsync").StartActive();
+			scope.Span.SetTag("CollectionName", _collection.CollectionNamespace.CollectionName);
+			return await _collection.ReplaceOneAsync(session, filter, replacement, options,  cancellationToken);
+		}
+
+		/// <inheritdoc />
+		public UpdateResult UpdateMany(FilterDefinition<T> filter, UpdateDefinition<T> update, UpdateOptions options = null!,
+			CancellationToken  cancellationToken = new CancellationToken())
+		{
+			return _collection.UpdateMany(filter, update, options,  cancellationToken);
+		}
+
+		/// <inheritdoc />
+		public UpdateResult UpdateMany(IClientSessionHandle session, FilterDefinition<T> filter, UpdateDefinition<T> update,
+			UpdateOptions options = null!, CancellationToken  cancellationToken = new CancellationToken())
+		{
+			return _collection.UpdateMany(session, filter, update, options,  cancellationToken);
+		}
+
+		/// <inheritdoc />
+		public async Task<UpdateResult> UpdateManyAsync(FilterDefinition<T> filter, UpdateDefinition<T> update, UpdateOptions options = null!,
+			CancellationToken  cancellationToken = new CancellationToken())
+		{
+			using IScope scope = GlobalTracer.Instance.BuildSpan(Prefix + "UpdateManyAsync").StartActive();
+			scope.Span.SetTag("CollectionName", _collection.CollectionNamespace.CollectionName);
+			return await _collection.UpdateManyAsync(filter, update, options,  cancellationToken);
+		}
+
+		/// <inheritdoc />
+		public async Task<UpdateResult> UpdateManyAsync(IClientSessionHandle session, FilterDefinition<T> filter, UpdateDefinition<T> update,
+			UpdateOptions options = null!, CancellationToken  cancellationToken = new CancellationToken())
+		{
+			using IScope scope = GlobalTracer.Instance.BuildSpan(Prefix + "UpdateManyAsync").StartActive();
+			scope.Span.SetTag("CollectionName", _collection.CollectionNamespace.CollectionName);
+			return await _collection.UpdateManyAsync(session, filter, update, options,  cancellationToken);
+		}
+
+		/// <inheritdoc />
+		public UpdateResult UpdateOne(FilterDefinition<T> filter, UpdateDefinition<T> update, UpdateOptions options = null!,
+			CancellationToken  cancellationToken = new CancellationToken())
+		{
+			return _collection.UpdateOne(filter, update, options,  cancellationToken);
+		}
+
+		/// <inheritdoc />
+		public UpdateResult UpdateOne(IClientSessionHandle session, FilterDefinition<T> filter, UpdateDefinition<T> update,
+			UpdateOptions options = null!, CancellationToken  cancellationToken = new CancellationToken())
+		{
+			return _collection.UpdateOne(session, filter, update, options,  cancellationToken);
+		}
+
+		/// <inheritdoc />
+		public async Task<UpdateResult> UpdateOneAsync(FilterDefinition<T> filter, UpdateDefinition<T> update, UpdateOptions options = null!,
+			CancellationToken  cancellationToken = new CancellationToken())
+		{
+			using IScope scope = GlobalTracer.Instance.BuildSpan(Prefix + "UpdateOneAsync").StartActive();
+			scope.Span.SetTag("CollectionName", _collection.CollectionNamespace.CollectionName);
+			return await _collection.UpdateOneAsync(filter, update, options,  cancellationToken);
+		}
+
+		/// <inheritdoc />
+		public async Task<UpdateResult> UpdateOneAsync(IClientSessionHandle session, FilterDefinition<T> filter, UpdateDefinition<T> update,
+			UpdateOptions options = null!, CancellationToken  cancellationToken = new CancellationToken())
+		{
+			using IScope scope = GlobalTracer.Instance.BuildSpan(Prefix + "UpdateOneAsync").StartActive();
+			scope.Span.SetTag("CollectionName", _collection.CollectionNamespace.CollectionName);
+			return await _collection.UpdateOneAsync(session, filter, update, options,  cancellationToken);
+		}
+
+		/// <inheritdoc />
+		public IChangeStreamCursor<TResult> Watch<TResult>(PipelineDefinition<ChangeStreamDocument<T>, TResult> pipeline, ChangeStreamOptions options = null!,
+			CancellationToken  cancellationToken = new CancellationToken())
+		{
+			return _collection.Watch(pipeline, options,  cancellationToken);
+		}
+
+		/// <inheritdoc />
+		public IChangeStreamCursor<TResult> Watch<TResult>(IClientSessionHandle session, PipelineDefinition<ChangeStreamDocument<T>, TResult> pipeline,
+			ChangeStreamOptions options = null!, CancellationToken  cancellationToken = new CancellationToken())
+		{
+			return _collection.Watch(session, pipeline, options,  cancellationToken);
+		}
+
+		/// <inheritdoc />
+		public async Task<IChangeStreamCursor<TResult>> WatchAsync<TResult>(PipelineDefinition<ChangeStreamDocument<T>, TResult> pipeline, ChangeStreamOptions options = null!,
+			CancellationToken  cancellationToken = new CancellationToken())
+		{
+			using IScope scope = GlobalTracer.Instance.BuildSpan(Prefix + "WatchAsync").StartActive();
+			scope.Span.SetTag("CollectionName", _collection.CollectionNamespace.CollectionName);
+			return await _collection.WatchAsync(pipeline, options,  cancellationToken);
+		}
+
+		/// <inheritdoc />
+		public async Task<IChangeStreamCursor<TResult>> WatchAsync<TResult>(IClientSessionHandle session, PipelineDefinition<ChangeStreamDocument<T>, TResult> pipeline, ChangeStreamOptions options = null!,
+			CancellationToken  cancellationToken = new CancellationToken())
+		{
+			using IScope scope = GlobalTracer.Instance.BuildSpan(Prefix + "WatchAsync").StartActive();
+			scope.Span.SetTag("CollectionName", _collection.CollectionNamespace.CollectionName);
+			return await _collection.WatchAsync(session, pipeline, options,  cancellationToken);
+		}
+
+		/// <inheritdoc />
+		public IMongoCollection<T> WithReadConcern(ReadConcern readConcern)
+		{
+			return _collection.WithReadConcern(readConcern);
+		}
+
+		/// <inheritdoc />
+		public IMongoCollection<T> WithReadPreference(ReadPreference readPreference)
+		{
+			return _collection.WithReadPreference(readPreference);
+		}
+
+		/// <inheritdoc />
+		public IMongoCollection<T> WithWriteConcern(WriteConcern writeConcern)
+		{
+			return _collection.WithWriteConcern(writeConcern);
+		}
+#pragma warning restore CS0618
+		
+		/// <inheritdoc />
+		public CollectionNamespace CollectionNamespace => _collection.CollectionNamespace;
+
+		/// <inheritdoc />
+		public IMongoDatabase Database => _collection.Database;
+
+		/// <inheritdoc />
+		public IBsonSerializer<T> DocumentSerializer => _collection.DocumentSerializer;
+
+		/// <inheritdoc />
+		public IMongoIndexManager<T> Indexes => _collection.Indexes;
+
+		/// <inheritdoc />
+		public MongoCollectionSettings Settings => _collection.Settings;
+	}
+}
\ No newline at end of file
diff --git a/Engine/Source/Programs/Horde/Horde.Server/Utilities/NoOpDogStatsd.cs b/Engine/Source/Programs/Horde/Horde.Server/Utilities/NoOpDogStatsd.cs
new file mode 100644
--- /dev/null
+++ b/Engine/Source/Programs/Horde/Horde.Server/Utilities/NoOpDogStatsd.cs
@@ -0,0 +1,102 @@
+﻿// Copyright Epic Games, Inc. All Rights Reserved.
+
+using System;
+using System.IO;
+using StatsdClient;
+
+namespace Horde.Server.Utilities
+{
+	/// <summary>
+	/// Empty implementation that does nothing but still satisfies the interface
+	/// Used when DogStatsD is not available.
+	/// </summary>
+	public sealed class NoOpDogStatsd : IDogStatsd
+	{
+		/// <inheritdoc />
+		public ITelemetryCounters TelemetryCounters { get; } = null!;
+		
+		/// <inheritdoc />
+		public void Configure(StatsdConfig config)
+		{
+		}
+
+		/// <inheritdoc />
+		public void Counter(string statName, double value, double sampleRate = 1, string[] tags = null!)
+		{
+		}
+
+		/// <inheritdoc />
+		public void Decrement(string statName, int value = 1, double sampleRate = 1, params string[] tags)
+		{
+		}
+
+		/// <inheritdoc />
+		public void Event(string title, string text, string alertType = null!, string aggregationKey = null!, string sourceType = null!,
+			int? dateHappened = null, string priority = null!, string hostname = null!, string[] tags = null!)
+		{
+		}
+
+		/// <inheritdoc />
+		public void Gauge(string statName, double value, double sampleRate = 1, string[] tags = null!)
+		{
+		}
+
+		/// <inheritdoc />
+		public void Histogram(string statName, double value, double sampleRate = 1, string[] tags = null!)
+		{
+		}
+
+		/// <inheritdoc />
+		public void Distribution(string statName, double value, double sampleRate = 1, string[] tags = null!)
+		{
+		}
+
+		/// <inheritdoc />
+		public void Increment(string statName, int value = 1, double sampleRate = 1, string[] tags = null!)
+		{
+		}
+
+		/// <inheritdoc />
+		public void Set<T>(string statName, T value, double sampleRate = 1, string[] tags = null!)
+		{
+		}
+
+		/// <inheritdoc />
+		public void Set(string statName, string value, double sampleRate = 1, string[] tags = null!)
+		{
+		}
+
+		/// <inheritdoc />
+		public IDisposable StartTimer(string name, double sampleRate = 1, string[] tags = null!)
+		{
+			// Some random object to satisfy IDisposable
+			return new MemoryStream();
+		}
+
+		/// <inheritdoc />
+		public void Time(Action action, string statName, double sampleRate = 1, string[] tags = null!)
+		{
+		}
+
+		/// <inheritdoc />
+		public T Time<T>(Func<T> func, string statName, double sampleRate = 1, string[] tags = null!)
+		{
+			return func();
+		}
+
+		/// <inheritdoc />
+		public void Timer(string statName, double value, double sampleRate = 1, string[] tags = null!)
+		{
+		}
+
+		/// <inheritdoc />
+		public void ServiceCheck(string name, Status status, int? timestamp = null, string hostname = null!, string[] tags = null!, string message = null!)
+		{
+		}
+		
+		/// <inheritdoc />
+		public void Dispose()
+		{
+		}
+	}
+}
\ No newline at end of file
diff --git a/Engine/Source/Programs/Horde/Horde.Server/Utilities/Schemas.cs b/Engine/Source/Programs/Horde/Horde.Server/Utilities/Schemas.cs
new file mode 100644
--- /dev/null
+++ b/Engine/Source/Programs/Horde/Horde.Server/Utilities/Schemas.cs
@@ -0,0 +1,47 @@
+// Copyright Epic Games, Inc. All Rights Reserved.
+
+using System;
+using System.Collections.Concurrent;
+using System.Reflection;
+using System.Xml;
+using EpicGames.Core;
+
+namespace Horde.Server.Utilities
+{
+	/// <summary>
+	/// Utility functions for generating schemas
+	/// </summary>
+	static class Schemas
+	{
+		static readonly ConcurrentDictionary<Assembly, XmlDocument?> s_cachedDocumentation = new ConcurrentDictionary<Assembly, XmlDocument?>();
+		static readonly ConcurrentDictionary<Type, JsonSchema> s_cachedSchemas = new ConcurrentDictionary<Type, JsonSchema>();
+
+		/// <summary>
+		/// Create a Json schema (or retrieve a cached schema)
+		/// </summary>
+		/// <param name="type"></param>
+		/// <returns></returns>
+		public static JsonSchema CreateSchema(Type type)
+		{
+			JsonSchema? schema;
+			if (!s_cachedSchemas.TryGetValue(type, out schema))
+			{
+				XmlDocument? documentation;
+				if (!s_cachedDocumentation.TryGetValue(type.Assembly, out documentation))
+				{
+					FileReference inputDocumentationFile = new FileReference(type.Assembly.Location).ChangeExtension(".xml");
+					if (FileReference.Exists(inputDocumentationFile))
+					{
+						documentation = new XmlDocument();
+						documentation.Load(inputDocumentationFile.FullName);
+					}
+					s_cachedDocumentation.TryAdd(type.Assembly, documentation);
+				}
+
+				schema = JsonSchema.FromType(type, documentation);
+				s_cachedSchemas.TryAdd(type, schema);
+			}
+			return schema;
+		}
+	}
+}
diff --git a/Engine/Source/Programs/Horde/Horde.Server/Utilities/SemaphoreSlimExtensions.cs b/Engine/Source/Programs/Horde/Horde.Server/Utilities/SemaphoreSlimExtensions.cs
new file mode 100644
--- /dev/null
+++ b/Engine/Source/Programs/Horde/Horde.Server/Utilities/SemaphoreSlimExtensions.cs
@@ -0,0 +1,65 @@
+// Copyright Epic Games, Inc. All Rights Reserved.
+
+using System;
+using System.Threading;
+using System.Threading.Tasks;
+
+namespace Horde.Server.Utilities
+{
+	/// <summary>
+	/// Extension method to wrap a semaphore in a using statement
+	/// </summary>
+	public static class SemaphoreSlimExtensions
+	{
+		/// <summary>
+		/// Returns a disposable semaphore
+		/// </summary>
+		/// <param name="semaphore">the semaphore to wrap</param>
+		/// <param name="cancelToken">optional cancellation token</param>
+		/// <returns></returns>
+		public static async Task<IDisposable> UseWaitAsync(this SemaphoreSlim semaphore, CancellationToken cancelToken = default)
+		{
+			await semaphore.WaitAsync(cancelToken).ConfigureAwait(false);
+			return new ReleaseWrapper(semaphore);
+		}
+
+		/// <summary>
+		/// Semaphore wrapper implementing IDisposable
+		/// </summary>
+		private class ReleaseWrapper : IDisposable
+		{
+			/// <summary>
+			/// The semaphore to wrap
+			/// </summary>
+			private readonly SemaphoreSlim _semaphore;
+
+			/// <summary>
+			/// Whether this has been disposed
+			/// </summary>
+			private bool _bIsDisposed;
+
+			/// <summary>
+			/// Constructor
+			/// </summary>
+			/// <param name="semaphore">The semaphore to wrap</param>
+			public ReleaseWrapper(SemaphoreSlim semaphore)
+			{
+				_semaphore = semaphore;
+			}
+
+			/// <summary>
+			/// Releases the lock on dispose
+			/// </summary>
+			public void Dispose()
+			{
+				if (_bIsDisposed)
+				{
+					return;
+				}
+
+				_semaphore.Release();
+				_bIsDisposed = true;
+			}
+		}
+	}
+}
diff --git a/Engine/Source/Programs/Horde/Horde.Server/Utilities/TracingExtensions.cs b/Engine/Source/Programs/Horde/Horde.Server/Utilities/TracingExtensions.cs
new file mode 100644
--- /dev/null
+++ b/Engine/Source/Programs/Horde/Horde.Server/Utilities/TracingExtensions.cs
@@ -0,0 +1,54 @@
+// Copyright Epic Games, Inc. All Rights Reserved.
+
+using System;
+using System.Collections.Generic;
+using EpicGames.Core;
+using Horde.Server.Jobs;
+using OpenTracing;
+
+namespace Horde.Server.Utilities
+{
+	/// <summary>
+	/// Extensions to handle Horde specific data types in the OpenTracing library
+	/// </summary>
+	public static class OpenTracingSpanExtensions
+	{
+		/// <summary>Set a key:value tag on the span</summary>
+		/// <returns>This span instance, for chaining</returns>
+		public static ISpan SetTag(this ISpan span, string key, ContentHash value)
+		{
+			span.SetTag(key, value.ToString());
+			return span;
+		}
+		
+		/// <summary>Set a key:value tag on the span</summary>
+		/// <returns>This span instance, for chaining</returns>
+		public static ISpan SetTag(this ISpan span, string key, SubResourceId value)
+		{
+			span.SetTag(key, value.ToString());
+			return span;
+		}
+		
+		/// <summary>Set a key:value tag on the span</summary>
+		/// <returns>This span instance, for chaining</returns>
+		public static ISpan SetTag(this ISpan span, string key, int? value)
+		{
+			if (value != null)
+			{
+				span.SetTag(key, value.Value);
+			}
+			return span;
+		}
+		
+		/// <summary>Set a key:value tag on the span</summary>
+		/// <returns>This span instance, for chaining</returns>
+		public static ISpan SetTag(this ISpan span, string key, DateTimeOffset? value)
+		{
+			if (value != null)
+			{
+				span.SetTag(key, value.ToString());
+			}
+			return span;
+		}
+	}
+}
\ No newline at end of file
diff --git a/Engine/Source/Programs/Horde/Horde.Server/appsettings.Local.json b/Engine/Source/Programs/Horde/Horde.Server/appsettings.Local.json
new file mode 100644
--- /dev/null
+++ b/Engine/Source/Programs/Horde/Horde.Server/appsettings.Local.json
@@ -0,0 +1,21 @@
+{
+	"Horde":
+	{
+		"EnableUpgradeTasks": false,
+		"EnableConformTasks": false,
+		"ForceConfigUpdateOnStartup": true,
+		"Commits":
+		{
+			"ReplicateMetadata": false,
+			"ReplicateContent": false
+		}
+	},
+	"Serilog": {
+		"MinimumLevel": {
+			"Override": {
+				"Horde.Server.Agents.Fleet": "Warning",
+				"Horde.Server.Agents.Fleet.Providers.NoOpFleetManager": "Warning"
+			}
+		}
+	}
+}
diff --git a/Engine/Source/Programs/Horde/Horde.Server/appsettings.Production.json b/Engine/Source/Programs/Horde/Horde.Server/appsettings.Production.json
new file mode 100644
--- /dev/null
+++ b/Engine/Source/Programs/Horde/Horde.Server/appsettings.Production.json
@@ -0,0 +1,5 @@
+{
+    "Horde": 
+	{
+    }
+}
