Horde: Fix issue with timeline and label filtering
#rnx
#skipci

diff --git a/Engine/Source/Programs/Horde/HordeDashboard/src/components/jobDetailsV2/JobDetailTimeline.tsx b/Engine/Source/Programs/Horde/HordeDashboard/src/components/jobDetailsV2/JobDetailTimeline.tsx
--- a/Engine/Source/Programs/Horde/HordeDashboard/src/components/jobDetailsV2/JobDetailTimeline.tsx
+++ b/Engine/Source/Programs/Horde/HordeDashboard/src/components/jobDetailsV2/JobDetailTimeline.tsx
@@ -139,23 +139,22 @@ class TimelineDataView extends JobDataView {
 
       let bresponses = job.batches?.filter(b => !!b.agentId && (!!b.startTime && !!b.readyTime && !!b.finishTime) && b.steps?.length)!.map(b => { return { ...b } as GetBatchResponse })!;
 
-      if (jobDetails.filter) {
+      if (!this.filterStepId && jobDetails.filter) {
          const label = jobDetails.filter.label;
          if (label) {
-
             bresponses?.forEach(b => {
                b.steps = b.steps?.filter(s => {
-                  if (!s.readyTime || !s.startTime || !s.finishTime) {
+                  if (!s.startTime || !s.finishTime) {
                      return false;
                   }
                   const node = jobDetails.nodeByStepId(s.id);
-                  return label.includedNodes.indexOf(node?.name ?? "") !== -1;
+                  return label.includedNodes.indexOf(node?.name ?? "") !== -1 || label.requiredNodes.indexOf(node?.name ?? "") !== -1;
                });
             });
 
             bresponses = bresponses.filter(b => !!b.steps?.length);
-         }
-      }
+         }         
+      }      
 
       let filterStep: StepData | undefined;
 
@@ -170,7 +169,6 @@ class TimelineDataView extends JobDataView {
          });
 
          bresponses = bresponses.filter(b => !!b.steps?.length);
-
       }
 
       let maxTime = 0;
@@ -199,7 +197,7 @@ class TimelineDataView extends JobDataView {
       });
 
       if (!batches.length) {
-         console.error("No batches to render in timeline");
+         console.log("No batches to render in timeline");
          return;
       }
 
