Horde: Fix path to Dockerfile in build script.

diff --git a/Engine/Source/Programs/Horde/BuildHorde.xml b/Engine/Source/Programs/Horde/BuildHorde.xml
--- a/Engine/Source/Programs/Horde/BuildHorde.xml
+++ b/Engine/Source/Programs/Horde/BuildHorde.xml
@@ -12,7 +12,7 @@
 	<Property Name="VersionArguments" Value="/p:Version=$(Version).0 /p:InformationalVersion=$(InformationalVersion)"/>
 
 	<Agent Name="HordeServer" Type="Linux;Win64_Docker">
-		<Property Name="StagingDir" Value="$(RootDir)/Engine/Saved/Horde.Build"/>
+		<Property Name="StagingDir" Value="$(RootDir)/Engine/Saved/Horde.Server"/>
 		
 		<Node Name="Build HordeServer">
 			<!-- Tag all the files that need to be staged to build -->
@@ -26,7 +26,7 @@
 			<Tag Files="$(StagePaths)" Except=".../.vs/...;.../.git/...;.../bin/...;.../obj/..." With="#InputFiles"/>
 
 			<!-- Build the Docker image and publish it -->
-			<Docker-Build BaseDir="Engine" Files="#InputFiles" Tag="hordeserver-public" DockerFile="Engine/Source/Programs/Horde/Horde.Build/Dockerfile" Arguments="--build-arg msbuild_args=&quot;$(VersionArguments)&quot;"/>
+			<Docker-Build BaseDir="Engine" Files="#InputFiles" Tag="hordeserver-public" DockerFile="Engine/Source/Programs/Horde/Horde.Server/Dockerfile" Arguments="--build-arg msbuild_args=&quot;$(VersionArguments)&quot;"/>
 		</Node>
 	</Agent>
 
