Horde: Display 'special' user ids for closing issues correctly in Slack notifications.

diff --git a/Engine/Source/Programs/Horde/Horde.Server/Notifications/Sinks/SlackNotificationSink.cs b/Engine/Source/Programs/Horde/Horde.Server/Notifications/Sinks/SlackNotificationSink.cs
--- a/Engine/Source/Programs/Horde/Horde.Server/Notifications/Sinks/SlackNotificationSink.cs
+++ b/Engine/Source/Programs/Horde/Horde.Server/Notifications/Sinks/SlackNotificationSink.cs
@@ -1473,7 +1473,18 @@ namespace Horde.Server.Notifications.Sinks
 			IUser? user = await _userCollection.GetUserAsync(userId);
 			if (user == null)
 			{
-				return $"User {userId}";
+				if (userId == IIssue.ResolvedByUnknownId)
+				{
+					return "Horde (Unknown)";
+				}
+				else if (userId == IIssue.ResolvedByTimeoutId)
+				{
+					return "Horde (Timeout)";
+				}
+				else
+				{
+					return $"User {userId}";
+				}
 			}
 			return user.Name;
 		}
