Horde: Forcibly terminate child process groups after 60 seconds if the stdout pipe isn't closed.

diff --git a/Engine/Source/Programs/Horde/Horde.Agent/Execution/JobExecutor.cs b/Engine/Source/Programs/Horde/Horde.Agent/Execution/JobExecutor.cs
--- a/Engine/Source/Programs/Horde/Horde.Agent/Execution/JobExecutor.cs
+++ b/Engine/Source/Programs/Horde/Horde.Agent/Execution/JobExecutor.cs
@@ -1141,13 +1141,24 @@ namespace Horde.Agent.Execution
 		static async Task<int> ExecuteProcessAsync(string fileName, string arguments, IReadOnlyDictionary<string, string>? newEnvironment, LogParser filter, ILogger logger, CancellationToken cancellationToken)
 		{
 			logger.LogInformation("Executing {File} {Arguments}", fileName.QuoteArgument(), arguments);
-			using (ManagedProcessGroup processGroup = new ManagedProcessGroup())
+
+			using (CancellationTokenSource cancellationSource = CancellationTokenSource.CreateLinkedTokenSource(cancellationToken))
 			{
-				using (ManagedProcess process = new ManagedProcess(processGroup, fileName, arguments, null, newEnvironment, null, ProcessPriorityClass.Normal))
+				using (ManagedProcessGroup processGroup = new ManagedProcessGroup())
 				{
-					await process.CopyToAsync((buffer, offset, length) => filter.WriteData(buffer.AsMemory(offset, length)), 4096, cancellationToken);
-					process.WaitForExit();
-					return process.ExitCode;
+					using (ManagedProcess process = new ManagedProcess(processGroup, fileName, arguments, null, newEnvironment, null, ProcessPriorityClass.Normal))
+					{
+						async Task WaitForExitOrCancelAsync()
+						{
+							await process.WaitForExitAsync(cancellationToken);
+							cancellationSource.CancelAfter(TimeSpan.FromSeconds(60.0));
+						}
+
+						Task cancelTask = WaitForExitOrCancelAsync();
+						await process.CopyToAsync((buffer, offset, length) => filter.WriteData(buffer.AsMemory(offset, length)), 4096, cancellationSource.Token);
+						await cancelTask;
+						return process.ExitCode;
+					}
 				}
 			}
 		}
diff --git a/Engine/Source/Programs/Shared/EpicGames.Core/ManagedProcess.cs b/Engine/Source/Programs/Shared/EpicGames.Core/ManagedProcess.cs
--- a/Engine/Source/Programs/Shared/EpicGames.Core/ManagedProcess.cs
+++ b/Engine/Source/Programs/Shared/EpicGames.Core/ManagedProcess.cs
@@ -1274,6 +1274,30 @@ namespace EpicGames.Core
 			}
 		}
 
+		class ProcessWaitHandle : WaitHandle
+		{
+			public ProcessWaitHandle(SafeFileHandle processHandle)
+			{
+				SafeWaitHandle = new SafeWaitHandle(processHandle.DangerousGetHandle(), false);
+			}
+		}
+
+		/// <summary>
+		/// Waits for the process to exit
+		/// </summary>
+		public async Task WaitForExitAsync(CancellationToken cancellationToken = default)
+		{
+			if (_frameworkProcess == null)
+			{
+				using ProcessWaitHandle waitHandle = new ProcessWaitHandle(_processHandle!);
+				await waitHandle.WaitOneAsync(cancellationToken); 
+			}
+			else
+			{
+				await _frameworkProcess.WaitForExitAsync(cancellationToken);
+			}
+		}
+
 		/// <summary>
 		/// The exit code of the process. Throws an exception if the process has not terminated.
 		/// </summary>
