Horde: Fix build errors.

diff --git a/Engine/Source/Programs/Horde/Horde.Server/Compute/ComputeTaskSource.cs b/Engine/Source/Programs/Horde/Horde.Server/Compute/ComputeTaskSource.cs
--- a/Engine/Source/Programs/Horde/Horde.Server/Compute/ComputeTaskSource.cs
+++ b/Engine/Source/Programs/Horde/Horde.Server/Compute/ComputeTaskSource.cs
@@ -9,6 +9,7 @@ using System.Threading;
 using System.Threading.Tasks;
 using EpicGames.Core;
 using EpicGames.Horde.Compute;
+using EpicGames.Horde.Compute.Transports;
 using Google.Protobuf;
 using Google.Protobuf.WellKnownTypes;
 using Horde.Server.Agents;
@@ -132,7 +133,7 @@ namespace Horde.Server.Compute
 		{
 			ComputeTask computeTask = new ComputeTask();
 			computeTask.Nonce = UnsafeByteOperations.UnsafeWrap(RandomNumberGenerator.GetBytes(ServerComputeClient.NonceLength));
-			computeTask.Key = UnsafeByteOperations.UnsafeWrap(ComputeLease.CreateKey());
+			computeTask.Key = UnsafeByteOperations.UnsafeWrap(AesTransport.CreateKey());
 			computeTask.Resources.Add(assignedResources);
 			return computeTask;
 		}
diff --git a/Engine/Source/Programs/Shared/EpicGames.Horde/Compute/ComputeChannel.cs b/Engine/Source/Programs/Shared/EpicGames.Horde/Compute/ComputeChannel.cs
--- a/Engine/Source/Programs/Shared/EpicGames.Horde/Compute/ComputeChannel.cs
+++ b/Engine/Source/Programs/Shared/EpicGames.Horde/Compute/ComputeChannel.cs
@@ -97,7 +97,9 @@ namespace EpicGames.Horde.Compute
 		readonly IComputeBufferReader _receiveBufferReader;
 		readonly IComputeBufferWriter _sendBufferWriter;
 
+#pragma warning disable CA2213
 		MessageBuilder? _currentBuilder;
+#pragma warning restore CA2213
 
 		/// <summary>
 		/// Constructor
