Horde: Move installer to regular build script.

diff --git a/Engine/Source/Programs/Horde/BuildHorde.xml b/Engine/Source/Programs/Horde/BuildHorde.xml
--- a/Engine/Source/Programs/Horde/BuildHorde.xml
+++ b/Engine/Source/Programs/Horde/BuildHorde.xml
@@ -11,6 +11,8 @@
 	<Property Name="InformationalVersion" Value="$(InformationalVersion)-PF-$(PreflightChange)" If="'$(PreflightChange)' != ''"/>
 	<Property Name="VersionArguments" Value="/p:Version=$(Version).0 /p:InformationalVersion=$(InformationalVersion)"/>
 
+	<!--- SERVER --->
+
 	<Agent Name="HordeServer" Type="Linux;Win64_Docker">
 		<Property Name="StagingDir" Value="$(RootDir)/Engine/Saved/Horde.Server"/>
 		
@@ -29,33 +31,47 @@
 			<Docker-Build BaseDir="Engine" Files="#InputFiles" Tag="hordeserver-public" DockerFile="Engine/Source/Programs/Horde/Horde.Server/Dockerfile" Arguments="--build-arg msbuild_args=&quot;$(VersionArguments)&quot;"/>
 		</Node>
 	</Agent>
+	
+	<!--- INSTALLER --->
+	
+	<Agent Name="HordeInstaller" Type="Win64_Docker">
+		<Property Name="StagingDir" Value="$(RootDir)/Staging"/>
 
-	<Agent Name="HordeStorage" Type="Linux">
-		<Property Name="ComposeFile" Value="$(RootDir)/Engine/Source/Programs/Horde/HordeStorage/Composes/docker-compose-tests.yml"/>
-		<Node Name="Test HordeStorage">
-			<!-- Make sure to run docker compose down after the build no matter what, also copy the logs to AutomationTool logs so we get the logs in Horde -->
-			<WriteTextFile File="$(RootDir)/Cleanup.sh" Text="docker-compose -f $(ComposeFile) logs --no-color > $(RootDir)/Engine/Programs/AutomationTool/Saved/Logs/docker-compose-logs.txt"/>
-			<WriteTextFile File="$(RootDir)/Cleanup.sh" Append="true" Text="docker-compose -f $(ComposeFile) down"/>
-
-			<Docker-Compose-Up File="$(ComposeFile)"/>
+		<Node Name="Stage HordeInstaller">
+			<Delete Files="$(StagingDir)/..."/>
+			<DotNet Arguments="publish &quot;$(RootDir)/Engine/Source/Programs/Horde/Horde.Server/Horde.Server.csproj&quot; --output &quot;$(StagingDir)/Server&quot; --runtime win-x64 $(VersionArguments)"/>
+			<DotNet Arguments="publish &quot;$(RootDir)/Engine/Source/Programs/Horde/Horde.Agent/Horde.Agent.csproj&quot; --output &quot;$(StagingDir)/Agent&quot; $(VersionArguments)"/>
 
-			<!-- Sleep a while to let scylla finish starting up before we execute the tests -->
-			<Spawn Exe="/bin/bash" Arguments="-c 'sleep 60'" />
+			<!-- Do not include development defaults which created a ue5 project and stream -->
+			<Delete Files="$(StagingDir)/Server/Defaults/..."/>
 
-			<!-- Run the unit tests, and override the dotnet sdk used to the global one for the machine rather then what is in the sync, to allow us to use newer versions then the rest of the engine has migrated to -->
-			<DotNet Arguments="test Engine/Source/Programs/Horde/HordeStorage/Horde.Storage.sln -l console;verbosity=normal" DotNetPath="/usr/share/dotnet/dotnet"/>
+			<!-- Create build settings -->
+			<WriteTextFile File="$(StagingDir)/Server/appsettings.Build.json" Text="{ &quot;Horde&quot;: { &quot;SingleInstance&quot;: true } }"/>
+	
+			<Spawn Exe="cmd.exe" Arguments="/C yarn install" WorkingDir="$(RootDir)/Engine/Source/Programs/Horde/HordeDashboard"/>
+			<Spawn Exe="cmd.exe" Arguments="/C yarn run build" WorkingDir="$(RootDir)/Engine/Source/Programs/Horde/HordeDashboard"/>
+	
+			<Copy From="$(RootDir)/Engine/Source/Programs/Horde/HordeDashboard/build/..." To="$(StagingDir)/Server/DashboardApp/..."/>
 		</Node>
 
-		<Node Name="Build HordeStorage" Requires="Test HordeStorage">
-			<!-- Tag all the files that need to be staged to build -->
-			<Property Name="StagePaths">
-				Engine/Source/Programs/Shared/...
-				Engine/Source/Programs/Horde/HordeStorage/...
-			</Property>
-			<Tag Files="$(StagePaths)" Except=".../.vs/...;.../.git/...;.../bin/...;.../obj/..." With="#InputFiles"/>
+		<Property Name="WixDir" Value="$(RootDir)/Engine/Source/ThirdParty/WiX/3.8"/>
+		<Property Name="SourceDir" Value="$(RootDir)/Engine/Source/Programs/Horde/Installer"/>
+		<Property Name="ObjDir" Value="obj/Debug"/>
+		<Property Name="BinDir" Value="bin/Debug"/>
+		
+		<Node Name="Build HordeInstaller" Requires="Stage HordeInstaller">
+			<Property Name="CommonArgs" Value="-dConfiguration=Debug -dPlatform=x64 -arch x64"/>
 
-			<!-- Build the Docker image -->
-			<Docker-Build BaseDir="Engine/Source" Files="#InputFiles" Tag="horde_storage:$(InformationalVersion)" DockerFile="Engine/Source/Programs/Horde/HordeStorage/dockerfile" Environment="UE_P4_CHANGELIST=(Change)" UseBuildKit="true" Arguments="--target horde-storage-datadog"/>
+			<Delete Files="$(StagingDir)/Server-Bulk/..."/>
+			<Tag Files="$(StagingDir)/Server/..." Except="HordeServer.exe" With="#ServerFiles"/>
+			<Copy Files="#ServerFiles" From="$(StagingDir)/Server" To="$(StagingDir)/Server-Bulk/Server"/>
+			<Spawn Exe="$(WixDir)/heat.exe" Arguments="dir &quot;$(StagingDir)/Server-Bulk&quot; -cg HordeServerFiles -scom -sreg -gg -dr InstallDir -srd -var var.SourceDir -out &quot;$(SourceDir)/HordeServerFiles.wxs&quot;"/>
+			<Spawn Exe="$(WixDir)/candle.exe" Arguments="$(CommonArgs) HordeServerFiles.wxs -dSourceDir=&quot;$(StagingDir)&quot; -out $(ObjDir)/HordeServerFiles.wixobj" WorkingDir="$(SourceDir)"/>
+
+			<Spawn Exe="$(WixDir)/candle.exe" Arguments="$(CommonArgs) Installer.wxs -dStagingDir=&quot;$(StagingDir)&quot; -out $(ObjDir)/Installer.wixobj" WorkingDir="$(SourceDir)"/>
+			<Spawn Exe="$(WixDir)/light.exe" Arguments="-out $(BinDir)/Horde.msi -sw1076 -pdbout $(BinDir)/Horde.pdb $(ObjDir)/Installer.wixobj $(ObjDir)/HordeServerFiles.wixobj" WorkingDir="$(SourceDir)"/>
 		</Node>
+		
 	</Agent>
+
 </BuildGraph>
diff --git a/Engine/Source/Programs/Horde/Installer/Installer.xml b/Engine/Source/Programs/Horde/Installer/Installer.xml
deleted file mode 100644
