Horde: Fix XML parse error.

diff --git a/Engine/Source/Programs/Horde/BuildHorde.xml b/Engine/Source/Programs/Horde/BuildHorde.xml
--- a/Engine/Source/Programs/Horde/BuildHorde.xml
+++ b/Engine/Source/Programs/Horde/BuildHorde.xml
@@ -11,7 +11,7 @@
 	<Property Name="InformationalVersion" Value="$(InformationalVersion)-PF-$(PreflightChange)" If="'$(PreflightChange)' != ''"/>
 	<Property Name="VersionArguments" Value="/p:Version=$(Version).0 /p:InformationalVersion=$(InformationalVersion)"/>
 
-	<!--- SERVER --->
+	<!-- SERVER -->
 
 	<Agent Name="HordeServer" Type="Linux;Win64_Docker">
 		<Property Name="StagingDir" Value="$(RootDir)/Engine/Saved/Horde.Server"/>
@@ -32,7 +32,7 @@
 		</Node>
 	</Agent>
 	
-	<!--- INSTALLER --->
+	<!-- INSTALLER -->
 	
 	<Agent Name="HordeInstaller" Type="Win64_Docker">
 		<Property Name="StagingDir" Value="$(RootDir)/Staging"/>
