Horde: Fix warning in build script.

diff --git a/Engine/Source/Programs/Horde/BuildHorde.xml b/Engine/Source/Programs/Horde/BuildHorde.xml
--- a/Engine/Source/Programs/Horde/BuildHorde.xml
+++ b/Engine/Source/Programs/Horde/BuildHorde.xml
@@ -34,6 +34,10 @@
 	
 	<!-- INSTALLER -->
 	
+	<Property Name="InstallerDir" Value="$(RootDir)/Engine/Source/Programs/Horde/Installer"/>
+	<Property Name="InstallerConfig" Value="Debug"/>
+	<Property Name="InstallerOutputDir" Value="$(RootDir)/Engine/Source/Programs/Horde/Installer/bin/$(InstallerConfig)"/>
+	
 	<Agent Name="HordeInstaller" Type="Win64_Docker">
 		<Property Name="StagingDir" Value="$(RootDir)/Staging"/>
 
@@ -55,21 +59,20 @@
 		</Node>
 
 		<Property Name="WixDir" Value="$(RootDir)/Engine/Source/ThirdParty/WiX/3.8"/>
-		<Property Name="SourceDir" Value="$(RootDir)/Engine/Source/Programs/Horde/Installer"/>
-		<Property Name="ObjDir" Value="obj/Debug"/>
-		<Property Name="BinDir" Value="bin/Debug"/>
+		<Property Name="ObjDir" Value="obj/$(InstallerConfig)"/>
+		<Property Name="BinDir" Value="bin/$(InstallerConfig)"/>
 		
 		<Node Name="Build HordeInstaller" Requires="Stage HordeInstaller">
-			<Property Name="CommonArgs" Value="-dConfiguration=Debug -dPlatform=x64 -arch x64"/>
+			<Property Name="CommonArgs" Value="-dConfiguration=$(InstallerConfig) -dPlatform=x64 -arch x64"/>
 
 			<Delete Files="$(StagingDir)/Server-Bulk/..."/>
 			<Tag Files="$(StagingDir)/Server/..." Except="HordeServer.exe" With="#ServerFiles"/>
 			<Copy Files="#ServerFiles" From="$(StagingDir)/Server" To="$(StagingDir)/Server-Bulk/Server"/>
-			<Spawn Exe="$(WixDir)/heat.exe" Arguments="dir &quot;$(StagingDir)/Server-Bulk&quot; -cg HordeServerFiles -scom -sreg -gg -dr InstallDir -srd -var var.SourceDir -out &quot;$(SourceDir)/HordeServerFiles.wxs&quot;"/>
-			<Spawn Exe="$(WixDir)/candle.exe" Arguments="$(CommonArgs) HordeServerFiles.wxs -dSourceDir=&quot;$(StagingDir)&quot; -out $(ObjDir)/HordeServerFiles.wixobj" WorkingDir="$(SourceDir)"/>
+			<Spawn Exe="$(WixDir)/heat.exe" Arguments="dir &quot;$(StagingDir)/Server-Bulk&quot; -cg HordeServerFiles -scom -sreg -gg -dr InstallDir -srd -var var.SourceDir -out &quot;$(InstallerDir)/HordeServerFiles.wxs&quot;"/>
+			<Spawn Exe="$(WixDir)/candle.exe" Arguments="$(CommonArgs) HordeServerFiles.wxs -dSourceDir=&quot;$(StagingDir)&quot; -out $(ObjDir)/HordeServerFiles.wixobj" WorkingDir="$(InstallerDir)"/>
 
-			<Spawn Exe="$(WixDir)/candle.exe" Arguments="$(CommonArgs) Installer.wxs -dStagingDir=&quot;$(StagingDir)&quot; -out $(ObjDir)/Installer.wixobj" WorkingDir="$(SourceDir)"/>
-			<Spawn Exe="$(WixDir)/light.exe" Arguments="-out $(BinDir)/Horde.msi -sw1076 -pdbout $(BinDir)/Horde.pdb $(ObjDir)/Installer.wixobj $(ObjDir)/HordeServerFiles.wixobj" WorkingDir="$(SourceDir)"/>
+			<Spawn Exe="$(WixDir)/candle.exe" Arguments="$(CommonArgs) Installer.wxs -dStagingDir=&quot;$(StagingDir)&quot; -out $(ObjDir)/Installer.wixobj" WorkingDir="$(InstallerDir)"/>
+			<Spawn Exe="$(WixDir)/light.exe" Arguments="-out $(BinDir)/Horde.msi -sw1076 -pdbout $(BinDir)/Horde.pdb $(ObjDir)/Installer.wixobj $(ObjDir)/HordeServerFiles.wixobj" WorkingDir="$(InstallerDir)"/>
 		</Node>
 		
 	</Agent>
