Horde: Attempt to use regular Win64 agents for installer builds.

diff --git a/Engine/Source/Programs/Horde/BuildHorde.xml b/Engine/Source/Programs/Horde/BuildHorde.xml
--- a/Engine/Source/Programs/Horde/BuildHorde.xml
+++ b/Engine/Source/Programs/Horde/BuildHorde.xml
@@ -38,7 +38,7 @@
 	<Property Name="InstallerConfig" Value="Debug"/>
 	<Property Name="InstallerOutputDir" Value="$(RootDir)/Engine/Source/Programs/Horde/Installer/bin/$(InstallerConfig)"/>
 	
-	<Agent Name="HordeInstaller" Type="Win64_Docker">
+	<Agent Name="HordeInstaller" Type="Win64;CompileWin64">
 		<Property Name="StagingDir" Value="$(RootDir)/Staging"/>
 
 		<Node Name="Stage Horde Installer">
