Horde: Avoid using observable.ref due to memory leaks
#rnx
#skipci

diff --git a/Engine/Source/Programs/Horde/HordeDashboard/src/backend/AgentStore.ts b/Engine/Source/Programs/Horde/HordeDashboard/src/backend/AgentStore.ts
--- a/Engine/Source/Programs/Horde/HordeDashboard/src/backend/AgentStore.ts
+++ b/Engine/Source/Programs/Horde/HordeDashboard/src/backend/AgentStore.ts
@@ -10,24 +10,40 @@ export class AgentStore {
         makeObservable(this);
     }
 
-    @observable.ref
-	agents: AgentData[] = [];
+    @observable
+    agentsUpdated: number = 0
+
+    get agents(): AgentData[] {
+        // subscribe in any observers
+        if (this.agentsUpdated) { }
+        return this._agents;
+    }
+            
+	private _agents: AgentData[] = [];
 	
-	@observable.ref
-	pools: PoolData[] = [];
+    @observable
+    poolsUpdated: number = 0;
+
+    get pools(): PoolData[] {
+        // subscribe in any observers
+        if (this.poolsUpdated) { }
+        return this._pools;
+    }
+    
+    private _pools: PoolData[] = [];
 
 	// don't need this for now
     modifiedAfterDate: Date = new Date(0);
 
     @action 
     private _setPools(data: PoolData[]) {
-        this.pools = data;
+        this._pools = data;
     }
 
     @action
     private _setAgents(data: AgentData[]) {
         if(data.length !== 0) {
-            const toSet: AgentData[] = [...this.agents];
+            const toSet: AgentData[] = [...this._agents];
             data.forEach(updatedAgent => {
                 const newUpdateIdx = toSet.findIndex(existingAgent => existingAgent.id === updatedAgent.id);
                 if(newUpdateIdx !== -1) {
@@ -37,7 +53,7 @@ export class AgentStore {
                     toSet.push(updatedAgent);
                 }
             });
-            this.agents = toSet;
+            this._agents = toSet;
         }
 	}
 
diff --git a/Engine/Source/Programs/Horde/HordeDashboard/src/backend/IssueStore.ts b/Engine/Source/Programs/Horde/HordeDashboard/src/backend/IssueStore.ts
--- a/Engine/Source/Programs/Horde/HordeDashboard/src/backend/IssueStore.ts
+++ b/Engine/Source/Programs/Horde/HordeDashboard/src/backend/IssueStore.ts
@@ -13,8 +13,16 @@ export class IssueStore {
         makeObservable(this);
     }
 
-    @observable.ref
-    issues: IssueData[] = [];
+    @observable
+    issuesUpdated: number = 0;
+
+    get issues(): IssueData[] {
+        // subscribe in any observers
+        if (this.issuesUpdated) { }
+        return this._issues;
+    }
+        
+    private _issues: IssueData[] = [];
 
     getIssues(jobId: string, stepId: string): IssueData[] {
         // @todo: refactor for new issue API
@@ -29,7 +37,7 @@ export class IssueStore {
 
         this.streamId = streamId;
 
-        if (this.issues.length) {
+        if (this._issues.length) {
             this.setData([]);
         }
 
@@ -45,7 +53,8 @@ export class IssueStore {
     @action
     private setData(data: IssueData[]) {
         // @todo: do a deep compare here to avoid redundant updates
-        this.issues = data;
+        this._issues = data;
+        this.issuesUpdated++;
     }
 
     private update(): void {
diff --git a/Engine/Source/Programs/Horde/HordeDashboard/src/backend/ProjectStore.ts b/Engine/Source/Programs/Horde/HordeDashboard/src/backend/ProjectStore.ts
--- a/Engine/Source/Programs/Horde/HordeDashboard/src/backend/ProjectStore.ts
+++ b/Engine/Source/Programs/Horde/HordeDashboard/src/backend/ProjectStore.ts
@@ -10,15 +10,23 @@ export class ProjectStore {
         makeObservable(this);
     }
 
-    @observable.ref
-    projects: ProjectData[] = [];
+    @observable
+    projectsUpdated: number = 0;
+
+    get projects(): ProjectData[] {
+        // subscribe in any observers
+        if (this.projectsUpdated) { }
+        return this._projects;
+    }
 
     @observable
     activeId?: string;
 
     @action
     setProjects(data: ProjectData[]) {
-        this.projects = data;
+        // @todo: compare to avoid update
+        this._projects = data;
+        this.projectsUpdated++;
     }
 
     @action
@@ -32,23 +40,23 @@ export class ProjectStore {
     }
 
     getActive(): ProjectData | undefined {
-        return this.projects.find(p => p.id === this.activeId);
+        return this._projects.find(p => p.id === this.activeId);
     }
 
     byId(id: string | undefined): ProjectData | undefined {
-        return this.projects.find(p => p.id === id);
+        return this._projects.find(p => p.id === id);
     }
 
     byName(name: string | undefined): ProjectData | undefined {
-        return this.projects.find(p => p.name === name);
+        return this._projects.find(p => p.name === name);
     }
 
     firstProject(): ProjectData | undefined {
-        return this.projects.length ? this.projects[0] : undefined;
+        return this._projects.length ? this._projects[0] : undefined;
     }
 
     projectByStreamId(streamId: string): ProjectData | undefined {
-        return this.projects.find(p => p.streams?.find(stream => stream.id === streamId))
+        return this._projects.find(p => p.streams?.find(stream => stream.id === streamId))
     }
 
     streamById(id: string | undefined): StreamData | undefined {
@@ -56,7 +64,7 @@ export class ProjectStore {
             return undefined;
         }
         let stream: StreamData | undefined;
-        this.projects.forEach(p => {
+        this._projects.forEach(p => {
             if (stream) {
                 return;
             }
@@ -70,7 +78,7 @@ export class ProjectStore {
 
         let stream: StreamData | undefined;
 
-        this.projects.forEach(p => {
+        this._projects.forEach(p => {
             if (stream) {
                 return;
             }
@@ -108,6 +116,10 @@ export class ProjectStore {
             });
         });
     }
+
+
+    private _projects: ProjectData[] = [];
+
 }
 
 export const projectStore = new ProjectStore();
\ No newline at end of file
diff --git a/Engine/Source/Programs/Horde/HordeDashboard/src/components/ChangeCallout.tsx b/Engine/Source/Programs/Horde/HordeDashboard/src/components/ChangeCallout.tsx
--- a/Engine/Source/Programs/Horde/HordeDashboard/src/components/ChangeCallout.tsx
+++ b/Engine/Source/Programs/Horde/HordeDashboard/src/components/ChangeCallout.tsx
@@ -59,13 +59,22 @@ class CommitCache {
       });
    }
 
+   get commitData(): ChangeSummaryData | undefined {
+      // subscribe in any observers
+      if (this.commitDataUpdated) { }
+      return this._commitData;
+   }
+
    @action
    setData(commitData: ChangeSummaryData) {
-      this.commitData = commitData;
+      this._commitData = commitData;
+      this.commitDataUpdated++;
    }
 
-   @observable.ref
-   commitData?: ChangeSummaryData;
+   @observable
+   commitDataUpdated = 0;
+
+   private _commitData?: ChangeSummaryData;
 
    change: number;
 
diff --git a/Engine/Source/Programs/Horde/HordeDashboard/src/components/HistoryModal.tsx b/Engine/Source/Programs/Horde/HordeDashboard/src/components/HistoryModal.tsx
--- a/Engine/Source/Programs/Horde/HordeDashboard/src/components/HistoryModal.tsx
+++ b/Engine/Source/Programs/Horde/HordeDashboard/src/components/HistoryModal.tsx
@@ -48,8 +48,17 @@ class HistoryModalState {
    constructor() {
       makeObservable(this);
    }
+
+   @observable
+   selectedAgentUpdated = 0;
+
+   get selectedAgent(): AgentData | undefined {
+      // subscribe in any observers
+      if (this.selectedAgentUpdated) { }
+      return this._selectedAgent;
+   } 
    
-   @observable.ref selectedAgent: AgentData | undefined = undefined;
+   private _selectedAgent: AgentData | undefined = undefined;
    @observable.shallow currentData: any = [];
    @observable.shallow infoItems: InfoPanelItem[] = [];
    @observable.shallow infoSubItems: InfoPanelSubItem[] = [];
@@ -69,7 +78,7 @@ class HistoryModalState {
    @action
    setSelectedAgent(selectedAgent?: AgentData | undefined) {
       // if we're closing don't reset
-      this.selectedAgent = selectedAgent;
+      this._selectedAgent = selectedAgent;
       if (this.selectedAgent) {
          // set date to now
          this._initBuilderItems();
@@ -81,6 +90,8 @@ class HistoryModalState {
          }
          this._doUpdate();
       }
+
+      this.selectedAgentUpdated++;
    }
 
    @action
diff --git a/Engine/Source/Programs/Horde/HordeDashboard/src/components/JobLabelCallout.tsx b/Engine/Source/Programs/Horde/HordeDashboard/src/components/JobLabelCallout.tsx
--- a/Engine/Source/Programs/Horde/HordeDashboard/src/components/JobLabelCallout.tsx
+++ b/Engine/Source/Programs/Horde/HordeDashboard/src/components/JobLabelCallout.tsx
@@ -10,13 +10,13 @@ import { JobDetails } from "../backend/JobDetails";
 import { StepStatusIcon } from './StatusIcon';
 
 type StepItem = {
-	step: StepData;
+   step: StepData;
 };
 
 export type CalloutState = {
-	jobId?: string;
-	label?: LabelData;
-	target?: string;
+   jobId?: string;
+   label?: LabelData;
+   target?: string;
 }
 
 export class CalloutController {
@@ -25,152 +25,162 @@ export class CalloutController {
       makeObservable(this);
    }
 
-	setState(state: CalloutState | undefined, now: boolean = false) {
+   setState(state: CalloutState | undefined, now: boolean = false) {
 
-		if (!state) {
+      if (!state) {
 
-			clearTimeout(this.setTimeout);
-			clearTimeout(this.closeTimeout);
-			this.setTimeout = this.closeTimeout = undefined;
-			this.pending = undefined;
+         clearTimeout(this.setTimeout);
+         clearTimeout(this.closeTimeout);
+         this.setTimeout = this.closeTimeout = undefined;
+         this.pending = undefined;
 
-			if (now) {
-				this.setStateInternal({});
-				this.details.clear();
-			} else {
-				this.closeTimeout = setTimeout(() => {
-					this.setStateInternal({});
-					this.details.clear();
-				}, 600);
+         if (now) {
+            this.setStateInternal({});
+            this.details.clear();
+         } else {
+            this.closeTimeout = setTimeout(() => {
+               this.setStateInternal({});
+               this.details.clear();
+            }, 600);
 
-			}
+         }
 
-			return;
+         return;
 
-		}
+      }
 
-		// check pending state
-		if (state.jobId === this.pending?.jobId && state.target === this.pending?.target) {
-			return;
-		}
+      // check pending state
+      if (state.jobId === this.pending?.jobId && state.target === this.pending?.target) {
+         return;
+      }
 
-		// check current state
-		if (state.jobId === this.state.jobId && state.target === this.state.target) {
-			return;
-		}
+      // check current state
+      if (state.jobId === this.state.jobId && state.target === this.state.target) {
+         return;
+      }
 
 
-		clearTimeout(this.setTimeout);
-		clearTimeout(this.closeTimeout);
-		this.closeTimeout = undefined;
-		this.pending = state;
+      clearTimeout(this.setTimeout);
+      clearTimeout(this.closeTimeout);
+      this.closeTimeout = undefined;
+      this.pending = state;
 
-		this.setTimeout = setTimeout(() => {
+      this.setTimeout = setTimeout(() => {
 
-			this.setTimeout = undefined;
-			this.pending = undefined;
+         this.setTimeout = undefined;
+         this.pending = undefined;
 
-			if (this.details.id === state.jobId) {
-				this.setStateInternal(state);
-				return;
-			}
+         if (this.details.id === state.jobId) {
+            this.setStateInternal(state);
+            return;
+         }
 
-			this.details.set(state.jobId!, undefined, undefined, undefined, 0, () => {
-				this.setStateInternal(state);
-			});
-		}, 800);
+         this.details.set(state.jobId!, undefined, undefined, undefined, 0, () => {
+            this.setStateInternal(state);
+         });
+      }, 800);
 
-	}
+   }
 
-	@action
-	private setStateInternal(state: CalloutState) {
-		this.state = { ...state };
-	}
+   @action
+   private setStateInternal(state: CalloutState) {
+      this._state = { ...state };
+      this.stateUpdated++;
+   }
 
-	@action
-	clear() {
-		this.details.clear();
-		this.state = {};
-	}
+   @action
+   clear() {
+      this.details.clear();
+      this._state = {};
+      this.stateUpdated++;
+   }
 
-	@observable.ref
-	state: CalloutState = {};
+   @observable
+   stateUpdated = 0;
+
+   get state(): CalloutState {
+      // subscribe
+      if (this.stateUpdated) { }
+      return this._state;
+   }
 
-	pending?: CalloutState;
+   private _state: CalloutState = {};
 
-	details = new JobDetails();
+   pending?: CalloutState;
 
-	setTimeout: any;
-	closeTimeout: any;
+   details = new JobDetails();
+
+   setTimeout: any;
+   closeTimeout: any;
 }
 
 export const JobLabelCallout: React.FC<{ controller: CalloutController }> = observer(({ controller }) => {
 
-	const state = controller.state;
-	const label = state.label;
-	const jobId = state.jobId;
-	const target = state.target;
-	const details = controller.details;
-
-	if (!label) {
-		return <div />;
-	}
-
-	const nodes = details.nodes?.filter(n => label.includedNodes?.find(on => on === n.name));
-	const allSteps = details.getSteps();
-	const steps = allSteps.filter(step => nodes.indexOf(details.nodeByStepId(step.id)!) !== -1);
-	const items = steps.map(step => { return { step: step }; });
-
-
-	const onRenderCell = (stepItem?: StepItem): JSX.Element => {
-
-		const step = stepItem?.step;
-
-		if (!step) {
-			return <div />;
-		}
-
-		const stepUrl = `/job/${jobId}?step=${step.id}`;
-
-		const stepId = step.id;
-		const stepName = details.getStepName(stepId);
-
-		return <Stack tokens={{ childrenGap: 12 }} >
-			<Stack horizontal>
-				<Link to={stepUrl} onClick={(ev) => { ev.stopPropagation(); }}><div style={{ cursor: "pointer" }}>
-					<Stack horizontal>
-						<StepStatusIcon step={step} style={{ fontSize: 10 }} />
-						<Text styles={{ root: { fontSize: 10, color: "#000000", paddingRight: 4, paddingTop: 0, userSelect: "none" } }}>{`${stepName}`}</Text>
-					</Stack>
-				</div></Link>
-			</Stack>
-		</Stack>;
-	};
-
-
-	return <Callout isBeakVisible={true}
-		onDismiss={() => { controller.setState(undefined) }}
-		beakWidth={12}
-		hidden={false}
-		target={target}
-		role="alertdialog"
-		gapSpace={8}
-		setInitialFocus={true}
-		shouldRestoreFocus={false}
-		directionalHint={DirectionalHint.bottomCenter}>
-		<Stack styles={{ root: { paddingTop: '4x', paddingLeft: '14px', paddingBottom: '20px', paddingRight: '20px' } }}
-			onMouseMove={(ev) => ev.stopPropagation()}
-			onMouseOver={(ev) => {
-				clearTimeout(controller.closeTimeout)
-			}}
-			onMouseLeave={(ev) => {
-				controller.setState(undefined, true);
-			}}>
-			<Stack styles={{ root: { paddingTop: 12, paddingLeft: 4 } }}>
-				<List id="steplist" items={items}
-					onRenderCell={onRenderCell}
-					data-is-focusable={false} />
-			</Stack>
-		</Stack>
-	</Callout>
+   const state = controller.state;
+   const label = state.label;
+   const jobId = state.jobId;
+   const target = state.target;
+   const details = controller.details;
+
+   if (!label) {
+      return <div />;
+   }
+
+   const nodes = details.nodes?.filter(n => label.includedNodes?.find(on => on === n.name));
+   const allSteps = details.getSteps();
+   const steps = allSteps.filter(step => nodes.indexOf(details.nodeByStepId(step.id)!) !== -1);
+   const items = steps.map(step => { return { step: step }; });
+
+
+   const onRenderCell = (stepItem?: StepItem): JSX.Element => {
+
+      const step = stepItem?.step;
+
+      if (!step) {
+         return <div />;
+      }
+
+      const stepUrl = `/job/${jobId}?step=${step.id}`;
+
+      const stepId = step.id;
+      const stepName = details.getStepName(stepId);
+
+      return <Stack tokens={{ childrenGap: 12 }} >
+         <Stack horizontal>
+            <Link to={stepUrl} onClick={(ev) => { ev.stopPropagation(); }}><div style={{ cursor: "pointer" }}>
+               <Stack horizontal>
+                  <StepStatusIcon step={step} style={{ fontSize: 10 }} />
+                  <Text styles={{ root: { fontSize: 10, color: "#000000", paddingRight: 4, paddingTop: 0, userSelect: "none" } }}>{`${stepName}`}</Text>
+               </Stack>
+            </div></Link>
+         </Stack>
+      </Stack>;
+   };
+
+
+   return <Callout isBeakVisible={true}
+      onDismiss={() => { controller.setState(undefined) }}
+      beakWidth={12}
+      hidden={false}
+      target={target}
+      role="alertdialog"
+      gapSpace={8}
+      setInitialFocus={true}
+      shouldRestoreFocus={false}
+      directionalHint={DirectionalHint.bottomCenter}>
+      <Stack styles={{ root: { paddingTop: '4x', paddingLeft: '14px', paddingBottom: '20px', paddingRight: '20px' } }}
+         onMouseMove={(ev) => ev.stopPropagation()}
+         onMouseOver={(ev) => {
+            clearTimeout(controller.closeTimeout)
+         }}
+         onMouseLeave={(ev) => {
+            controller.setState(undefined, true);
+         }}>
+         <Stack styles={{ root: { paddingTop: 12, paddingLeft: 4 } }}>
+            <List id="steplist" items={items}
+               onRenderCell={onRenderCell}
+               data-is-focusable={false} />
+         </Stack>
+      </Stack>
+   </Callout>
 });
\ No newline at end of file
diff --git a/Engine/Source/Programs/Horde/HordeDashboard/src/components/LogSource.tsx b/Engine/Source/Programs/Horde/HordeDashboard/src/components/LogSource.tsx
--- a/Engine/Source/Programs/Horde/HordeDashboard/src/components/LogSource.tsx
+++ b/Engine/Source/Programs/Horde/HordeDashboard/src/components/LogSource.tsx
@@ -78,15 +78,15 @@ export abstract class LogSource {
    resize(lineCount: number): LogItem[] | undefined {
 
       // resize items if necessary
-      let count = lineCount - this.logItems.length;
+      let count = lineCount - this._logItems.length;
 
       if (count <= 0) {
          return undefined;
       }
 
-      const items = [...this.logItems];
+      const items = [...this._logItems];
 
-      let line = this.logItems.length + 1;
+      let line = this._logItems.length + 1;
       while (count--) {
          items.push({
             lineNumber: line++,
@@ -110,10 +110,10 @@ export abstract class LogSource {
 
          for (let i = 0; i < count; i++) {
             const offset = i + index;
-            if (offset >= this.logItems.length) {
+            if (offset >= this._logItems.length) {
                break;
             }
-            const item = this.logItems![i + index];
+            const item = this._logItems![i + index];
             if (!item.requested) {
                item.requested = true;
                anyRequested = true;
@@ -135,10 +135,10 @@ export abstract class LogSource {
                const line = data.lines![i];
 
                const offset = i + data.index;
-               if (offset >= this.logItems.length) {
+               if (offset >= this._logItems.length) {
                   break;
                }
-               const item = this.logItems[offset];
+               const item = this._logItems[offset];
                item.line = line;
 
 
@@ -151,7 +151,7 @@ export abstract class LogSource {
 
             }
 
-            this.setLogItems([...this.logItems]);
+            this.setLogItems([...this._logItems]);
             resolve(true);
 
          }).catch(reason => reject(reason));
@@ -236,7 +236,8 @@ export abstract class LogSource {
 
    @action
    setLogItems(items: LogItem[]) {
-      this.logItems = items;
+      this._logItems = items;
+      this.logItemsUpdated++;
    }
 
    @action
@@ -249,8 +250,16 @@ export abstract class LogSource {
       this.fatalError = error;
    }
 
-   @observable.ref
-   logItems: LogItem[] = [];
+   @observable
+   logItemsUpdated = 0;      
+
+   get logItems(): LogItem[] {
+      // subscribe
+      if (this.logItemsUpdated) { }
+      return this._logItems;
+   }
+
+   private _logItems: LogItem[] = [];
 
    @observable
    active = true;
