Horde: Add a compute command for forking the remote message pump. Allows running multiple requests simultaneously, while multiplexing message data onto a single underlying socket.

diff --git a/Engine/Source/Programs/Horde/Horde.Agent/Commands/Compute/XorComputeCommand.cs b/Engine/Source/Programs/Horde/Horde.Agent/Commands/Compute/XorComputeCommand.cs
--- a/Engine/Source/Programs/Horde/Horde.Agent/Commands/Compute/XorComputeCommand.cs
+++ b/Engine/Source/Programs/Horde/Horde.Agent/Commands/Compute/XorComputeCommand.cs
@@ -31,8 +31,13 @@ namespace Horde.Agent.Commands.Compute
 		/// <inheritdoc/>
 		protected override async Task<bool> HandleRequestAsync(IComputeLease lease, CancellationToken cancellationToken)
 		{
-			await using (IComputeChannel channel = lease.CreateChannel(0))
+			await using (IComputeChannel mainChannel = lease.CreateChannel(0))
 			{
+				_logger.LogInformation("Forking compute channel...");
+				mainChannel.Fork(1);
+
+				await using IComputeChannel channel = lease.CreateChannel(1);
+
 				_logger.LogInformation("Sending XOR request");
 				channel.XorRequest(new byte[] { 1, 2, 3, 4, 5 }, (byte)123);
 
diff --git a/Engine/Source/Programs/Horde/Horde.Agent/Leases/Handlers/ComputeHandler.cs b/Engine/Source/Programs/Horde/Horde.Agent/Leases/Handlers/ComputeHandler.cs
--- a/Engine/Source/Programs/Horde/Horde.Agent/Leases/Handlers/ComputeHandler.cs
+++ b/Engine/Source/Programs/Horde/Horde.Agent/Leases/Handlers/ComputeHandler.cs
@@ -77,29 +77,46 @@ namespace Horde.Agent.Leases.Handlers
 
 		public async Task RunAsync(IComputeLease lease, CancellationToken cancellationToken)
 		{
-			await using IComputeChannel channel = lease.CreateChannel(0);
-			for (; ; )
+			await RunAsync(lease, 0, cancellationToken);
+		}
+
+		async Task RunAsync(IComputeLease lease, int id, CancellationToken cancellationToken)
+		{
+			await using (IComputeChannel channel = lease.CreateChannel(id))
 			{
-				using IComputeMessage message = await channel.ReceiveAsync(cancellationToken);
-				switch (message.Type)
+				List<Task> childTasks = new List<Task>();
+				for (; ; )
 				{
-					case ComputeMessageType.None:
-						return;
-					case ComputeMessageType.XorRequest:
-						{
-							XorRequestMessage xorRequest = message.AsXorRequest();
-							RunXor(channel, xorRequest.Data, xorRequest.Value);
-						}
-						break;
-					case ComputeMessageType.CppBegin:
-						{
-							CppBeginMessage cppBegin = message.AsCppBegin();
-							await using IComputeChannel replyChannel = lease.CreateChannel(cppBegin.ReplyChannelId);
-							await RunCppAsync(replyChannel, cppBegin.Locator, cancellationToken);
-						}
-						break;
-					default:
-						throw new NotImplementedException();
+					using IComputeMessage message = await channel.ReceiveAsync(cancellationToken);
+					_logger.LogTrace("Compute Channel {ChannelId}: {MessageType}", id, message.Type);
+
+					switch (message.Type)
+					{
+						case ComputeMessageType.None:
+							await Task.WhenAll(childTasks);
+							return;
+						case ComputeMessageType.Fork:
+							{
+								ForkMessage fork = message.AsForkMessage();
+								childTasks.Add(RunAsync(lease, fork.Id, cancellationToken));
+							}
+							break;
+						case ComputeMessageType.XorRequest:
+							{
+								XorRequestMessage xorRequest = message.AsXorRequest();
+								RunXor(channel, xorRequest.Data, xorRequest.Value);
+							}
+							break;
+						case ComputeMessageType.CppBegin:
+							{
+								CppBeginMessage cppBegin = message.AsCppBegin();
+								await using IComputeChannel replyChannel = lease.CreateChannel(cppBegin.ReplyChannelId);
+								await RunCppAsync(replyChannel, cppBegin.Locator, cancellationToken);
+							}
+							break;
+						default:
+							throw new NotImplementedException();
+					}
 				}
 			}
 		}
diff --git a/Engine/Source/Programs/Shared/EpicGames.Horde/Compute/IComputeChannel.cs b/Engine/Source/Programs/Shared/EpicGames.Horde/Compute/IComputeChannel.cs
--- a/Engine/Source/Programs/Shared/EpicGames.Horde/Compute/IComputeChannel.cs
+++ b/Engine/Source/Programs/Shared/EpicGames.Horde/Compute/IComputeChannel.cs
@@ -1,7 +1,6 @@
 // Copyright Epic Games, Inc. All Rights Reserved.
 
 using System;
-using System.Buffers;
 using System.Threading;
 using System.Threading.Tasks;
 using EpicGames.Core;
diff --git a/Engine/Source/Programs/Shared/EpicGames.Horde/Compute/IComputeMessage.cs b/Engine/Source/Programs/Shared/EpicGames.Horde/Compute/IComputeMessage.cs
--- a/Engine/Source/Programs/Shared/EpicGames.Horde/Compute/IComputeMessage.cs
+++ b/Engine/Source/Programs/Shared/EpicGames.Horde/Compute/IComputeMessage.cs
@@ -1,6 +1,7 @@
 // Copyright Epic Games, Inc. All Rights Reserved.
 
 using System;
+using System.Buffers.Binary;
 using System.Threading;
 using System.Threading.Tasks;
 using EpicGames.Core;
@@ -17,6 +18,11 @@ namespace EpicGames.Horde.Compute
 		/// </summary>
 		None = 0x00,
 
+		/// <summary>
+		/// Fork a new request channel
+		/// </summary>
+		Fork = 0x01,
+
 		#region Test Requests
 
 		/// <summary>
@@ -82,6 +88,12 @@ namespace EpicGames.Horde.Compute
 		public ReadOnlyMemory<byte> Data { get; }
 	}
 
+	/// <summary>
+	/// Message for forking the message loop
+	/// </summary>
+	/// <param name="Id">New channel id</param>
+	public record struct ForkMessage(int Id);
+
 	/// <summary>
 	/// Message for running an XOR command
 	/// </summary>
@@ -94,6 +106,29 @@ namespace EpicGames.Horde.Compute
 	/// </summary>
 	public static class ComputeMessageExtensions
 	{
+		/// <summary>
+		/// Creates a new remote message loop using the given channel id
+		/// </summary>
+		/// <param name="channel">Current channel</param>
+		/// <param name="id">New channel id</param>
+		public static void Fork(this IComputeChannel channel, int id)
+		{
+			using (IComputeMessageBuilder builder = channel.CreateMessage(ComputeMessageType.Fork))
+			{
+				builder.WriteInt32(id);
+				builder.Send();
+			}
+		}
+
+		/// <summary>
+		/// Parses a fork message from the given compute message
+		/// </summary>
+		public static ForkMessage AsForkMessage(this IComputeMessage message)
+		{
+			ReadOnlyMemory<byte> data = message.Data;
+			return new ForkMessage(BinaryPrimitives.ReadInt32LittleEndian(data.Span));
+		}
+
 		/// <summary>
 		/// Send a message to request that a byte string be xor'ed with a particular value
 		/// </summary>
