Horde: Fix Serilog configuration.

* Previous code was adding a Serilog provider to the Microsoft logger, causing settings for both Microsoft's logging framework and Serilogs settings to apply. This was problematic for filtering, because each system imposes it's own minimum output level.
* Now constructs a single Serilog logger factory using settings read from config files.
* Moved all channel verbosity overrides into config files.

diff --git a/Engine/Source/Programs/Horde/Horde.Agent/Commands/Workspace/WorkspaceBase.cs b/Engine/Source/Programs/Horde/Horde.Agent/Commands/Workspace/WorkspaceBase.cs
--- a/Engine/Source/Programs/Horde/Horde.Agent/Commands/Workspace/WorkspaceBase.cs
+++ b/Engine/Source/Programs/Horde/Horde.Agent/Commands/Workspace/WorkspaceBase.cs
@@ -68,9 +68,7 @@ namespace Horde.Agent.Commands.Workspace
 			InfoRecord info = await perforce.GetInfoAsync(InfoOptions.ShortOutput, CancellationToken.None);
 			bool useHaveTable = UseHaveTable.Equals("true", StringComparison.Ordinal);
 
-			using Logging.HordeLoggerProvider loggerProvider = new ();
-			ILogger repoLogger = loggerProvider.CreateLogger("Repository");
-			ManagedWorkspace repo = await ManagedWorkspace.LoadOrCreateAsync(info.ClientHost!, BaseDir, Overwrite, useHaveTable, repoLogger, CancellationToken.None);
+			ManagedWorkspace repo = await ManagedWorkspace.LoadOrCreateAsync(info.ClientHost!, BaseDir, Overwrite, useHaveTable, logger, CancellationToken.None);
 			await ExecuteAsync(perforce, repo, logger);
 			return 0;
 		}
diff --git a/Engine/Source/Programs/Horde/Horde.Agent/Horde.Agent.csproj b/Engine/Source/Programs/Horde/Horde.Agent/Horde.Agent.csproj
--- a/Engine/Source/Programs/Horde/Horde.Agent/Horde.Agent.csproj
+++ b/Engine/Source/Programs/Horde/Horde.Agent/Horde.Agent.csproj
@@ -53,6 +53,7 @@
     <PackageReference Include="Serilog" Version="2.12.0" />
     <PackageReference Include="Serilog.Extensions.Hosting" Version="5.0.1" />
     <PackageReference Include="Serilog.Extensions.Logging" Version="3.1.0" />
+    <PackageReference Include="Serilog.Settings.Configuration" Version="3.4.0" />
     <PackageReference Include="Serilog.Sinks.Console" Version="4.1.0" />
     <PackageReference Include="Serilog.Sinks.File" Version="5.0.0" />
     <PackageReference Include="System.ServiceProcess.ServiceController" Version="6.0.0" />
diff --git a/Engine/Source/Programs/Horde/Horde.Agent/Logging.cs b/Engine/Source/Programs/Horde/Horde.Agent/Logging.cs
--- a/Engine/Source/Programs/Horde/Horde.Agent/Logging.cs
+++ b/Engine/Source/Programs/Horde/Horde.Agent/Logging.cs
@@ -3,7 +3,10 @@
 using System;
 using System.Runtime.InteropServices;
 using EpicGames.Core;
+using Microsoft.Extensions.Configuration;
+using Microsoft.Extensions.DependencyInjection;
 using Microsoft.Extensions.Logging;
+using Microsoft.Extensions.Options;
 using OpenTracing;
 using OpenTracing.Util;
 using Serilog;
@@ -11,6 +14,7 @@ using Serilog.Core;
 using Serilog.Events;
 using Serilog.Extensions.Logging;
 using Serilog.Formatting.Json;
+using Serilog.Settings.Configuration;
 using Serilog.Sinks.SystemConsole.Themes;
 
 namespace Horde.Agent
@@ -24,8 +28,6 @@ namespace Horde.Agent
 			s_env = newEnv;
 		}
 
-		private static readonly Lazy<Serilog.ILogger> s_logger = new Lazy<Serilog.ILogger>(CreateSerilogLogger, true);
-		
 		public static LoggingLevelSwitch LogLevelSwitch =  new LoggingLevelSwitch();
 
 		private class DatadogLogEnricher : ILogEventEnricher
@@ -44,30 +46,14 @@ namespace Horde.Agent
 				}
 			}
 		}
-		
-		public sealed class HordeLoggerProvider : ILoggerProvider
-		{
-			private readonly SerilogLoggerProvider _inner;
-
-			public HordeLoggerProvider()
-			{
-				_inner = new SerilogLoggerProvider(s_logger.Value);
-			}
 
-			public Microsoft.Extensions.Logging.ILogger CreateLogger(string categoryName)
-			{
-				Microsoft.Extensions.Logging.ILogger logger = _inner.CreateLogger(categoryName);
-				logger = new DefaultLoggerIndentHandler(logger);
-				return logger;
-			}
-
-			public void Dispose()
-			{
-				_inner.Dispose();
-			}
+		public static ILoggerFactory CreateLoggerFactory(IConfiguration configuration)
+		{
+			Serilog.ILogger logger = CreateSerilogLogger(configuration);
+			return new SerilogLoggerFactory(logger, true);
 		}
 
-		static Serilog.ILogger CreateSerilogLogger()
+		static Serilog.ILogger CreateSerilogLogger(IConfiguration configuration)
 		{
 			DirectoryReference.CreateDirectory(Program.DataDir);
 
@@ -82,23 +68,13 @@ namespace Horde.Agent
 			}
 
 			return new LoggerConfiguration()
-				.MinimumLevel.Debug()
-				.MinimumLevel.Override("Microsoft", LogEventLevel.Information)
-				.MinimumLevel.Override("System.Net.Http.HttpClient", LogEventLevel.Warning)
-				.MinimumLevel.Override("Microsoft.AspNetCore.Routing.EndpointMiddleware", LogEventLevel.Warning)
-				.MinimumLevel.Override("Microsoft.AspNetCore.Authorization.DefaultAuthorizationService", LogEventLevel.Warning)
-				.MinimumLevel.Override("HordeServer.Authentication.HordeJwtBearerHandler", LogEventLevel.Warning)
-				.MinimumLevel.Override("HordeServer.Authentication.OktaHandler", LogEventLevel.Warning)
-				.MinimumLevel.Override("Microsoft.AspNetCore.Hosting.Diagnostics", LogEventLevel.Warning)
-				.MinimumLevel.Override("Microsoft.AspNetCore.Mvc.Infrastructure.ControllerActionInvoker", LogEventLevel.Warning)
-				.MinimumLevel.Override("Serilog.AspNetCore.RequestLoggingMiddleware", LogEventLevel.Warning)
-				.MinimumLevel.Override("Grpc.Net.Client.Internal.GrpcCall", LogEventLevel.Warning)
-				.MinimumLevel.ControlledBy(LogLevelSwitch)
-				.Enrich.FromLogContext()
-				.Enrich.With<DatadogLogEnricher>()
 				.WriteTo.Console(outputTemplate: "[{Timestamp:HH:mm:ss} {Level:w3}] {Indent}{Message:l}{NewLine}{Exception}", theme: theme)
 				.WriteTo.File(FileReference.Combine(Program.DataDir, "Log-.txt").FullName, fileSizeLimitBytes: 50 * 1024 * 1024, rollingInterval: RollingInterval.Day, rollOnFileSizeLimit: true, retainedFileCountLimit: 10)
 				.WriteTo.File(new JsonFormatter(renderMessage: true), FileReference.Combine(Program.DataDir, "Log-.json").FullName, fileSizeLimitBytes: 50 * 1024 * 1024, rollingInterval: RollingInterval.Day, rollOnFileSizeLimit: true, retainedFileCountLimit: 10)
+				.ReadFrom.Configuration(configuration)
+				.MinimumLevel.ControlledBy(LogLevelSwitch)
+				.Enrich.FromLogContext()
+				.Enrich.With<DatadogLogEnricher>()
 				.CreateLogger();
 		}
 	}
diff --git a/Engine/Source/Programs/Horde/Horde.Agent/Program.cs b/Engine/Source/Programs/Horde/Horde.Agent/Program.cs
--- a/Engine/Source/Programs/Horde/Horde.Agent/Program.cs
+++ b/Engine/Source/Programs/Horde/Horde.Agent/Program.cs
@@ -15,6 +15,8 @@ using Datadog.Trace;
 using Datadog.Trace.Configuration;
 using Datadog.Trace.OpenTracing;
 using EpicGames.Core;
+using EpicGames.Horde.Compute;
+using EpicGames.Horde.Compute.Buffers;
 using EpicGames.Horde.Storage;
 using EpicGames.Horde.Storage.Backends;
 using Horde.Agent.Execution;
@@ -96,8 +98,6 @@ namespace Horde.Agent
 		{
 			Program.Args = args;
 
-			using Logging.HordeLoggerProvider loggerProvider = new Logging.HordeLoggerProvider();
-
 			string? environment = Environment.GetEnvironmentVariable("DOTNET_ENVIRONMENT");
 			if (String.IsNullOrEmpty(environment))
 			{
@@ -116,7 +116,7 @@ namespace Horde.Agent
 				configOverrides.Add($"{AgentSettings.SectionName}:{nameof(AgentSettings.WorkingDir)}", workingDirOverride);
 			}
 
-			IConfiguration config = new ConfigurationBuilder()
+			IConfiguration configuration = new ConfigurationBuilder()
 				.SetBasePath(AppDir.FullName)
 				.AddJsonFile("appsettings.json", optional: false)
 				.AddJsonFile("appsettings.Build.json", optional: true) // specific settings for builds (installer/dockerfile)
@@ -126,9 +126,11 @@ namespace Horde.Agent
 				.AddEnvironmentVariables()
 				.Build();
 
+			using ILoggerFactory loggerFactory = Logging.CreateLoggerFactory(configuration);
+
 			IServiceCollection services = new ServiceCollection();
 			services.AddCommandsFromAssembly(Assembly.GetExecutingAssembly());
-			services.AddLogging(builder => builder.AddProvider(loggerProvider));
+			services.AddSingleton(loggerFactory);
 
 			// Enable unencrypted HTTP/2 for gRPC channel without TLS
 			AppContext.SetSwitch("System.Net.Http.SocketsHttpHandler.Http2UnencryptedSupport", true);
@@ -145,7 +147,7 @@ namespace Horde.Agent
 			}
 
 			// Add all the default 
-			IConfigurationSection configSection = config.GetSection(AgentSettings.SectionName);
+			IConfigurationSection configSection = configuration.GetSection(AgentSettings.SectionName);
 			services.AddOptions<AgentSettings>().Configure(options => configSection.Bind(options)).ValidateDataAnnotations();
 
 			AgentSettings settings = new AgentSettings();
@@ -156,7 +158,7 @@ namespace Horde.Agent
 
 			Logging.SetEnv(serverProfile.Environment);
 
-			ILogger certificateLogger = loggerProvider.CreateLogger(typeof(CertificateHelper).FullName!);
+			ILogger certificateLogger = loggerFactory.CreateLogger(typeof(CertificateHelper).FullName!);
 			services.AddHttpClient(Program.HordeServerClientName, config =>
 			{
 				config.BaseAddress = serverProfile.Url;
@@ -210,7 +212,7 @@ namespace Horde.Agent
 			services.AddMemoryCache();
 
 			// Allow commands to augment the service collection for their own DI service providers
-			services.AddSingleton<DefaultServices>(x => new DefaultServices(config, services));
+			services.AddSingleton<DefaultServices>(x => new DefaultServices(configuration, services));
 
 			// Execute all the commands
 			IServiceProvider serviceProvider = services.BuildServiceProvider();
diff --git a/Engine/Source/Programs/Horde/Horde.Agent/appsettings.json b/Engine/Source/Programs/Horde/Horde.Agent/appsettings.json
--- a/Engine/Source/Programs/Horde/Horde.Agent/appsettings.json
+++ b/Engine/Source/Programs/Horde/Horde.Agent/appsettings.json
@@ -13,5 +13,22 @@
 				}
 			}
 		]
+	},
+	"Serilog": {
+		"MinimumLevel": {
+			"Default": "Information",
+			"Override": {
+				"Microsoft": "Information",
+				"System.Net.Http.HttpClient": "Warning",
+				"Microsoft.AspNetCore.Routing.EndpointMiddleware": "Warning",
+				"Microsoft.AspNetCore.Authorization.DefaultAuthorizationService": "Warning",
+				"HordeServer.Authentication.HordeJwtBearerHandler": "Warning",
+				"HordeServer.Authentication.OktaHandler": "Warning",
+				"Microsoft.AspNetCore.Hosting.Diagnostics": "Warning",
+				"Microsoft.AspNetCore.Mvc.Infrastructure.ControllerActionInvoker": "Warning",
+				"Serilog.AspNetCore.RequestLoggingMiddleware": "Warning",
+				"Grpc.Net.Client.Internal.GrpcCall": "Warning"
+			}
+		}
 	}
 }
diff --git a/Engine/Source/Programs/Shared/EpicGames.Core/CommandHost.cs b/Engine/Source/Programs/Shared/EpicGames.Core/CommandHost.cs
--- a/Engine/Source/Programs/Shared/EpicGames.Core/CommandHost.cs
+++ b/Engine/Source/Programs/Shared/EpicGames.Core/CommandHost.cs
@@ -244,7 +244,7 @@ namespace EpicGames.Core
 			}
 
 			// Configure the command
-			ILogger logger = serviceProvider.GetRequiredService<ILoggerProvider>().CreateLogger("CommandHost");
+			ILogger logger = serviceProvider.GetRequiredService<ILogger<Command>>();
 			try
 			{
 				command.Configure(args, logger);
