Horde: Fix garbage output when compute process doesn't fill output buffer.

diff --git a/Engine/Source/Programs/Horde/Horde.Agent/Leases/Handlers/ComputeHandler.cs b/Engine/Source/Programs/Horde/Horde.Agent/Leases/Handlers/ComputeHandler.cs
--- a/Engine/Source/Programs/Horde/Horde.Agent/Leases/Handlers/ComputeHandler.cs
+++ b/Engine/Source/Programs/Horde/Horde.Agent/Leases/Handlers/ComputeHandler.cs
@@ -296,7 +296,7 @@ namespace Horde.Agent.Leases.Handlers
 							channel.SendExecuteProcessResponse(process.ExitCode);
 							return;
 						}
-						channel.SendProcessOutput(buffer);
+						channel.SendProcessOutput(buffer.AsMemory(0, length));
 					}
 				}
 			}
