Horde: Remove logging of cooldowns for fleet service

Ignores them during state change comparisons as well.

diff --git a/Engine/Source/Programs/Horde/Horde.Server/Agents/Fleet/FleetService.cs b/Engine/Source/Programs/Horde/Horde.Server/Agents/Fleet/FleetService.cs
--- a/Engine/Source/Programs/Horde/Horde.Server/Agents/Fleet/FleetService.cs
+++ b/Engine/Source/Programs/Horde/Horde.Server/Agents/Fleet/FleetService.cs
@@ -250,6 +250,7 @@ namespace Horde.Server.Agents.Fleet
 			
 			scope.Span.SetTag("IsDowntimeActive", _downtimeService.IsDowntimeActive);
 			
+			ScaleResult cooldownActiveResult = new(FleetManagerOutcome.NoOp, 0, 0, "Cooldown active");
 			ScaleResult result = new (FleetManagerOutcome.NoOp, 0, 0);
 			try
 			{
@@ -261,7 +262,7 @@ namespace Horde.Server.Agents.Fleet
 					}
 					else if (isScaleOutCoolingDown)
 					{
-						result = new(FleetManagerOutcome.NoOp, 0, 0, "Cannot scale out, cooldown active");
+						result = cooldownActiveResult;
 					}
 					else
 					{
@@ -275,7 +276,7 @@ namespace Horde.Server.Agents.Fleet
 				{
 					if (isScaleInCoolingDown)
 					{
-						result = new (FleetManagerOutcome.NoOp, 0, 0, "Cannot scale in, cooldown active");
+						result = cooldownActiveResult;
 					}
 					else
 					{
@@ -290,17 +291,24 @@ namespace Horde.Server.Agents.Fleet
 				return new ScaleResult(FleetManagerOutcome.Failure, 0, 0, "Exception during scaling. See log.");
 			}
 
-			if (pool.LastScaleResult == null || !pool.LastScaleResult.Equals(result))
+			bool isResultDifferentFromLastTime = pool.LastScaleResult == null || !pool.LastScaleResult.Equals(result);
+			bool isCooldownResult = ReferenceEquals(result, cooldownActiveResult);
+			if (isResultDifferentFromLastTime && !isCooldownResult)
 			{
 				_logger.LogInformation("Scale result: Outcome={Outcome} AgentsAdded={AgentsAdded} AgentsRemoved={AgentsRemoved} Message={Message}", 
 					result.Outcome, result.AgentsAddedCount, result.AgentsAddedCount, result.Message);
 			}
+			
+			scope.Span.SetTag("ResultOutcome", result.Outcome.ToString());
+			scope.Span.SetTag("ResultAgentsAdded", result.AgentsAddedCount);
+			scope.Span.SetTag("ResultAgentsRemoved", result.AgentsRemovedCount);
+			scope.Span.SetTag("ResultOutcome", result.Message);
 
 			await _poolCollection.TryUpdateAsync(
 				pool,
 				lastScaleUpTime: scaleOutTime,
 				lastScaleDownTime: scaleInTime,
-				lastScaleResult: result,
+				lastScaleResult: isCooldownResult ? null : result,
 				lastAgentCount: currentAgentCount,
 				lastDesiredAgentCount: desiredAgentCount);
 
