Horde: Fix log polling issue that would leave the last poll out when job step finished
#rnx
#skipci

diff --git a/Engine/Source/Programs/Horde/HordeDashboard/src/components/LogSource.tsx b/Engine/Source/Programs/Horde/HordeDashboard/src/components/LogSource.tsx
--- a/Engine/Source/Programs/Horde/HordeDashboard/src/components/LogSource.tsx
+++ b/Engine/Source/Programs/Horde/HordeDashboard/src/components/LogSource.tsx
@@ -190,17 +190,14 @@ export abstract class LogSource {
          return;
       }
 
-      if (!this.active) {
-         this.stopPolling();
-         return;
-      }
-
       if (!this.pollMS) {
          return;
       }
 
-      this.pollID = setTimeout(() => { this.poll(); }, this.pollMS);
-
+      if (this.active) {
+         this.pollID = setTimeout(() => { this.poll(); }, this.pollMS);
+      } 
+      
       if (this.polling) {
          return;
       }
@@ -213,9 +210,11 @@ export abstract class LogSource {
          });
       }).finally(() => {
          this.polling = false;
+         if (!this.active) {
+            this.stopPolling();
+         }
       });
 
-
    }
 
    stopPolling() {
