Horde: Add another method of running compute connections locally, which spawns a separate process without server involvement and communicates over the loopback adapter.

diff --git a/Engine/Source/Programs/Horde/Horde.Agent/Commands/Compute/ComputeCommand.cs b/Engine/Source/Programs/Horde/Horde.Agent/Commands/Compute/ComputeCommand.cs
--- a/Engine/Source/Programs/Horde/Horde.Agent/Commands/Compute/ComputeCommand.cs
+++ b/Engine/Source/Programs/Horde/Horde.Agent/Commands/Compute/ComputeCommand.cs
@@ -1,18 +1,26 @@
 // Copyright Epic Games, Inc. All Rights Reserved.
 
 using System;
+using System.Collections.Generic;
+using System.Diagnostics;
+using System.Net;
 using System.Net.Http;
 using System.Net.Http.Headers;
+using System.Net.Sockets;
+using System.Reflection;
 using System.Threading;
 using System.Threading.Tasks;
+using Amazon.EC2.Model;
 using EpicGames.Core;
 using EpicGames.Horde.Common;
 using EpicGames.Horde.Compute;
 using EpicGames.Horde.Compute.Clients;
+using EpicGames.Horde.Compute.Transports;
 using Horde.Agent.Leases.Handlers;
 using Microsoft.Extensions.DependencyInjection;
 using Microsoft.Extensions.Logging;
 using Microsoft.Extensions.Options;
+using Polly;
 
 namespace Horde.Agent.Commands.Compute
 {
@@ -24,6 +32,9 @@ namespace Horde.Agent.Commands.Compute
 		[CommandLine("-Requirements=", Description = "Match the agent to run on")]
 		public string? Requirements { get; set; }
 
+		[CommandLine("-Local")]
+		public bool Local { get; set; }
+
 		[CommandLine("-Loopback")]
 		public bool Loopback { get; set; }
 
@@ -65,9 +76,13 @@ namespace Horde.Agent.Commands.Compute
 		IComputeClient CreateClient()
 		{
 			ILoggerFactory loggerFactory = _serviceProvider.GetRequiredService<ILoggerFactory>();
-			if (Loopback)
+			if (Local)
+			{
+				return new LocalComputeClient((socket, ctx) => ComputeWorkerCommand.RunWorkerAsync(socket, _serviceProvider, ctx), 2000, loggerFactory);
+			}
+			else if (Loopback)
 			{
-				return new LoopbackComputeClient(RunListenerAsync, 2000, loggerFactory);
+				return new LoopbackComputeClient(2000, loggerFactory);
 			}
 			else
 			{
@@ -85,9 +100,116 @@ namespace Horde.Agent.Commands.Compute
 			return client;
 		}
 
-		async Task RunListenerAsync(IComputeSocket socket, CancellationToken cancellationToken)
+		class LoopbackComputeClient : IComputeClient
+		{
+			readonly int _port;
+			readonly ILoggerFactory _loggerFactory;
+
+			public LoopbackComputeClient(int port, ILoggerFactory loggerFactory)
+			{
+				_port = port;
+				_loggerFactory = loggerFactory;
+			}
+
+			public ValueTask DisposeAsync() => new ValueTask();
+
+			public async Task<TResult> ExecuteAsync<TResult>(ClusterId clusterId, Requirements? requirements, Func<IComputeLease, CancellationToken, Task<TResult>> handler, CancellationToken cancellationToken)
+			{
+				ILogger logger = _loggerFactory.CreateLogger<LoopbackComputeClient>();
+				logger.LogInformation("** INITIATOR **");
+
+				TResult result;
+				using (Socket listener = new Socket(SocketType.Stream, ProtocolType.IP))
+				{
+					listener.Bind(new IPEndPoint(IPAddress.Loopback, _port));
+					listener.Listen();
+
+					using (CancellationTokenSource cancellationSource = CancellationTokenSource.CreateLinkedTokenSource(cancellationToken))
+					{
+						Task<TResult> listenTask = ListenAsync<TResult>(listener, handler, cancellationSource.Token);
+						try
+						{
+							using ManagedProcessGroup group = new ManagedProcessGroup();
+
+							using Process process = new Process();
+							process.StartInfo.FileName = "dotnet";
+							process.StartInfo.ArgumentList.Add(Assembly.GetExecutingAssembly().Location);
+							process.StartInfo.ArgumentList.Add("computeworker");
+							process.StartInfo.ArgumentList.Add($"-port={_port}");
+							process.StartInfo.UseShellExecute = true;
+							process.Start();
+
+							group.AddProcess(process.Handle);
+
+							await process.WaitForExitAsync(cancellationToken);
+						}
+						finally
+						{
+							cancellationSource.Cancel();
+						}
+						result = await listenTask;
+					}
+				}
+				return result;
+			}
+
+			async Task<TResult> ListenAsync<TResult>(Socket listener, Func<IComputeLease, CancellationToken, Task<TResult>> handler, CancellationToken cancellationToken)
+			{
+				TResult result;
+				using (Socket tcpSocket = await listener.AcceptAsync(cancellationToken))
+				{
+					await using (ComputeSocket socket = new ComputeSocket(new TcpTransport(tcpSocket), _loggerFactory))
+					{
+						await using ComputeLease lease = new ComputeLease(new List<string>(), new Dictionary<string, int>(), socket);
+						result = await handler(lease, cancellationToken);
+						await socket.CloseAsync(cancellationToken);
+					}
+				}
+				return result;
+			}
+		}
+	}
+
+	/// <summary>
+	/// Helper command for hosting a local compute worker in a separate process
+	/// </summary>
+	[Command("computeworker", "Runs the agent as a local compute host, accepting incoming connections on the loopback adapter with a given port")]
+	class ComputeWorkerCommand : Command
+	{
+		readonly IServiceProvider _serviceProvider;
+
+		[CommandLine("-Port=")]
+		int Port { get; set; } = 2000;
+
+		public ComputeWorkerCommand(IServiceProvider serviceProvider)
+		{
+			_serviceProvider = serviceProvider;
+		}
+
+		public override async Task<int> ExecuteAsync(ILogger logger)
+		{
+			logger.LogInformation("** WORKER **");
+
+			ILoggerFactory loggerFactory = _serviceProvider.GetRequiredService<ILoggerFactory>();
+
+			using Socket tcpSocket = new Socket(SocketType.Stream, ProtocolType.IP);
+			await tcpSocket.ConnectAsync(IPAddress.Loopback, Port);
+
+			await using (ComputeSocket socket = new ComputeSocket(new TcpTransport(tcpSocket), loggerFactory))
+			{
+				logger.LogInformation("Running worker...");
+				await RunWorkerAsync(socket, _serviceProvider, CancellationToken.None);
+				logger.LogInformation("Worker complete");
+				await socket.CloseAsync(CancellationToken.None);
+			}
+
+			logger.LogInformation("Stopping");
+			return 0;
+		}
+
+		public static async Task RunWorkerAsync(IComputeSocket socket, IServiceProvider serviceProvider, CancellationToken cancellationToken)
 		{
-			ComputeHandler handler = ActivatorUtilities.CreateInstance<ComputeHandler>(_serviceProvider);
+			ComputeHandler handler = ActivatorUtilities.CreateInstance<ComputeHandler>(serviceProvider);
 			await handler.RunAsync(socket, cancellationToken);
 		}
 	}
diff --git a/Engine/Source/Programs/Shared/EpicGames.Horde/Compute/Clients/LocalComputeClient.cs b/Engine/Source/Programs/Shared/EpicGames.Horde/Compute/Clients/LocalComputeClient.cs
new file mode 100644
--- /dev/null
+++ b/Engine/Source/Programs/Shared/EpicGames.Horde/Compute/Clients/LocalComputeClient.cs
@@ -0,0 +1,77 @@
+// Copyright Epic Games, Inc. All Rights Reserved.
+
+using System;
+using System.Collections.Generic;
+using System.Net;
+using System.Net.Sockets;
+using System.Threading;
+using System.Threading.Tasks;
+using EpicGames.Core;
+using EpicGames.Horde.Compute.Transports;
+using Microsoft.Extensions.Logging;
+
+namespace EpicGames.Horde.Compute.Clients
+{
+	/// <summary>
+	/// Implementation of <see cref="IComputeChannel"/> which marshals data over a loopback connection to a method running on a background task in the same process.
+	/// </summary>
+	public sealed class LocalComputeClient : IComputeClient
+	{
+		readonly BackgroundTask _listenerTask;
+		readonly Socket _listener;
+		readonly Socket _socket;
+		readonly ILoggerFactory _loggerFactory;
+
+		/// <summary>
+		/// Constructor
+		/// </summary>
+		/// <param name="serverFunc">Callback to process the server side of the connection</param>
+		/// <param name="port">Port to connect on</param>
+		/// <param name="loggerFactory">Logger for diagnostic output</param>
+		public LocalComputeClient(Func<IComputeSocket, CancellationToken, Task> serverFunc, int port, ILoggerFactory loggerFactory)
+		{
+			_loggerFactory = loggerFactory;
+
+			_listener = new Socket(SocketType.Stream, ProtocolType.IP);
+			_listener.Bind(new IPEndPoint(IPAddress.Loopback, port));
+			_listener.Listen();
+
+			_listenerTask = BackgroundTask.StartNew(ctx => RunListenerAsync(_listener, serverFunc, loggerFactory, ctx));
+
+			_socket = new Socket(SocketType.Stream, ProtocolType.IP);
+			_socket.Connect(IPAddress.Loopback, port);
+		}
+
+		/// <inheritdoc/>
+		public async ValueTask DisposeAsync()
+		{
+			_socket.Dispose();
+			await _listenerTask.DisposeAsync();
+			_listener.Dispose();
+		}
+
+		/// <summary>
+		/// Sets up the loopback listener and calls the server method
+		/// </summary>
+		static async Task RunListenerAsync(Socket listener, Func<IComputeSocket, CancellationToken, Task> serverFunc, ILoggerFactory loggerFactory, CancellationToken cancellationToken)
+		{
+			using Socket tcpSocket = await listener.AcceptAsync(cancellationToken);
+
+			await using (ComputeSocket socket = new ComputeSocket(new TcpTransport(tcpSocket), loggerFactory))
+			{
+				await serverFunc(socket, cancellationToken);
+				await socket.CloseAsync(cancellationToken);
+			}
+		}
+
+		/// <inheritdoc/>
+		public async Task<TResult> ExecuteAsync<TResult>(ClusterId clusterId, Requirements? requirements, Func<IComputeLease, CancellationToken, Task<TResult>> handler, CancellationToken cancellationToken)
+		{
+			await using ComputeSocket socket = new ComputeSocket(new TcpTransport(_socket), _loggerFactory);
+			await using ComputeLease lease = new ComputeLease(new List<string>(), new Dictionary<string, int>(), socket);
+			TResult result = await handler(lease, cancellationToken);
+			await socket.CloseAsync(cancellationToken);
+			return result;
+		}
+	}
+}
diff --git a/Engine/Source/Programs/Shared/EpicGames.Horde/Compute/Clients/LoopbackComputeClient.cs b/Engine/Source/Programs/Shared/EpicGames.Horde/Compute/Clients/LoopbackComputeClient.cs
deleted file mode 100644
diff --git a/Engine/Source/Programs/Shared/EpicGames.Horde/Compute/ComputeChannel.cs b/Engine/Source/Programs/Shared/EpicGames.Horde/Compute/ComputeChannel.cs
--- a/Engine/Source/Programs/Shared/EpicGames.Horde/Compute/ComputeChannel.cs
+++ b/Engine/Source/Programs/Shared/EpicGames.Horde/Compute/ComputeChannel.cs
@@ -183,7 +183,7 @@ namespace EpicGames.Horde.Compute
 			{
 				bytes.Append("..");
 			}
-			_logger.LogTrace("{Verb} {Type,-20} [{Length,8:n0}] = {Bytes}", verb, type, data.Length, bytes.ToString());
+			_logger.LogTrace("{Verb} {Type,-22} [{Length,8:n0}] = {Bytes}", verb, type, data.Length, bytes.ToString());
 		}
 
 		/// <inheritdoc/>
