Horde: Use docker to build dashboard for installer build
#rnx
#skipci

diff --git a/Engine/Source/Programs/Horde/BuildHorde.xml b/Engine/Source/Programs/Horde/BuildHorde.xml
--- a/Engine/Source/Programs/Horde/BuildHorde.xml
+++ b/Engine/Source/Programs/Horde/BuildHorde.xml
@@ -38,11 +38,42 @@
 	<Property Name="InstallerConfig" Value="Debug"/>
 	<Property Name="InstallerOutputDir" Value="$(RootDir)/Engine/Source/Programs/Horde/Installer/bin/$(InstallerConfig)"/>
 	
+	<Agent Name="HordeInstallerDashboard" Type="Linux">
+		<Property Name="StagingDir" Value="$(RootDir)/Staging"/>
+		<Property Name="DashboardVersion" Value="Unversioned"/>
+		<Property Name="DashboardVersion" Value="CL-$(Change)" If="'$(Change)' != ''"/>
+		<Property Name="DashboardVersion" Value="$(DashboardVersion)-PF-$(PreflightChange)" If="'$(PreflightChange)' != ''" />
+
+		<Node Name="Stage Horde Installer Dashboard" Produces="#DashboardOutputFiles">
+			<Delete Files="$(StagingDir)/Dashboard/..."/>		
+			<WriteTextFile File="$(StagingDir)/Dashboard/index.html" Text="placeholder"/>			
+		
+			<!-- Build the dashboard -->
+			<Tag Files="Engine/Source/Programs/Horde/HordeDashboard/..." With="#DashboardInputFiles"/>
+			<Docker-Build BaseDir="Engine/Source" Files="#DashboardInputFiles" Tag="hordedashboard:installer" DockerFile="Engine/Source/Programs/Horde/HordeDashboard/installer/Dockerfile" Arguments="--build-arg &quot;VersionInfo=$(DashboardVersion)&quot;"/>
+			
+			<!-- Create a container from the image that can copy files from, note that this does not start the container -->
+			<Docker Arguments="create --name horde-installer-container hordedashboard:installer"/>
+			
+			<!-- Copy the dashboard build to the server -->
+			<Docker Arguments="cp horde-installer-container:app/Dashboard/build/. &quot;$(StagingDir)/Dashboard&quot;"/>
+
+			<!-- Remove the container now that files are copied -->
+			<Docker Arguments="rm horde-installer-container"/>
+	
+			<!-- Copy the Horde documentation to the dashboard public folder -->
+			<Copy From="$(RootDir)/Engine/Restricted/NotForLicensees/Source/Programs/Horde/Docs/..." To="$(StagingDir)/Dashboard/public/documentation/..."/>		
+			
+			<Tag Files="$(StagingDir)/Dashboard/..." With="#DashboardOutputFiles"/>
+		</Node>
+	</Agent>
+	
 	<Agent Name="HordeInstaller" Type="Win64;CompileWin64">
 		<Property Name="StagingDir" Value="$(RootDir)/Staging"/>
 
-		<Node Name="Stage Horde Installer">
-			<Delete Files="$(StagingDir)/..."/>
+		<Node Name="Stage Horde Installer" Requires="#DashboardOutputFiles">
+			<Delete Files="$(StagingDir)/Server/..."/>
+			<Delete Files="$(StagingDir)/Agent/..."/>
 			<DotNet Arguments="publish &quot;$(RootDir)/Engine/Source/Programs/Horde/Horde.Server/Horde.Server.csproj&quot; --output &quot;$(StagingDir)/Server&quot; --runtime win-x64 $(VersionArguments)"/>
 			<DotNet Arguments="publish &quot;$(RootDir)/Engine/Source/Programs/Horde/Horde.Agent/Horde.Agent.csproj&quot; --output &quot;$(StagingDir)/Agent&quot; $(VersionArguments)"/>
 
@@ -51,11 +82,9 @@
 
 			<!-- Create build settings -->
 			<WriteTextFile File="$(StagingDir)/Server/appsettings.Build.json" Text="{ &quot;Horde&quot;: { &quot;SingleInstance&quot;: true } }"/>
-	
-			<Spawn Exe="cmd.exe" Arguments="/C yarn install" WorkingDir="$(RootDir)/Engine/Source/Programs/Horde/HordeDashboard"/>
-			<Spawn Exe="cmd.exe" Arguments="/C yarn run build" WorkingDir="$(RootDir)/Engine/Source/Programs/Horde/HordeDashboard"/>
-	
-			<Copy From="$(RootDir)/Engine/Source/Programs/Horde/HordeDashboard/build/..." To="$(StagingDir)/Server/DashboardApp/..."/>
+			
+			<!-- Copy the Dashboard -->
+			<Copy Files="#DashboardOutputFiles" From="$(StagingDir)/Dashboard/..." To="$(StagingDir)/Server/DashboardApp/..." />
 		</Node>
 
 		<Property Name="WixDir" Value="$(RootDir)/Engine/Source/ThirdParty/WiX/3.8"/>
@@ -66,7 +95,7 @@
 			<Property Name="CommonArgs" Value="-dConfiguration=$(InstallerConfig) -dPlatform=x64 -arch x64"/>
 
 			<Delete Files="$(StagingDir)/Server-Bulk/..."/>
-			<Tag Files="$(StagingDir)/Server/..." Except="HordeServer.exe" With="#ServerFiles"/>
+			<Tag Files="$(StagingDir)/Server/..." Except="Horde.Server.exe" With="#ServerFiles"/>
 			<Copy Files="#ServerFiles" From="$(StagingDir)/Server" To="$(StagingDir)/Server-Bulk/Server"/>
 			<Spawn Exe="$(WixDir)/heat.exe" Arguments="dir &quot;$(StagingDir)/Server-Bulk&quot; -cg HordeServerFiles -scom -sreg -gg -dr InstallDir -srd -var var.SourceDir -out &quot;$(InstallerDir)/HordeServerFiles.wxs&quot;"/>
 			<Spawn Exe="$(WixDir)/candle.exe" Arguments="$(CommonArgs) HordeServerFiles.wxs -dSourceDir=&quot;$(StagingDir)&quot; -out $(ObjDir)/HordeServerFiles.wixobj" WorkingDir="$(InstallerDir)"/>
diff --git a/Engine/Source/Programs/Horde/HordeDashboard/installer/Dockerfile b/Engine/Source/Programs/Horde/HordeDashboard/installer/Dockerfile
new file mode 100644
--- /dev/null
+++ b/Engine/Source/Programs/Horde/HordeDashboard/installer/Dockerfile
@@ -0,0 +1,9 @@
+FROM node:16.13.0 as build-env
+ARG VersionInfo
+ENV REACT_APP_VERSION_INFO=${VersionInfo}
+WORKDIR /app
+
+COPY Programs/Horde/HordeDashboard ./Dashboard
+RUN cd ./Dashboard && yarn install
+RUN cd ./Dashboard && yarn run build
+
