Horde: Adding intregrated markdown doc viewer to dashboard
#rnx
#skipci

diff --git a/Engine/Source/Programs/Horde/HordeDashboard/src/App.tsx b/Engine/Source/Programs/Horde/HordeDashboard/src/App.tsx
--- a/Engine/Source/Programs/Horde/HordeDashboard/src/App.tsx
+++ b/Engine/Source/Programs/Horde/HordeDashboard/src/App.tsx
@@ -15,6 +15,7 @@ import { AutomationView } from './components/AutomationView';
 import { DashboardView } from './components/DashboardView';
 import { DebugView } from './components/DebugView';
 import { DeviceView } from './components/DeviceView';
+import { DocView } from './components/docs/DocView';
 import { ErrorDialog, ErrorHandler } from './components/ErrorHandler';
 import { JobDetailViewV2 } from './components/jobDetailsV2/JobDetailViewV2';
 import { LogView } from './components/LogView';
@@ -146,7 +147,9 @@ const Main: React.FC = () => {
                { path: "audit/issue/:issueId", element: <AuditLogView /> },
                { path: "automation", element: <AutomationView /> },
                { path: "tools", element: <ToolView /> },
-               { path: "debug/lease/:leaseId", element: <DebugView /> }
+               { path: "debug/lease/:leaseId", element: <DebugView /> },
+               { path: "docs", element: <DocView /> },
+               { path: "docs/*", element: <DocView /> }
             ]
          }
       ];
@@ -203,7 +206,7 @@ const HomeRedirect: React.FC = () => {
 const Root: React.FC = () => {
    return <div>
       <Outlet />
-      <HomeRedirect />      
+      <HomeRedirect />
    </div>
 }
 
diff --git a/Engine/Source/Programs/Horde/HordeDashboard/src/base/components/Markdown/MarkdownLink.tsx b/Engine/Source/Programs/Horde/HordeDashboard/src/base/components/Markdown/MarkdownLink.tsx
--- a/Engine/Source/Programs/Horde/HordeDashboard/src/base/components/Markdown/MarkdownLink.tsx
+++ b/Engine/Source/Programs/Horde/HordeDashboard/src/base/components/Markdown/MarkdownLink.tsx
@@ -5,6 +5,11 @@ import { removeAnchorLink } from '../../utilities/index';
 
 export const MarkdownLink: React.FunctionComponent<ILinkProps> = props => {
    let href = props.href;
+
+   if (href?.indexOf("README.md") !== -1) {
+      href = "/docs/"
+   }
+
    if (href && href[0] === '#' && href.indexOf('/') === -1) {
       // This is an anchor link within this page. We need to prepend the current route.
       href = removeAnchorLink(window.location.hash) + href;
diff --git a/Engine/Source/Programs/Horde/HordeDashboard/src/base/components/SideRail/SideRail.styles.ts b/Engine/Source/Programs/Horde/HordeDashboard/src/base/components/SideRail/SideRail.styles.ts
--- a/Engine/Source/Programs/Horde/HordeDashboard/src/base/components/SideRail/SideRail.styles.ts
+++ b/Engine/Source/Programs/Horde/HordeDashboard/src/base/components/SideRail/SideRail.styles.ts
@@ -1,4 +1,5 @@
-import { FontWeights, getFocusOutlineStyle, IStyleFunction } from '@fluentui/react/lib/index';
+//import { FontWeights, getFocusOutlineStyle, IStyleFunction } from '@fluentui/react/lib/index';
+import { FontWeights, getFocusOutlineStyle, IStyleFunction } from '@fluentui/react';
 import { modeColors } from '../../../styles/Styles';
 import { ISideRailStyleProps, ISideRailStyles } from './SideRail.types';
 
@@ -43,7 +44,7 @@ export const getStyles: IStyleFunction<ISideRailStyleProps, ISideRailStyles> = p
             flex: '1',
             padding: '4px 0px',
             selectors: {              
-              ':active,:visited,:hover,:focus': { background: theme.palette.neutralLight }
+              ':active,:visited,:hover': { color: modeColors.text, textDecoration: "none", background: theme.palette.neutralLight }
             }
           },
           getFocusOutlineStyle(theme, 1)
diff --git a/Engine/Source/Programs/Horde/HordeDashboard/src/base/components/SideRail/SideRail.tsx b/Engine/Source/Programs/Horde/HordeDashboard/src/base/components/SideRail/SideRail.tsx
--- a/Engine/Source/Programs/Horde/HordeDashboard/src/base/components/SideRail/SideRail.tsx
+++ b/Engine/Source/Programs/Horde/HordeDashboard/src/base/components/SideRail/SideRail.tsx
@@ -44,13 +44,13 @@ class SideRailBase extends React.Component<ISideRailProps, ISideRailState> {
       this._classNames = getClassNames(this.props.styles, { theme: this.props.theme });
 
       const jumpLinkList = this._renderJumpLinkList();
-      const relatedLinkList = this._renderLinkList(this.props.relatedLinks, 'Job');
+      const relatedLinkList = this._renderLinkList(this.props.relatedLinks, 'Reference');
       const contactLinkList = this._renderLinkList(this.props.contactLinks, 'Contacts');
 
       return jumpLinkList || relatedLinkList || contactLinkList ? (
          <FocusZone direction={FocusZoneDirection.vertical} className={this._classNames.root}>
-            {jumpLinkList}            
             {relatedLinkList}
+            {jumpLinkList}                        
             {contactLinkList}
          </FocusZone>
       ) : null;
diff --git a/Engine/Source/Programs/Horde/HordeDashboard/src/components/TopNav.tsx b/Engine/Source/Programs/Horde/HordeDashboard/src/components/TopNav.tsx
--- a/Engine/Source/Programs/Horde/HordeDashboard/src/components/TopNav.tsx
+++ b/Engine/Source/Programs/Horde/HordeDashboard/src/components/TopNav.tsx
@@ -456,6 +456,26 @@ export const TopNav: React.FC<{ suppressServer?: boolean }> = observer(({ suppre
                bottomDivider: true
             }
          });
+
+         // currently disabled until doc copy to public folder is sorted
+         /*
+         const docsItems: IContextualMenuItem[] = [];
+         docsItems.push({
+            key: "server_docs",
+            text: "Documentation",
+            link: `/docs`
+         });
+
+         subItems.push({
+            itemType: ContextualMenuItemType.Section,
+            key: `server_docs_item`,
+            sectionProps: {
+               title: "Documentation",
+               items: docsItems,
+               bottomDivider: true
+            }
+         });
+         */
       }
 
       // Configuration
diff --git a/Engine/Source/Programs/Horde/HordeDashboard/src/components/docs/DocView.tsx b/Engine/Source/Programs/Horde/HordeDashboard/src/components/docs/DocView.tsx
new file mode 100644
--- /dev/null
+++ b/Engine/Source/Programs/Horde/HordeDashboard/src/components/docs/DocView.tsx
@@ -0,0 +1,240 @@
+import { Stack } from "@fluentui/react";
+import { action, makeObservable, observable } from "mobx";
+import { observer } from "mobx-react-lite";
+import { useRef, useState } from "react";
+import { useLocation } from "react-router-dom";
+import { Markdown } from "../../base/components/Markdown";
+import { ISideRailLink, SideRail } from "../../base/components/SideRail";
+import { useWindowSize } from "../../base/utilities/hooks";
+import { hordeClasses, modeColors } from "../../styles/Styles";
+import { BreadcrumbItem, Breadcrumbs } from "../Breadcrumbs";
+import { TopNav } from "../TopNav";
+
+type Anchor = {
+   text: string;
+   anchor: string;
+}
+
+const documentCache = new Map<string, string>();
+const crumbCache = new Map<string, BreadcrumbItem[]>();
+const anchorCache = new Map<string, Anchor[]>();
+
+type State = {
+   crumbs: BreadcrumbItem[];
+   jumpLinks: ISideRailLink[];
+}
+
+class LinkState {
+
+   constructor() {
+      makeObservable(this);
+   }
+
+   @observable
+   state: State = { crumbs: [], jumpLinks: [] };
+
+   @action
+   setState(crumbs: BreadcrumbItem[], jumpLinks: ISideRailLink[]) {
+      this.state = { crumbs: crumbs, jumpLinks: jumpLinks };
+   }
+}
+
+const linkState = new LinkState();
+
+const DocPanel: React.FC<{ docName: string }> = ({ docName }) => {
+
+   const [state, setState] = useState(0);
+
+   const docRef = useRef<HTMLDivElement>(null);
+
+   const text = documentCache.get(docName);
+
+   if (!text) {
+      fetch(`/${docName}`, { cache: "no-cache" })
+         .then((response) => response.text())
+         .then((textContent) => {
+
+            let crumbs: BreadcrumbItem[] = [];
+            let anchors: Anchor[] = [];
+
+            textContent = textContent.trim();
+
+            if (textContent?.indexOf("<!doctype html>") === -1) {
+
+               // generate anchors
+               let lines = (textContent.match(/[^\r\n]+/g) ?? []) as string[];
+               lines = lines.map(i => i.trim()).filter( i => !!i);
+               lines.forEach(line => {
+                  let anchor = "";
+
+                  if (line.startsWith("# ")) {
+                     anchor = line.split("# ")[1];
+                  }
+
+                  if (line.startsWith("## ")) {
+                     anchor = line.split("## ")[1];
+                  }
+
+                  if (anchor) {
+                     anchor = anchor.trim();
+                     anchors.push({text: anchor, anchor: anchor.replace(/[^a-z0-9- ]/gi, '').replace(/ /gi, '-').toLowerCase() })
+                  }
+               })               
+
+               // generate crumbs
+               if (lines && lines.length) {
+                  const line = lines[0];
+                  if (line.startsWith("[Horde]")) {
+
+                     line.split(">").map(c => c.trim()).forEach(c => {
+                        const link = c.split("](");
+                        if (link.length === 1) {
+                           crumbs.push({ text: link[0] })
+                        } else {
+                           let [cname, clink] = [link[0].replace("[", ""), link[1].replace(")", "")];
+
+                           let docPath = docName.split("/").slice(2);
+                           docPath.pop();
+
+                           const relative = (clink.match(/..\//g) || []).length;
+
+                           if (clink.indexOf("README.md") !== -1) {
+                              clink = "/docs"
+                           } else {
+                              docPath = docPath.slice(docPath.length - 1, -relative);
+                              const elements = clink.split("/");
+                              docPath.push(elements[elements.length - 1]);
+                              clink = docPath.join("/");
+                           }
+
+                           if (cname !== "Horde") {
+                              crumbs.push({ text: cname, link: clink });
+                           }
+
+                        }
+                     });
+
+                     textContent = textContent.replace(line, "");
+                  }
+               }
+
+               documentCache.set(docName, textContent);
+               crumbCache.set(docName, crumbs);
+               anchorCache.set(docName, anchors);
+
+            } else {
+               documentCache.set(docName, `###Missing document ${docName}`);
+               crumbCache.set(docName, []);
+               anchorCache.set(docName, []);
+            }
+
+            setState(state + 1);
+
+         });
+   }
+
+   if (!text) {
+      return null;
+   }
+
+   let crumbs = crumbCache.get(docName) ?? [];
+   let anchors = anchorCache.get(docName) ?? [];
+
+   linkState.setState(crumbs, anchors.map(a => {
+      return { text: a.text, url: a.anchor}
+   }));
+
+
+   return <Stack>
+      <Stack styles={{ root: { paddingTop: 0, paddingLeft: 12, paddingRight: 12, width: "100%" } }} >
+         <Stack tokens={{ childrenGap: 12 }}>
+            <Stack style={{ paddingBottom: 24, paddingLeft: 6, paddingRight: 6, paddingTop: 12 }}>
+               <div ref={docRef}>
+                  <Markdown>{text}</Markdown>
+               </div>
+            </Stack>
+         </Stack>
+      </Stack>
+   </Stack>;
+
+
+}
+
+const DocRail = observer(() => {
+
+   const state = linkState.state;
+
+   const refLinks: ISideRailLink[] = [];
+   refLinks.push({ text: "Introduction", url: "/docs" });
+   refLinks.push({ text: "User Guild", url: "/docs/Users.md" });
+   refLinks.push({ text: "Deployment", url: "/docs/Deployment.md" });
+   refLinks.push({ text: "Configuration", url: "/docs/Config.md" });
+   refLinks.push({ text: "Horde Internals", url: "/docs/Internals.md" });
+
+   return <SideRail jumpLinks={state.jumpLinks} relatedLinks={refLinks} />
+
+})
+
+const DocCrumbs = observer(() => {
+
+   const state = linkState.state;
+
+   const crumbs: BreadcrumbItem[] = [];
+   crumbs.push({ text: "Documentation", link: "/docs" })
+   crumbs.push(...state.crumbs.map(c => { return { text: c.text, link: c.link } }));
+
+   return <Breadcrumbs items={crumbs} />
+
+})
+
+
+
+export const DocView = () => {
+
+   const location = useLocation();
+
+   // fixme
+   let docName = location.pathname.replace("/docs/", "").replace("/docs", "").trim();
+   if (!docName || docName.indexOf("README.md") !== -1) {
+      docName = "documentation/README.md";
+   } else {
+
+      if (!docName.startsWith("Docs/")) {
+         docName = `documentation/Docs/${docName}`;
+      } else {
+         docName = docName.replace("Docs/", "documentation/Docs/");
+      }
+   }
+
+   const windowSize = useWindowSize();
+   const vw = Math.max(document.documentElement.clientWidth, window.innerWidth || 0);
+
+   return <Stack className={hordeClasses.horde}>
+      <TopNav />
+      <DocCrumbs />
+      <Stack horizontal>
+         <div key={`windowsize_streamview_${windowSize.width}_${windowSize.height}`} style={{ width: vw / 2 - (1440 / 2), flexShrink: 0, backgroundColor: modeColors.background }} />
+         <Stack tokens={{ childrenGap: 0 }} styles={{ root: { backgroundColor: modeColors.background, width: "100%", "position": "relative", paddingTop: 12 } }}>
+            <div style={{ overflowY: 'scroll', overflowX: 'hidden', height: "calc(100vh - 162px)" }} data-is-scrollable={true}>
+               <Stack horizontal>
+                  <Stack style={{ width: 1240, paddingTop: 6, marginLeft: 4, height: '100%' }}>
+                     <Stack className={hordeClasses.raised}>
+                        <Stack style={{ width: "100%", height: "max-content" }} tokens={{ childrenGap: 18 }}>
+                           <DocPanel docName={docName} />
+                        </Stack>
+                     </Stack>
+                     <Stack style={{ paddingBottom: 24 }} />
+                  </Stack>
+                  <Stack style={{ paddingLeft: 1280, paddingTop: 12, position: "absolute", pointerEvents: "none" }}>
+                     <div style={{ pointerEvents: "all" }}>
+                        <DocRail />
+                     </div>
+                  </Stack>
+               </Stack>
+            </div>
+         </Stack>
+      </Stack>
+   </Stack>
+
+}
+
