Horde: Enable docs viewer in dev
#rnx
#skipci

diff --git a/Engine/Source/Programs/Horde/HordeDashboard/src/components/TopNav.tsx b/Engine/Source/Programs/Horde/HordeDashboard/src/components/TopNav.tsx
--- a/Engine/Source/Programs/Horde/HordeDashboard/src/components/TopNav.tsx
+++ b/Engine/Source/Programs/Horde/HordeDashboard/src/components/TopNav.tsx
@@ -457,8 +457,6 @@ export const TopNav: React.FC<{ suppressServer?: boolean }> = observer(({ suppre
             }
          });
 
-         // currently disabled until doc copy to public folder is sorted
-         /*
          const docsItems: IContextualMenuItem[] = [];
          docsItems.push({
             key: "server_docs",
@@ -475,7 +473,6 @@ export const TopNav: React.FC<{ suppressServer?: boolean }> = observer(({ suppre
                bottomDivider: true
             }
          });
-         */
       }
 
       // Configuration
