Horde: Fix deadlock waiting for child process to exit in compute handler.

diff --git a/Engine/Source/Programs/Horde/Horde.Agent/Leases/Handlers/ComputeHandler.cs b/Engine/Source/Programs/Horde/Horde.Agent/Leases/Handlers/ComputeHandler.cs
--- a/Engine/Source/Programs/Horde/Horde.Agent/Leases/Handlers/ComputeHandler.cs
+++ b/Engine/Source/Programs/Horde/Horde.Agent/Leases/Handlers/ComputeHandler.cs
@@ -217,7 +217,7 @@ namespace Horde.Agent.Leases.Handlers
 			await using ComputeChannel controlChannel = new ComputeChannel(receiveBuffer.Reader, sendBuffer.Writer, _loggerFactory.CreateLogger<ComputeChannel>());
 
 			Task processTask = ExecuteProcessInternalAsync(channel, executable, arguments, workingDir, envVars, cancellationToken);
-			processTask = processTask.ContinueWith(x => sendBuffer.Writer.MarkComplete(), cancellationToken, TaskContinuationOptions.None, TaskScheduler.Default);
+			processTask = processTask.ContinueWith(x => receiveBuffer.Writer.MarkComplete(), cancellationToken, TaskContinuationOptions.None, TaskScheduler.Default);
 
 			List<IpcBuffer> buffers = new List<IpcBuffer>();
 			try
