Horde: Minor build xml change to include README.md
#rnx
#skipci

diff --git a/Engine/Source/Programs/Horde/BuildHorde.xml b/Engine/Source/Programs/Horde/BuildHorde.xml
--- a/Engine/Source/Programs/Horde/BuildHorde.xml
+++ b/Engine/Source/Programs/Horde/BuildHorde.xml
@@ -62,7 +62,8 @@
 			<Docker Arguments="rm horde-installer-container"/>
 	
 			<!-- Copy the Horde documentation to the dashboard public folder -->
-			<Copy From="$(RootDir)/Engine/Restricted/NotForLicensees/Source/Programs/Horde/Docs/..." To="$(StagingDir)/Dashboard/public/documentation/..."/>		
+			<Copy From="$(RootDir)/Engine/Restricted/NotForLicensees/Source/Programs/Horde/Docs/..." To="$(StagingDir)/Dashboard/public/documentation/Docs/..."/>		
+			<Copy From="$(RootDir)/Engine/Restricted/NotForLicensees/Source/Programs/Horde/README.md" To="$(StagingDir)/Dashboard/public/documentation/README.md"/>
 			
 			<Tag Files="$(StagingDir)/Dashboard/..." With="#DashboardOutputFiles"/>
 		</Node>
