Horde: Support filtering to parameters in simple job search
#rnx
#skipci

diff --git a/Engine/Source/Programs/Horde/HordeDashboard/src/components/JobSearchSimple.tsx b/Engine/Source/Programs/Horde/HordeDashboard/src/components/JobSearchSimple.tsx
--- a/Engine/Source/Programs/Horde/HordeDashboard/src/components/JobSearchSimple.tsx
+++ b/Engine/Source/Programs/Horde/HordeDashboard/src/components/JobSearchSimple.tsx
@@ -24,7 +24,7 @@ export const JobSearchSimpleModal: React.FC<{ streamId: string, onClose: () => v
       job: GetJobResponse
    };
 
-   const [searchState, setSearchState] = useState<{ items: JobItem[], groups?: IGroup[], templateId?: string, name?: string, userId?: string, global?: boolean, streamIdOverride?: string, preflights: boolean, preflightOnly?: boolean, minDate?: Date, maxDate?: Date, minCL?: string, maxCL?: string, querying: boolean, containsStep?: string }>({ items: [], querying: false, preflights: false });
+   const [searchState, setSearchState] = useState<{ items: JobItem[], groups?: IGroup[], templateId?: string, name?: string, userId?: string, global?: boolean, streamIdOverride?: string, preflights: boolean, preflightOnly?: boolean, minDate?: Date, maxDate?: Date, minCL?: string, maxCL?: string, querying: boolean, containsStep?: string, containsParameter?:string }>({ items: [], querying: false, preflights: false });
    const [templateData, setTemplateData] = useState<{ streamId: string; templates: GetTemplateRefResponse[] } | undefined>(undefined);
 
    const stream = projectStore.streamById(streamId);
@@ -310,7 +310,7 @@ export const JobSearchSimpleModal: React.FC<{ streamId: string, onClose: () => v
                id: jobIds.length ? jobIds : undefined,
                streamId: searchStreamId?.toLowerCase(),
                name: name,
-               filter: "id,name,change,createTime,streamId,preflightChange,startedByUserInfo,streamId",
+               filter: "id,name,change,createTime,streamId,preflightChange,startedByUserInfo,streamId,arguments",
                minChange: minCL,
                maxChange: maxCL,
                template: (!searchState.global && template?.id && !searchState.containsStep) ? [template.id] : undefined,
@@ -323,6 +323,15 @@ export const JobSearchSimpleModal: React.FC<{ streamId: string, onClose: () => v
             }
    
             jobs = await backend.getJobs(query, false);
+
+            if (searchState.containsParameter) {
+               const p = searchState.containsParameter.toLowerCase();
+               jobs = jobs.filter(j => {
+                  return !!j.arguments.find(a => {
+                     return a.toLowerCase().indexOf(p) !== -1
+                  })
+               });
+            }
          }         
 
       } catch (reason) {
@@ -470,6 +479,12 @@ export const JobSearchSimpleModal: React.FC<{ streamId: string, onClose: () => v
                         }} ></TextField>
                      </Stack>
 
+                     <Stack>
+                        <TextField value={searchState.containsParameter ?? ""} spellCheck={false} style={{ width: 350 }} label="Contains Parameter" onChange={(event, newValue) => {
+                           setSearchState({ ...searchState, containsParameter: newValue?.trim() ? newValue : undefined })
+                        }} ></TextField>
+                     </Stack>
+
                      <Stack>
                         <TextField value={searchState.containsStep ?? ""} spellCheck={false} style={{ width: 350 }} label="Contains Step" description="Note: Exact match is required" onChange={(event, newValue) => {
                            setSearchState({ ...searchState, containsStep: newValue?.trim() ? newValue : undefined })
