Horde: Constrain parameter job search to template
#rnx
#skipci

diff --git a/Engine/Source/Programs/Horde/HordeDashboard/src/components/JobSearchSimple.tsx b/Engine/Source/Programs/Horde/HordeDashboard/src/components/JobSearchSimple.tsx
--- a/Engine/Source/Programs/Horde/HordeDashboard/src/components/JobSearchSimple.tsx
+++ b/Engine/Source/Programs/Horde/HordeDashboard/src/components/JobSearchSimple.tsx
@@ -100,7 +100,7 @@ export const JobSearchSimpleModal: React.FC<{ streamId: string, onClose: () => v
 
    })
 
-   templateOptions.push({ key: `show_all_templates`, text: "All Templates", onClick: (ev) => { ev?.stopPropagation(); ev?.preventDefault(); setSearchState({ ...searchState, templateId: undefined }) } });
+   templateOptions.push({ key: `show_all_templates`, text: "All Templates", onClick: (ev) => { ev?.stopPropagation(); ev?.preventDefault(); setSearchState({ ...searchState, templateId: undefined, containsParameter: undefined }) } });
 
    const templateMenuProps: IContextualMenuProps = {
       shouldFocusOnMount: true,
@@ -254,7 +254,7 @@ export const JobSearchSimpleModal: React.FC<{ streamId: string, onClose: () => v
 
          }
 
-         const maxCount = 512;
+         const maxCount = 2048;
 
          let searchStreamId: string | undefined = streamId;
 
@@ -440,7 +440,7 @@ export const JobSearchSimpleModal: React.FC<{ streamId: string, onClose: () => v
                      </Stack>
 
                      <Stack>
-                        <TextField value={searchState.name ?? ""} spellCheck={false} style={{ width: 350 }} label="Name" onChange={(event, newValue) => {
+                        <TextField value={searchState.name ?? ""} spellCheck={false} autoComplete="off" style={{ width: 350 }} label="Name" onChange={(event, newValue) => {
                            setSearchState({ ...searchState, name: newValue })
 
                         }} ></TextField>
@@ -471,7 +471,7 @@ export const JobSearchSimpleModal: React.FC<{ streamId: string, onClose: () => v
                      </Stack>
 
                      <Stack>
-                        <TextField value={searchState.streamIdOverride ?? ""} spellCheck={false} style={{ width: 350 }} label="Stream Id" onChange={(event, newValue) => {
+                        <TextField value={searchState.streamIdOverride ?? ""} spellCheck={false} autoComplete="off" style={{ width: 350 }} label="Stream Id" onChange={(event, newValue) => {
                            const templateId = searchState.templateId;
                            const global = searchState.global;
                            setSearchState({ ...searchState, streamIdOverride: newValue, global: newValue ? false : global, templateId: newValue ? undefined : templateId })
@@ -480,13 +480,13 @@ export const JobSearchSimpleModal: React.FC<{ streamId: string, onClose: () => v
                      </Stack>
 
                      <Stack>
-                        <TextField value={searchState.containsParameter ?? ""} spellCheck={false} style={{ width: 350 }} label="Contains Parameter" onChange={(event, newValue) => {
+                        <TextField disabled={!template} value={searchState.containsParameter ?? ""} spellCheck={false} autoComplete="off" style={{ width: 350 }} label="Contains Parameter" placeholder={ !template ? "Please select a template to search parameters" : ""}  onChange={(event, newValue) => {
                            setSearchState({ ...searchState, containsParameter: newValue?.trim() ? newValue : undefined })
                         }} ></TextField>
                      </Stack>
 
                      <Stack>
-                        <TextField value={searchState.containsStep ?? ""} spellCheck={false} style={{ width: 350 }} label="Contains Step" description="Note: Exact match is required" onChange={(event, newValue) => {
+                        <TextField value={searchState.containsStep ?? ""} spellCheck={false} autoComplete="off" style={{ width: 350 }} label="Contains Step" description="Note: Exact match is required" onChange={(event, newValue) => {
                            setSearchState({ ...searchState, containsStep: newValue?.trim() ? newValue : undefined })
                         }} ></TextField>
                      </Stack>
