Horde: Another simple job search tweak, third time should be the charm
#rnx
#skipci

diff --git a/Engine/Source/Programs/Horde/HordeDashboard/src/components/JobSearchSimple.tsx b/Engine/Source/Programs/Horde/HordeDashboard/src/components/JobSearchSimple.tsx
--- a/Engine/Source/Programs/Horde/HordeDashboard/src/components/JobSearchSimple.tsx
+++ b/Engine/Source/Programs/Horde/HordeDashboard/src/components/JobSearchSimple.tsx
@@ -24,7 +24,7 @@ export const JobSearchSimpleModal: React.FC<{ streamId: string, onClose: () => v
       job: GetJobResponse
    };
 
-   const [searchState, setSearchState] = useState<{ items: JobItem[], groups?: IGroup[], templateId?: string, name?: string, userId?: string, global?: boolean, streamIdOverride?: string, preflights: boolean, preflightOnly?: boolean, minDate?: Date, maxDate?: Date, minCL?: string, maxCL?: string, querying: boolean, containsStep?: string, containsParameter?:string }>({ items: [], querying: false, preflights: false });
+   const [searchState, setSearchState] = useState<{ items: JobItem[], groups?: IGroup[], templateId?: string, name?: string, userId?: string, global?: boolean, streamIdOverride?: string, preflights: boolean, preflightOnly?: boolean, minDate?: Date, maxDate?: Date, minCL?: string, maxCL?: string, querying: boolean, containsStep?: string, containsParameter?: string }>({ items: [], querying: false, preflights: false });
    const [templateData, setTemplateData] = useState<{ streamId: string; templates: GetTemplateRefResponse[] } | undefined>(undefined);
 
    const stream = projectStore.streamById(streamId);
@@ -117,7 +117,7 @@ export const JobSearchSimpleModal: React.FC<{ streamId: string, onClose: () => v
       { key: 'column2', name: 'Name', minWidth: 340, maxWidth: 340, isResizable: false },
       { key: 'column3', name: 'Step', minWidth: 200, maxWidth: 200, isResizable: false },
       { key: 'column4', name: 'StartedBy', minWidth: 120, maxWidth: 120, isResizable: false },
-      { key: 'column5', name: 'Created', minWidth: 100, maxWidth: 100, isResizable: false },      
+      { key: 'column5', name: 'Created', minWidth: 100, maxWidth: 100, isResizable: false },
    ];
 
    const onRenderGroupHeader: IDetailsGroupRenderProps['onRenderHeader'] = (props) => {
@@ -156,7 +156,7 @@ export const JobSearchSimpleModal: React.FC<{ streamId: string, onClose: () => v
       }
 
       if (column.name === "Name") {
-         return <Stack horizontal verticalFill={true} verticalAlign="center" tokens={{ childrenGap: 0, padding: 0 }} style={{whiteSpace:"break-spaces"}} ><Text variant="tiny">{item.job.name}</Text></Stack>;
+         return <Stack horizontal verticalFill={true} verticalAlign="center" tokens={{ childrenGap: 0, padding: 0 }} style={{ whiteSpace: "break-spaces" }} ><Text variant="tiny">{item.job.name}</Text></Stack>;
       }
 
       if (column.name === "Change") {
@@ -174,10 +174,10 @@ export const JobSearchSimpleModal: React.FC<{ streamId: string, onClose: () => v
 
       if (column.name === "Step") {
          const ref = jobToStep.get(item.job.id);
-         if (!ref || !jobStepName) { 
+         if (!ref || !jobStepName) {
             return null;
          }
-         return <Link to={`/job/${ref.jobId}?step=${ref.stepId}`}><Stack horizontal verticalAlign="center" tokens={{ childrenGap: 0, padding: 0 }} style={{ width: "100%", height: "100%" }} >{<StepRefStatusIcon stepRef={ref} />}<Stack style={{whiteSpace:"break-spaces"}}><Text variant="tiny">{jobStepName}</Text></Stack></Stack></Link>;
+         return <Link to={`/job/${ref.jobId}?step=${ref.stepId}`}><Stack horizontal verticalAlign="center" tokens={{ childrenGap: 0, padding: 0 }} style={{ width: "100%", height: "100%" }} >{<StepRefStatusIcon stepRef={ref} />}<Stack style={{ whiteSpace: "break-spaces" }}><Text variant="tiny">{jobStepName}</Text></Stack></Stack></Link>;
       }
 
 
@@ -301,7 +301,7 @@ export const JobSearchSimpleModal: React.FC<{ streamId: string, onClose: () => v
          }
 
          const jobIds = stepJobIds.slice(0, 256);
-         
+
          if (stepName && !jobIds.length) {
             jobs = [];
          } else {
@@ -321,18 +321,18 @@ export const JobSearchSimpleModal: React.FC<{ streamId: string, onClose: () => v
                startedByUserId: searchState.userId,
                count: maxCount
             }
-   
+
             jobs = await backend.getJobs(query, false);
 
-            if (searchState.containsParameter) {
-               const p = searchState.containsParameter.toLowerCase();
+            const p = searchState.containsParameter?.trim().toLowerCase();
+            if (p) {
                jobs = jobs.filter(j => {
                   return !!j.arguments.find(a => {
                      return a.toLowerCase().indexOf(p) !== -1
                   })
                });
             }
-         }         
+         }
 
       } catch (reason) {
 
@@ -413,7 +413,7 @@ export const JobSearchSimpleModal: React.FC<{ streamId: string, onClose: () => v
       title = `Find Jobs in ${searchState.streamIdOverride}`;
    }
 
-   return (<Modal isOpen={true} topOffsetFixed={true} styles={{ main: { padding: 8, width: 1400, height: '800px', backgroundColor: '#FFFFFF', hasBeenOpened: false, top: "80px", position: "absolute" } } } className={hordeClasses.modal} onDismiss={() => { onClose() }}>
+   return (<Modal isOpen={true} topOffsetFixed={true} styles={{ main: { padding: 8, width: 1400, height: '800px', backgroundColor: '#FFFFFF', hasBeenOpened: false, top: "80px", position: "absolute" } }} className={hordeClasses.modal} onDismiss={() => { onClose() }}>
       {searchState.querying && <StatusModal text={"Finding Jobs"} />}
       <Stack styles={{ root: { paddingTop: 8, paddingLeft: 24, paddingRight: 12, paddingBottom: 8 } }}>
          <Stack tokens={{ childrenGap: 12 }}>
@@ -480,7 +480,7 @@ export const JobSearchSimpleModal: React.FC<{ streamId: string, onClose: () => v
                      </Stack>
 
                      <Stack>
-                        <TextField disabled={!template} value={searchState.containsParameter ?? ""} spellCheck={false} autoComplete="off" style={{ width: 350 }} label="Contains Parameter" placeholder={ !template ? "Please select a template to search parameters" : ""}  onChange={(event, newValue) => {
+                        <TextField disabled={!template} value={searchState.containsParameter ?? ""} spellCheck={false} autoComplete="off" style={{ width: 350 }} label="Contains Parameter" placeholder={!template ? "Please select a template to search parameters" : ""} onChange={(event, newValue) => {
                            setSearchState({ ...searchState, containsParameter: newValue?.trim() ? newValue : undefined })
                         }} ></TextField>
                      </Stack>
