Horde: Fix summary view overflow
#rnx
#skipci

diff --git a/Engine/Source/Programs/Horde/HordeDashboard/src/components/StreamSummary.tsx b/Engine/Source/Programs/Horde/HordeDashboard/src/components/StreamSummary.tsx
--- a/Engine/Source/Programs/Horde/HordeDashboard/src/components/StreamSummary.tsx
+++ b/Engine/Source/Programs/Horde/HordeDashboard/src/components/StreamSummary.tsx
@@ -80,9 +80,9 @@ class SummaryHandler {
          }
 
          let jiraKeys: string[] = [];
-         
+
          const jiraIssues: Set<number> = new Set();
-         values[0].forEach(issue => {            
+         values[0].forEach(issue => {
             if (issue.externalIssueKey) {
                jiraIssues.add(issue.id)
                jiraKeys.push(issue.externalIssueKey);
@@ -215,13 +215,13 @@ class SummaryHandler {
 
             let jiraTag = "";
 
-            if (issue.externalIssueKey) {                     
-               const url = `${dashboard.externalIssueService?.url}/browse/${issue.externalIssueKey}`;      
+            if (issue.externalIssueKey) {
+               const url = `${dashboard.externalIssueService?.url}/browse/${issue.externalIssueKey}`;
                jiraTag = `<a href="${url}" target="_blank">${issue.externalIssueKey}</a>`;
             }
 
             let workflowTag = "";
-            if (issue.workflowThreadUrl) {                     
+            if (issue.workflowThreadUrl) {
                workflowTag = `<a href="${issue.workflowThreadUrl}" target="_blank">Slack Thread</a>`;
             }
 
@@ -429,12 +429,12 @@ const HealthPanelIssues: React.FC<{ desktopAlerts?: boolean }> = observer(({ des
       issues = handler.unpromoted;
    }
 
-   const stream = projectStore.streamById(handler.streamId);   
+   const stream = projectStore.streamById(handler.streamId);
    if (stream) {
       const config = stream.workflows?.find(w => w.id === currentPivot);
       if (config) {
          issues = allIssues.filter(issue => {
-            return issue.openWorkflows.indexOf(currentPivot) !== -1;            
+            return issue.openWorkflows.indexOf(currentPivot) !== -1;
          });
       }
    }
@@ -442,7 +442,7 @@ const HealthPanelIssues: React.FC<{ desktopAlerts?: boolean }> = observer(({ des
    const [items, groups] = sortIssueGroups(issues.map(i => { return { issue: i } }), projectStore);
 
    // 508
-   const columns:IColumn[] = [
+   const columns: IColumn[] = [
       { key: 'health_column1', name: 'Summary', minWidth: 370, maxWidth: 370, isResizable: false, isPadded: false }, // see summary render for px setting needed by ellipsis
       { key: 'health_column2', name: 'Quarantine', minWidth: 80, maxWidth: 80, isResizable: false, isPadded: false },
       { key: 'health_column3', name: 'Jira', minWidth: 80, maxWidth: 80, isResizable: false, isPadded: false },
@@ -488,7 +488,7 @@ const HealthPanelIssues: React.FC<{ desktopAlerts?: boolean }> = observer(({ des
 
          let summary = issue.summary;
 
-         return <Stack horizontal style={{ paddingTop: 6, height: "100%"}}>{<IssueStatusIconV2 issue={issue} />}<Text style={{width: "370px", display: "block", textAlign: "left", overflow: "hidden", "textOverflow" : "ellipsis"}} variant={textSize}>{`${issue.id} - ${summary}`}</Text></Stack>;
+         return <Stack horizontal style={{ paddingTop: 6, height: "100%" }}>{<IssueStatusIconV2 issue={issue} />}<Text style={{ width: "370px", display: "block", textAlign: "left", overflow: "hidden", "textOverflow": "ellipsis" }} variant={textSize}>{`${issue.id} - ${summary}`}</Text></Stack>;
       }
 
       if (column.name === "Quarantine") {
@@ -498,10 +498,10 @@ const HealthPanelIssues: React.FC<{ desktopAlerts?: boolean }> = observer(({ des
          }
 
          if (issue.workflowThreadUrl) {
-            return <a onClick={(e) => e.stopPropagation()}  href={issue.workflowThreadUrl} target="_blank"  rel="noreferrer"><Stack style={{ paddingTop: 6, height: "100%"}} horizontalAlign="start" disableShrink={true}  >Slack Thread</Stack></a>
+            return <a onClick={(e) => e.stopPropagation()} href={issue.workflowThreadUrl} target="_blank" rel="noreferrer"><Stack style={{ paddingTop: 6, height: "100%" }} horizontalAlign="start" disableShrink={true}  >Slack Thread</Stack></a>
          }
 
-         return <Stack style={{ paddingTop: 6, height: "100%"}} horizontalAlign="start" disableShrink={true} ><Text variant={textSize}>Quarantined</Text></Stack>;
+         return <Stack style={{ paddingTop: 6, height: "100%" }} horizontalAlign="start" disableShrink={true} ><Text variant={textSize}>Quarantined</Text></Stack>;
       }
 
 
@@ -511,7 +511,7 @@ const HealthPanelIssues: React.FC<{ desktopAlerts?: boolean }> = observer(({ des
             status += ` (${ack})`;
          }
 
-         return <Stack style={{ paddingTop: 6, height: "100%"}} horizontalAlign="end" disableShrink={true} ><Text variant={textSize}>{status}</Text></Stack>;
+         return <Stack style={{ paddingTop: 6, height: "100%" }} horizontalAlign="end" disableShrink={true} ><Text variant={textSize}>{status}</Text></Stack>;
       }
 
       if (issue.externalIssueKey) {
@@ -531,7 +531,7 @@ const HealthPanelIssues: React.FC<{ desktopAlerts?: boolean }> = observer(({ des
                priority = priority.slice(priority.indexOf("- ") + 2);
             }
 
-            return <a href={jissue.link} target="_blank" rel="noreferrer" onClick={(ev) => ev.stopPropagation()}><Stack style={{ paddingTop: 6, height: "100%"}} horizontalAlign="start" disableShrink={true}><Text style={{ textDecoration: textDecoration }} variant={textSize}>{priority}</Text></Stack></a>;
+            return <a href={jissue.link} target="_blank" rel="noreferrer" onClick={(ev) => ev.stopPropagation()}><Stack style={{ paddingTop: 6, height: "100%" }} horizontalAlign="start" disableShrink={true}><Text style={{ textDecoration: textDecoration }} variant={textSize}>{priority}</Text></Stack></a>;
 
          }
 
@@ -542,7 +542,7 @@ const HealthPanelIssues: React.FC<{ desktopAlerts?: boolean }> = observer(({ des
                return null;
             }
 
-            return <a href={jissue.link} target="_blank" rel="noreferrer" onClick={(ev) => ev.stopPropagation()}><Stack style={{ paddingTop: 6, height: "100%"}} horizontalAlign="center" disableShrink={true}><Text style={{ textDecoration: textDecoration }} variant={textSize}>{jissue.assigneeDisplayName ?? "Unassigned (Jira)"}</Text></Stack></a>
+            return <a href={jissue.link} target="_blank" rel="noreferrer" onClick={(ev) => ev.stopPropagation()}><Stack style={{ paddingTop: 6, height: "100%" }} horizontalAlign="center" disableShrink={true}><Text style={{ textDecoration: textDecoration }} variant={textSize}>{jissue.assigneeDisplayName ?? "Unassigned (Jira)"}</Text></Stack></a>
          }
 
          if (column.name === "JiraStatus") {
@@ -553,7 +553,7 @@ const HealthPanelIssues: React.FC<{ desktopAlerts?: boolean }> = observer(({ des
                return null;
             }
 
-            return <a href={jissue.link} target="_blank" rel="noreferrer" onClick={(ev) => ev.stopPropagation()}><Stack style={{ paddingTop: 6, height: "100%"}} horizontalAlign="start" disableShrink={true}><Text style={{ textDecoration: textDecoration }} variant={textSize}>{desc}</Text></Stack></a>;
+            return <a href={jissue.link} target="_blank" rel="noreferrer" onClick={(ev) => ev.stopPropagation()}><Stack style={{ paddingTop: 6, height: "100%" }} horizontalAlign="start" disableShrink={true}><Text style={{ textDecoration: textDecoration }} variant={textSize}>{desc}</Text></Stack></a>;
          }
 
          if (column.name === "Jira") {
@@ -562,7 +562,7 @@ const HealthPanelIssues: React.FC<{ desktopAlerts?: boolean }> = observer(({ des
                return null;
             }
 
-            return <a href={jissue.link} target="_blank" rel="noreferrer" onClick={(ev) => ev.stopPropagation()}><Stack style={{ paddingTop: 6, height: "100%"}} horizontalAlign="start" disableShrink={true}><Text style={{ textDecoration: textDecoration }} variant={textSize}>{jissue.key}</Text></Stack></a>;
+            return <a href={jissue.link} target="_blank" rel="noreferrer" onClick={(ev) => ev.stopPropagation()}><Stack style={{ paddingTop: 6, height: "100%" }} horizontalAlign="start" disableShrink={true}><Text style={{ textDecoration: textDecoration }} variant={textSize}>{jissue.key}</Text></Stack></a>;
          }
       }
 
@@ -583,7 +583,7 @@ const HealthPanelIssues: React.FC<{ desktopAlerts?: boolean }> = observer(({ des
    const classes = mergeStyleSets({
       detailsRow: {
          selectors: {
-            '.ms-DetailsRow-cell': {             
+            '.ms-DetailsRow-cell': {
                padding: 0,
                overflow: "hidden",
                whiteSpace: "nowrap"
@@ -776,18 +776,16 @@ const StreamSummaryInner: React.FC = observer(() => {
    // subscribe
    if (handler.update) { }
 
-   return <Stack tokens={{ childrenGap: 0 }} styles={{ root: { width: 1440, backgroundColor: "#fffffff", margin: 0, paddingTop: 8 } }}>
-      <Stack style={{ padding: 0 }} className={detailClasses.detailsRow}>
-         <FocusZone direction={FocusZoneDirection.vertical} style={{ padding: 0 }}>
-            <div className={detailClasses.container} style={{ width: "100%", height: 'calc(100vh - 278px)', position: 'relative' }} data-is-scrollable={true}>
-               <ScrollablePane scrollbarVisibility={ScrollbarVisibility.always} onScroll={() => { }}>
-                  <Stack tokens={{ childrenGap: 18 }} style={{ padding: 0 }}>
-                     <HealthPanel desktopAlerts={false} />
-                     <SchedulePanel />
-                  </Stack>
-               </ScrollablePane>
-            </div>
-         </FocusZone>
+   return <Stack tokens={{ childrenGap: 0 }} styles={{ root: { width: "100%", backgroundColor: "#fffffff", margin: 0, paddingTop: 8 } }}>
+      <Stack style={{ padding: 0 }}>
+         <div style={{ marginTop: 8, height: 'calc(100vh - 258px)', position: 'relative' }} data-is-scrollable={true}>
+            <ScrollablePane scrollbarVisibility={ScrollbarVisibility.always} onScroll={() => { }}>
+               <Stack tokens={{ childrenGap: 18 }} style={{ padding: 0 }}>
+                  <HealthPanel desktopAlerts={false} />
+                  <SchedulePanel />
+               </Stack>
+            </ScrollablePane>
+         </div>
       </Stack>
    </Stack>
 });
