Horde: Adjust CL range to be +1 to previous CL, also add a Show CL History, which shows that CL and all previous
#rnx
#skipci

diff --git a/Engine/Source/Programs/Horde/HordeDashboard/src/components/ChangeButton.tsx b/Engine/Source/Programs/Horde/HordeDashboard/src/components/ChangeButton.tsx
--- a/Engine/Source/Programs/Horde/HordeDashboard/src/components/ChangeButton.tsx
+++ b/Engine/Source/Programs/Horde/HordeDashboard/src/components/ChangeButton.tsx
@@ -115,10 +115,9 @@ export const ChangeContextMenu: React.FC<{ target: ChangeContextMenuTarget, onDi
    const stream = projectStore.streamById(job.streamId)!;
    const project = projectStore.byId(stream!.projectId)!;
    const name = project.name === "Engine" ? "UE4" : project.name;
-
-   let url = "";
-   url = `${dashboard.swarmUrl}/files/${name}/${stream.name}?range=@${jobChange}#commits`;
-
+   
+   const historyUrl = `${dashboard.swarmUrl}/files/${name}/${stream.name}?range=@${jobChange}#commits`;
+   menuItems.push({ key: 'open_in_swarm_history', text: "Open CL History in Swarm", onClick: (ev) => { window.open(historyUrl) } })
 
    let highCL = jobChange;
    let lowCL = jobChange;
@@ -126,20 +125,11 @@ export const ChangeContextMenu: React.FC<{ target: ChangeContextMenuTarget, onDi
    if (rangeCL !== undefined) {
       lowCL = Math.min(jobChange, rangeCL);
       highCL = Math.max(jobChange, rangeCL);
-      url = `${dashboard.swarmUrl}/files/${name}/${stream.name}?range=@${lowCL},@${highCL}#commits`;
-   }
-
-   let rangeText = 'Open CL Range in Swarm';
-
-   if (rangeCL) {
-      rangeText = `Open CL Range ${lowCL} - ${highCL}`;
-   }
-
-
-   if (url) {
+      const url = `${dashboard.swarmUrl}/files/${name}/${stream.name}?range=@${lowCL},@${highCL}#commits`;
+      const rangeText = `Open CL Range ${lowCL} - ${highCL}`;
       menuItems.push({ key: 'open_in_swarm_range', text: rangeText, onClick: (ev) => { window.open(url) } })
    }
-
+   
    return (<ContextualMenu
       styles={{ list: { selectors: { '.ms-ContextualMenu-itemText': { fontSize: "10px", paddingLeft: 8 } } } }}
       items={menuItems}
diff --git a/Engine/Source/Programs/Horde/HordeDashboard/src/components/JobOperationsContextMenu.tsx b/Engine/Source/Programs/Horde/HordeDashboard/src/components/JobOperationsContextMenu.tsx
--- a/Engine/Source/Programs/Horde/HordeDashboard/src/components/JobOperationsContextMenu.tsx
+++ b/Engine/Source/Programs/Horde/HordeDashboard/src/components/JobOperationsContextMenu.tsx
@@ -64,7 +64,7 @@ export const JobOperationsContextMenu: React.FC<{ job: GetJobResponse, children?
          }
       },
       {
-         key: 'open_in_swarm_range', text: 'Open CL Range in Swarm', onClick: () => {
+         key: 'open_in_swarm_history', text: 'Open CL History in Swarm', onClick: () => {
 
             const stream = projectStore.streamById(job.streamId)!;
             const project = projectStore.byId(stream!.projectId)!;
diff --git a/Engine/Source/Programs/Horde/HordeDashboard/src/components/JobViewIncremental.tsx b/Engine/Source/Programs/Horde/HordeDashboard/src/components/JobViewIncremental.tsx
--- a/Engine/Source/Programs/Horde/HordeDashboard/src/components/JobViewIncremental.tsx
+++ b/Engine/Source/Programs/Horde/HordeDashboard/src/components/JobViewIncremental.tsx
@@ -185,7 +185,7 @@ const JobList: React.FC<{ tab: string; filter: JobFilterSimple, controller: Call
    let [update, setUpdate] = useState(0);
 
    // depending how this is optimized in production, might need to change around so rolling over cl doesn't re-render
-   const [commitState, setCommitState] = useState<{ target?: ChangeContextMenuTarget, job?: JobData, commit?: GetChangeSummaryResponse }>({});
+   const [commitState, setCommitState] = useState<{ target?: ChangeContextMenuTarget, job?: JobData, commit?: GetChangeSummaryResponse, rangeCL?: number }>({});
 
    useEffect(() => {
 
@@ -514,7 +514,27 @@ const JobList: React.FC<{ tab: string; filter: JobFilterSimple, controller: Call
          }
 
 
-         return <Stack style={{ margin: 0, cursor: "pointer" }} onClick={(ev) => { ev.stopPropagation(); ev.preventDefault(); setCommitState({ job: item.job, commit: commit, target: { point: { x: ev.clientX, y: ev.clientY } } }) }} >
+         return <Stack style={{ margin: 0, cursor: "pointer" }} onClick={(ev) => {
+            ev.stopPropagation();
+            ev.preventDefault();
+
+            let rangeCL: number | undefined;
+            if (item.job) {
+
+               const cjob = jobs.find((j => j.id === item.job!.id));
+               if (cjob) {
+                  const index = jobs.indexOf(cjob);
+                  if (index < jobs.length - 1) {
+                     rangeCL = jobs[index + 1].change;
+                     if (rangeCL) {
+                        rangeCL += 1;                        
+                     }
+                  }
+               }
+            }
+
+            setCommitState({ job: item.job, rangeCL: rangeCL, commit: commit, target: { point: { x: ev.clientX, y: ev.clientY } } })
+         }} >
             <Stack verticalAlign="center" verticalFill={true} horizontal horizontalAlign={"start"} tokens={{ childrenGap: 18 }}>
                <Stack verticalAlign="center" verticalFill={true} horizontalAlign="start"> <div style={{ paddingBottom: "1px" }}>
                   <span style={{ padding: "2px 6px 2px 6px", height: "18px", cursor: "pointer" }} className={item.job.startedByUserInfo ? "cl-callout-button-user" : "cl-callout-button"} >{change}</span>
@@ -598,7 +618,7 @@ const JobList: React.FC<{ tab: string; filter: JobFilterSimple, controller: Call
    return (
       <Stack tokens={{ childrenGap: 0 }} className={detailClasses.detailsRow}>
          <FocusZone direction={FocusZoneDirection.vertical} >
-            {commitState.target && commitState.job && <ChangeContextMenu target={commitState.target} job={commitState.job} commit={commitState.commit} onDismiss={() => setCommitState({})} />}
+            {commitState.target && commitState.job && <ChangeContextMenu target={commitState.target} job={commitState.job} commit={commitState.commit} rangeCL={commitState.rangeCL} onDismiss={() => setCommitState({})} />}
             <div className={detailClasses.container} style={{ width: "100%", height: 'calc(100vh - 240px)', position: 'relative', marginTop: 0 }} data-is-scrollable={true}>
 
                {<ScrollablePane scrollbarVisibility={ScrollbarVisibility.always} onScroll={() => { controller.setState(undefined, true) }} style={{ overflow: "visible" }}>
diff --git a/Engine/Source/Programs/Horde/HordeDashboard/src/components/jobDetailsV2/JobDetailStepHistory.tsx b/Engine/Source/Programs/Horde/HordeDashboard/src/components/jobDetailsV2/JobDetailStepHistory.tsx
--- a/Engine/Source/Programs/Horde/HordeDashboard/src/components/jobDetailsV2/JobDetailStepHistory.tsx
+++ b/Engine/Source/Programs/Horde/HordeDashboard/src/components/jobDetailsV2/JobDetailStepHistory.tsx
@@ -130,7 +130,7 @@ export const StepHistoryPanel: React.FC<{ jobDetails: JobDetailsV2; stepId: stri
       if (column.name === "Change") {
 
          const sindex = dataView.history.indexOf(ref); 
-         return <ChangeButton job={jobDetails.jobData!} stepRef={ref} hideAborted={true} rangeCL={ sindex < (dataView.history.length - 1) ? dataView.history[sindex + 1].change : undefined}/>;
+         return <ChangeButton job={jobDetails.jobData!} stepRef={ref} hideAborted={true} rangeCL={ sindex < (dataView.history.length - 1) ? (dataView.history[sindex + 1].change + 1) : undefined}/>;
 
       }
 
diff --git a/Engine/Source/Programs/Horde/HordeDashboard/src/components/jobDetailsV2/JobDetailStepHistoryV2.tsx b/Engine/Source/Programs/Horde/HordeDashboard/src/components/jobDetailsV2/JobDetailStepHistoryV2.tsx
--- a/Engine/Source/Programs/Horde/HordeDashboard/src/components/jobDetailsV2/JobDetailStepHistoryV2.tsx
+++ b/Engine/Source/Programs/Horde/HordeDashboard/src/components/jobDetailsV2/JobDetailStepHistoryV2.tsx
@@ -660,7 +660,7 @@ const GraphTooltip: React.FC<{ dataView: StepHistoryDataView }> = observer(({ da
          </Link>
          <Stack style={{ paddingLeft: 2, paddingTop: 8 }} tokens={{ childrenGap: 8 }}>
             <Stack>
-               <ChangeButton prefix="CL" job={dataView.details!.jobData!} stepRef={ref} hideAborted={true} rangeCL={sindex < (dataView.history.length - 1) ? dataView.history[sindex + 1].change : undefined} />
+               <ChangeButton prefix="CL" job={dataView.details!.jobData!} stepRef={ref} hideAborted={true} rangeCL={sindex < (dataView.history.length - 1) ? (dataView.history[sindex + 1].change + 1) : undefined} />
             </Stack>
             <Stack><Text variant={textSize}>{displayTimeStr}</Text></Stack>
             <Stack><Text variant={textSize}>Duration: {duration}</Text></Stack>
