Horde: Add a helper method to do an OIDC login in the RemoteClient app.

diff --git a/Engine/Source/Programs/Horde/Samples/RemoteClient/RemoteClient.csproj b/Engine/Source/Programs/Horde/Samples/RemoteClient/RemoteClient.csproj
--- a/Engine/Source/Programs/Horde/Samples/RemoteClient/RemoteClient.csproj
+++ b/Engine/Source/Programs/Horde/Samples/RemoteClient/RemoteClient.csproj
@@ -9,6 +9,7 @@
 
   <ItemGroup>
     <ProjectReference Include="..\..\..\Shared\EpicGames.Horde\EpicGames.Horde.csproj" />
+    <ProjectReference Include="..\..\..\Shared\EpicGames.OIDC\EpicGames.OIDC.csproj" />
   </ItemGroup>
 
 </Project>
diff --git a/Engine/Source/Programs/Horde/Samples/RemoteClient/RemoteClientApp.cs b/Engine/Source/Programs/Horde/Samples/RemoteClient/RemoteClientApp.cs
--- a/Engine/Source/Programs/Horde/Samples/RemoteClient/RemoteClientApp.cs
+++ b/Engine/Source/Programs/Horde/Samples/RemoteClient/RemoteClientApp.cs
@@ -7,6 +7,8 @@ using EpicGames.Horde.Compute.Clients;
 using EpicGames.Horde.Storage;
 using EpicGames.Horde.Storage.Backends;
 using EpicGames.Horde.Storage.Nodes;
+using EpicGames.OIDC;
+using Microsoft.Extensions.Configuration;
 using Microsoft.Extensions.Logging;
 
 namespace RemoteClient
@@ -82,5 +84,20 @@ namespace RemoteClient
 				return new ServerComputeClient(new Uri("https://localhost:5001"), null, logger);
 			}
 		}
+
+		static async Task<string> GetBearerToken(DirectoryReference engineDir, DirectoryReference? projectDir, string providerIdentifier)
+		{
+			using ITokenStore tokenStore = TokenStoreFactory.CreateTokenStore();
+			IConfiguration providerConfiguration = ProviderConfigurationFactory.ReadConfiguration(engineDir.ToDirectoryInfo(), projectDir?.ToDirectoryInfo());
+			OidcTokenManager oidcTokenManager = OidcTokenManager.CreateTokenManager(providerConfiguration, tokenStore, new List<string>() { providerIdentifier });
+			OidcTokenInfo result = await oidcTokenManager.Login(providerIdentifier);
+
+			if (result.AccessToken == null)
+			{
+				throw new Exception($"Unable to get access token for {providerIdentifier}");
+			}
+
+			return result.AccessToken;
+		}
 	}
 }
\ No newline at end of file
