Horde: Add Slack app icon and manifest to source control.

diff --git a/Engine/Source/Programs/Horde/Horde.Server/Slack/manifest.json b/Engine/Source/Programs/Horde/Horde.Server/Slack/manifest.json
new file mode 100644
--- /dev/null
+++ b/Engine/Source/Programs/Horde/Horde.Server/Slack/manifest.json
@@ -0,0 +1,39 @@
+{
+    "display_information": {
+        "name": "Horde",
+        "description": "Allow for interaction with the Horde build system.",
+        "background_color": "#000000"
+    },
+    "features": {
+        "bot_user": {
+            "display_name": "Horde",
+            "always_online": false
+        }
+    },
+    "oauth_config": {
+        "scopes": {
+            "user": [
+                "admin.conversations:write"
+            ],
+            "bot": [
+                "chat:write",
+                "chat:write.public",
+                "reactions:read",
+                "reactions:write",
+                "users.profile:read",
+                "users:read",
+                "users:read.email",
+                "channels:manage"
+            ]
+        }
+    },
+    "settings": {
+        "interactivity": {
+            "is_enabled": true,
+            "request_url": "https://hordeserver.devtools.epicgames.com/api/v1/slack"
+        },
+        "org_deploy_enabled": true,
+        "socket_mode_enabled": true,
+        "token_rotation_enabled": false
+    }
+}
