Horde: Remove internal URL from Slack manifest.

diff --git a/Engine/Source/Programs/Horde/Horde.Server/Slack/manifest.json b/Engine/Source/Programs/Horde/Horde.Server/Slack/manifest.json
--- a/Engine/Source/Programs/Horde/Horde.Server/Slack/manifest.json
+++ b/Engine/Source/Programs/Horde/Horde.Server/Slack/manifest.json
@@ -29,8 +29,7 @@
     },
     "settings": {
         "interactivity": {
-            "is_enabled": true,
-            "request_url": "https://hordeserver.devtools.epicgames.com/api/v1/slack"
+            "is_enabled": true
         },
         "org_deploy_enabled": true,
         "socket_mode_enabled": true,
