Horde: Add a check that files extracted from temp storage match the sizes listed in the manifest.

diff --git a/Engine/Source/Programs/Horde/Horde.Agent/Utility/TempStorage.cs b/Engine/Source/Programs/Horde/Horde.Agent/Utility/TempStorage.cs
--- a/Engine/Source/Programs/Horde/Horde.Agent/Utility/TempStorage.cs
+++ b/Engine/Source/Programs/Horde/Horde.Agent/Utility/TempStorage.cs
@@ -704,10 +704,15 @@ namespace Horde.Storage.Utility
 				manifest = TempStorageBlockManifest.Load(localManifestFile);
 
 				// Update the timestamps to match the manifest.
-				foreach (TempStorageFile ManifestFile in manifest.Files)
+				foreach (TempStorageFile manifestFile in manifest.Files)
 				{
-					FileReference File = ManifestFile.ToFileReference(rootDir);
-					System.IO.File.SetLastWriteTimeUtc(File.FullName, new DateTime(ManifestFile.LastWriteTimeUtcTicks, DateTimeKind.Utc));
+					FileInfo fileInfo = manifestFile.ToFileReference(rootDir).ToFileInfo();
+					fileInfo.LastWriteTimeUtc = new DateTime(manifestFile.LastWriteTimeUtcTicks, DateTimeKind.Utc);
+
+					if (fileInfo.Length != manifestFile.Length)
+					{
+						throw new TempStorageException($"File extracted from temp storage has different size to file in manifest (manifest: {manifestFile.Length}, local: {fileInfo.Length})");
+					}
 				}
 			}
 
