Horde: Assign compute leases directly to agents, bypassing the regular task source flow.

diff --git a/Engine/Source/Programs/Horde/Horde.Server/Agents/AgentCollection.cs b/Engine/Source/Programs/Horde/Horde.Server/Agents/AgentCollection.cs
--- a/Engine/Source/Programs/Horde/Horde.Server/Agents/AgentCollection.cs
+++ b/Engine/Source/Programs/Horde/Horde.Server/Agents/AgentCollection.cs
@@ -552,7 +552,7 @@ namespace Horde.Server.Agents
 		}
 
 		/// <inheritdoc/>
-		Task PublishUpdateEventAsync(AgentId agentId)
+		public Task PublishUpdateEventAsync(AgentId agentId)
 		{
 			return _redisService.GetDatabase().PublishAsync(_updateEventChannel, agentId);
 		}
diff --git a/Engine/Source/Programs/Horde/Horde.Server/Agents/IAgentCollection.cs b/Engine/Source/Programs/Horde/Horde.Server/Agents/IAgentCollection.cs
--- a/Engine/Source/Programs/Horde/Horde.Server/Agents/IAgentCollection.cs
+++ b/Engine/Source/Programs/Horde/Horde.Server/Agents/IAgentCollection.cs
@@ -149,6 +149,12 @@ namespace Horde.Server.Agents
 		/// <returns></returns>
 		IAuditLogChannel<AgentId> GetLogger(AgentId agentId);
 
+		/// <summary>
+		/// Sends a notification that an update event has ocurred
+		/// </summary>
+		/// <param name="agentId">Agent that has been updated</param>
+		Task PublishUpdateEventAsync(AgentId agentId);
+
 		/// <summary>
 		/// Subscribe to notifications on agent states being updated
 		/// </summary>
diff --git a/Engine/Source/Programs/Horde/Horde.Server/Compute/ComputeController.cs b/Engine/Source/Programs/Horde/Horde.Server/Compute/ComputeController.cs
--- a/Engine/Source/Programs/Horde/Horde.Server/Compute/ComputeController.cs
+++ b/Engine/Source/Programs/Horde/Horde.Server/Compute/ComputeController.cs
@@ -3,6 +3,7 @@
 using System;
 using System.Collections.Generic;
 using System.Net;
+using System.Threading.Tasks;
 using EpicGames.Core;
 using EpicGames.Horde.Compute;
 using Horde.Server.Acls;
@@ -69,15 +70,15 @@ namespace Horde.Server.Compute
 	[Route("[controller]")]
 	public class ComputeControllerV2 : HordeControllerBase
 	{
-		readonly ComputeTaskSource _computeTaskSource;
+		readonly ComputeService _computeService;
 		readonly IOptionsSnapshot<GlobalConfig> _globalConfig;
 
 		/// <summary>
 		/// Constructor
 		/// </summary>
-		public ComputeControllerV2(ComputeTaskSource computeTaskSource, IOptionsSnapshot<GlobalConfig> globalConfig)
+		public ComputeControllerV2(ComputeService computeService, IOptionsSnapshot<GlobalConfig> globalConfig)
 		{
-			_computeTaskSource = computeTaskSource;
+			_computeService = computeService;
 			_globalConfig = globalConfig;
 		}
 
@@ -90,7 +91,7 @@ namespace Horde.Server.Compute
 		[HttpPost]
 		[Authorize]
 		[Route("/api/v2/compute/{clusterId}")]
-		public ActionResult<AssignComputeResponse> AssignComputeResourceAsync(ClusterId clusterId, [FromBody] AssignComputeRequest request)
+		public async Task<ActionResult<AssignComputeResponse>> AssignComputeResourceAsync(ClusterId clusterId, [FromBody] AssignComputeRequest request)
 		{
 			ComputeClusterConfig? clusterConfig;
 			if (!_globalConfig.Value.TryGetComputeCluster(clusterId, out clusterConfig))
@@ -104,7 +105,7 @@ namespace Horde.Server.Compute
 
 			Requirements requirements = request.Requirements ?? new Requirements();
 
-			ComputeResource? computeResource = _computeTaskSource.TryAllocateResource(clusterId, requirements);
+			ComputeResource? computeResource = await _computeService.TryAllocateResource(requirements);
 			if (computeResource == null)
 			{
 				return StatusCode((int)HttpStatusCode.ServiceUnavailable);
diff --git a/Engine/Source/Programs/Horde/Horde.Server/Compute/ComputeService.cs b/Engine/Source/Programs/Horde/Horde.Server/Compute/ComputeService.cs
new file mode 100644
--- /dev/null
+++ b/Engine/Source/Programs/Horde/Horde.Server/Compute/ComputeService.cs
@@ -0,0 +1,96 @@
+// Copyright Epic Games, Inc. All Rights Reserved.
+
+using System;
+using System.Collections.Generic;
+using System.Linq;
+using System.Net;
+using System.Security.Cryptography;
+using System.Threading.Tasks;
+using EpicGames.Horde.Compute;
+using EpicGames.Horde.Compute.Clients;
+using EpicGames.Horde.Compute.Transports;
+using Google.Protobuf;
+using Google.Protobuf.WellKnownTypes;
+using Horde.Server.Agents;
+using Horde.Server.Agents.Leases;
+using HordeCommon;
+using HordeCommon.Rpc.Tasks;
+
+namespace Horde.Server.Compute
+{
+	/// <summary>
+	/// Assigns compute leases to agents
+	/// </summary>
+	public class ComputeService
+	{
+		readonly IAgentCollection _agentCollection;
+
+		/// <summary>
+		/// Constructor
+		/// </summary>
+		/// <param name="agentCollection"></param>
+		public ComputeService(IAgentCollection agentCollection)
+		{
+			_agentCollection = agentCollection;
+		}
+
+		/// <summary>
+		/// Allocates a compute resource
+		/// </summary>
+		/// <returns></returns>
+		public async Task<ComputeResource?> TryAllocateResource(Requirements requirements)
+		{
+			List<IAgent> agents = await _agentCollection.FindAsync();
+			foreach (IAgent agent in agents)
+			{
+				Dictionary<string, int> assignedResources = new Dictionary<string, int>();
+				if (agent.MeetsRequirements(requirements, assignedResources))
+				{
+					ComputeTask computeTask = CreateComputeTask(assignedResources);
+
+					byte[] payload = Any.Pack(computeTask).ToByteArray();
+					AgentLease lease = new AgentLease(LeaseId.GenerateNewId(), "Compute task", null, null, null, LeaseState.Pending, assignedResources, requirements.Exclusive, payload);
+
+					ComputeResource? resource = TryAssign(agent, computeTask);
+					if (resource != null)
+					{
+						IAgent? newAgent = await _agentCollection.TryAddLeaseAsync(agent, lease);
+						if (newAgent != null)
+						{
+							await _agentCollection.PublishUpdateEventAsync(agent.Id);
+							return resource;
+						}
+					}
+				}
+			}
+			return null;
+		}
+
+		static ComputeResource? TryAssign(IAgent agent, ComputeTask computeTask)
+		{
+			string? ipStr = agent.GetPropertyValues("ComputeIp").FirstOrDefault();
+			if (ipStr == null || !IPAddress.TryParse(ipStr, out IPAddress? ip))
+			{
+				return null;
+			}
+
+			string? portStr = agent.GetPropertyValues("ComputePort").FirstOrDefault();
+			if (portStr == null || !Int32.TryParse(portStr, out int port))
+			{
+				return null;
+			}
+
+			return new ComputeResource(ip, port, computeTask, agent.Properties);
+		}
+
+
+		static ComputeTask CreateComputeTask(Dictionary<string, int> assignedResources)
+		{
+			ComputeTask computeTask = new ComputeTask();
+			computeTask.Nonce = UnsafeByteOperations.UnsafeWrap(RandomNumberGenerator.GetBytes(ServerComputeClient.NonceLength));
+			computeTask.Key = UnsafeByteOperations.UnsafeWrap(AesTransport.CreateKey());
+			computeTask.Resources.Add(assignedResources);
+			return computeTask;
+		}
+	}
+}
diff --git a/Engine/Source/Programs/Horde/Horde.Server/Startup.cs b/Engine/Source/Programs/Horde/Horde.Server/Startup.cs
--- a/Engine/Source/Programs/Horde/Horde.Server/Startup.cs
+++ b/Engine/Source/Programs/Horde/Horde.Server/Startup.cs
@@ -396,6 +396,7 @@ namespace Horde.Server
 			services.AddSingleton<AgentService>();
 			services.AddSingleton<ConsistencyService>();
 			services.AddSingleton<RequestTrackerService>();
+			services.AddSingleton<ComputeService>();
 			services.AddSingleton<CredentialService>();
 			services.AddSingleton<MongoService>();
 			services.AddSingleton<GlobalsService>();
