Horde: Additional logging for compute service.

diff --git a/Engine/Source/Programs/Horde/Horde.Agent/Leases/Handlers/ComputeHandler.cs b/Engine/Source/Programs/Horde/Horde.Agent/Leases/Handlers/ComputeHandler.cs
--- a/Engine/Source/Programs/Horde/Horde.Agent/Leases/Handlers/ComputeHandler.cs
+++ b/Engine/Source/Programs/Horde/Horde.Agent/Leases/Handlers/ComputeHandler.cs
@@ -38,7 +38,7 @@ namespace Horde.Agent.Leases.Handlers
 		/// <inheritdoc/>
 		public override async Task<LeaseResult> ExecuteAsync(ISession session, string leaseId, ComputeTask computeTask, CancellationToken cancellationToken)
 		{
-			_logger.LogInformation("Starting compute task (lease {LeaseId}). Waiting for connection...", leaseId);
+			_logger.LogInformation("Starting compute task (lease {LeaseId}). Waiting for connection with nonce {Nonce}...", leaseId, StringUtils.FormatHexString(computeTask.Nonce.Span));
 
 			TcpClient? tcpClient = null;
 			try
@@ -52,6 +52,8 @@ namespace Horde.Agent.Leases.Handlers
 					return LeaseResult.Success;
 				}
 
+				_logger.LogInformation("Matched connection for {Nonce}", StringUtils.FormatHexString(computeTask.Nonce));
+
 				await using (ClientComputeSocket socket = new ClientComputeSocket(new TcpTransport(tcpClient.Client), _logger))
 				{
 					ComputeWorker worker = new ComputeWorker(_sandboxDir, _memoryCache, _logger);
diff --git a/Engine/Source/Programs/Horde/Horde.Agent/Services/ComputeListenerService.cs b/Engine/Source/Programs/Horde/Horde.Agent/Services/ComputeListenerService.cs
--- a/Engine/Source/Programs/Horde/Horde.Agent/Services/ComputeListenerService.cs
+++ b/Engine/Source/Programs/Horde/Horde.Agent/Services/ComputeListenerService.cs
@@ -175,6 +175,7 @@ namespace Horde.Agent.Services
 				// Read the nonce for the connection
 				byte[] nonceBuffer = new byte[ServerComputeClient.NonceLength];
 				await socket.ReceiveMessageAsync(nonceBuffer, SocketFlags.None, cancellationToken);
+				_logger.LogInformation("Incoming connection with nonce {Nonce}", StringUtils.FormatHexString(nonceBuffer));
 
 				// Register the socket
 				ByteString nonce = new ByteString(nonceBuffer);
