Horde: Increase timeout for compute listener.

diff --git a/Engine/Source/Programs/Horde/Horde.Agent/Services/ComputeListenerService.cs b/Engine/Source/Programs/Horde/Horde.Agent/Services/ComputeListenerService.cs
--- a/Engine/Source/Programs/Horde/Horde.Agent/Services/ComputeListenerService.cs
+++ b/Engine/Source/Programs/Horde/Horde.Agent/Services/ComputeListenerService.cs
@@ -202,8 +202,8 @@ namespace Horde.Agent.Services
 						return;
 					}
 
-					// Wait 30s for it to be claimed
-					Task delayTask = Task.Delay(TimeSpan.FromSeconds(30.0), cancellationToken);
+					// Wait 2m for it to be claimed
+					Task delayTask = Task.Delay(TimeSpan.FromMinutes(2.0), cancellationToken);
 					await Task.WhenAny(waitingClient!.WaitForClientToBeTaken(), delayTask);
 				}
 				finally
