Horde: Remove AgentCollection.TryUpdateAsyncWithRetries() (added in CL 20241607).

Operations on agent documents are meant to be compare-and-swap operations, hence the dependency on the UpdateIndex matching. Applying the same update to an updated document is not guaranteed to be valid, and can result in state being incorrectly overwritten. Update errors must to be propagated back to the caller to determine the correct course of action for an updated agent document.

diff --git a/Engine/Source/Programs/Horde/Horde.Server/Agents/AgentCollection.cs b/Engine/Source/Programs/Horde/Horde.Server/Agents/AgentCollection.cs
--- a/Engine/Source/Programs/Horde/Horde.Server/Agents/AgentCollection.cs
+++ b/Engine/Source/Programs/Horde/Horde.Server/Agents/AgentCollection.cs
@@ -169,7 +169,7 @@ namespace Horde.Server.Agents
 			AgentDocument agent = (AgentDocument)agentInterface;
 
 			UpdateDefinition<AgentDocument> update = Builders<AgentDocument>.Update.Set(x => x.Deleted, true);
-			return await TryUpdateAsyncWithRetries(agent, update);
+			return await TryUpdateAsync(agent, update);
 		}
 
 		/// <inheritdoc/>
@@ -249,36 +249,6 @@ namespace Horde.Server.Agents
 			return await _agents.FindOneAndUpdateAsync<AgentDocument>(filter, updateWithIndex, new FindOneAndUpdateOptions<AgentDocument, AgentDocument> { ReturnDocument = ReturnDocument.After });
 		}
 
-		/// <summary>
-		/// Update a single document, with retries to ensure it's applied
-		/// If an outdated AgentDocument is supplied, it needs to be refreshed so that the UpdateIndex matches.
-		/// </summary>
-		/// <param name="agentDoc">The document to update</param>
-		/// <param name="update">The update definition</param>
-		/// <param name="numRetries">Number of retries to attempt</param>
-		/// <returns>The updated agent document or null if update failed</returns>
-		private async Task<AgentDocument?> TryUpdateAsyncWithRetries(AgentDocument agentDoc, UpdateDefinition<AgentDocument> update, int numRetries = 5)
-		{
-			AgentDocument current = agentDoc;
-			for (int retryCount = 0; retryCount < numRetries; retryCount++)
-			{
-				AgentDocument? result = await TryUpdateAsync(current, update);
-				if (result != null)
-				{
-					return result;
-				}
-
-				AgentId id = current.Id; // Capture id to avoid doing it inside closure (and warnings)
-				current = await _agents.Find(x => x.Id == id).FirstOrDefaultAsync();
-				if (current == null)
-				{
-					return null;
-				}
-			}
-
-			return null;
-		}
-
 		/// <inheritdoc/>
 		public async Task<IAgent?> TryUpdateSettingsAsync(IAgent agentInterface, bool? enabled = null, bool? requestConform = null, bool? requestFullConform = null, bool? requestRestart = null, bool? requestShutdown = null, string? shutdownReason = null, List<PoolId>? pools = null, string? comment = null)
 		{
@@ -340,7 +310,7 @@ namespace Horde.Server.Agents
 			}
 
 			// Apply the update
-			IAgent? newAgent = await TryUpdateAsyncWithRetries(agent, updateBuilder.Combine(updates));
+			IAgent? newAgent = await TryUpdateAsync(agent, updateBuilder.Combine(updates));
 			if (newAgent != null)
 			{
 				if (newAgent.RequestRestart != agent.RequestRestart || newAgent.RequestConform != agent.RequestConform || newAgent.RequestShutdown != agent.RequestShutdown)
@@ -416,7 +386,7 @@ namespace Horde.Server.Agents
 			}
 
 			// Update the agent, and try to create new lease documents if we succeed
-			return await TryUpdateAsyncWithRetries(agent, updateBuilder.Combine(updates));
+			return await TryUpdateAsync(agent, updateBuilder.Combine(updates));
 		}
 
 		static bool ResourcesEqual(IReadOnlyDictionary<string, int> dictA, IReadOnlyDictionary<string, int> dictB)
@@ -455,7 +425,7 @@ namespace Horde.Server.Agents
 			}
 
 			// Update the agent
-			return await TryUpdateAsyncWithRetries(agent, update);
+			return await TryUpdateAsync(agent, update);
 		}
 
 		/// <inheritdoc/>
@@ -491,7 +461,7 @@ namespace Horde.Server.Agents
 			}
 
 			// Apply the update
-			return await TryUpdateAsyncWithRetries(agent, updateBuilder.Combine(updates));
+			return await TryUpdateAsync(agent, updateBuilder.Combine(updates));
 		}
 
 		/// <inheritdoc/>
@@ -512,7 +482,7 @@ namespace Horde.Server.Agents
 				update = update.Set(x => x.Deleted, agent.Deleted);
 			}
 
-			return await TryUpdateAsyncWithRetries(agent, update);
+			return await TryUpdateAsync(agent, update);
 		}
 
 		/// <inheritdoc/>
@@ -528,7 +498,7 @@ namespace Horde.Server.Agents
 			leases.Add(newLease);
 
 			UpdateDefinition<AgentDocument> update = Builders<AgentDocument>.Update.Set(x => x.Leases, leases);
-			return await TryUpdateAsyncWithRetries(agent, update);
+			return await TryUpdateAsync(agent, update);
 		}
 
 		/// <inheritdoc/>
@@ -537,7 +507,7 @@ namespace Horde.Server.Agents
 			AgentDocument agent = (AgentDocument)agentInterface;
 
 			UpdateDefinition<AgentDocument> update = Builders<AgentDocument>.Update.Set(x => x.Leases![leaseIdx].State, LeaseState.Cancelled);
-			IAgent? newAgent = await TryUpdateAsyncWithRetries(agent, update);
+			IAgent? newAgent = await TryUpdateAsync(agent, update);
 			if (newAgent != null)
 			{
 				await PublishUpdateEventAsync(agent.Id);
