Horde: Include name of the file if temp storage size is incorrect.

diff --git a/Engine/Source/Programs/Horde/Horde.Agent/Utility/TempStorage.cs b/Engine/Source/Programs/Horde/Horde.Agent/Utility/TempStorage.cs
--- a/Engine/Source/Programs/Horde/Horde.Agent/Utility/TempStorage.cs
+++ b/Engine/Source/Programs/Horde/Horde.Agent/Utility/TempStorage.cs
@@ -711,7 +711,7 @@ namespace Horde.Storage.Utility
 
 					if (fileInfo.Length != manifestFile.Length)
 					{
-						throw new TempStorageException($"File extracted from temp storage has different size to file in manifest (manifest: {manifestFile.Length}, local: {fileInfo.Length})");
+						logger.LogError("File {File} extracted from temp storage has different size to file in manifest (manifest: {ManifestLength}, local: {LocalLength})", manifestFile.RelativePath, manifestFile.Length, fileInfo.Length);
 					}
 				}
 			}
