Horde: Include additional fields in storage controller responses.

diff --git a/Engine/Source/Programs/Horde/Horde.Server/Storage/StorageController.cs b/Engine/Source/Programs/Horde/Horde.Server/Storage/StorageController.cs
--- a/Engine/Source/Programs/Horde/Horde.Server/Storage/StorageController.cs
+++ b/Engine/Source/Programs/Horde/Horde.Server/Storage/StorageController.cs
@@ -589,13 +589,13 @@ namespace Horde.Server.Storage
 						List<object> directories = new List<object>();
 						foreach ((Utf8String name, DirectoryEntry entry) in directoryNode.NameToDirectory)
 						{
-							directories.Add(new { name = name.ToString(), link = GetNodeLink(linkBase, entry) });
+							directories.Add(new { name = name.ToString(), length = entry.Length, hash = entry.Hash, link = GetNodeLink(linkBase, entry) });
 						}
 
 						List<object> files = new List<object>();
 						foreach ((Utf8String name, FileEntry entry) in directoryNode.NameToFile)
 						{
-							files.Add(new { name = name.ToString(), link = GetNodeLink(linkBase, entry) });
+							files.Add(new { name = name.ToString(), length = entry.Length, flags = entry.Flags, hash = entry.Hash, link = GetNodeLink(linkBase, entry) });
 						}
 
 						content = new { directoryNode.Length, directories, files };
