Horde: Exit the update session loop immediately if a new lease has been assigned to an agent externally.

diff --git a/Engine/Source/Programs/Horde/Horde.Server/Agents/AgentService.cs b/Engine/Source/Programs/Horde/Horde.Server/Agents/AgentService.cs
--- a/Engine/Source/Programs/Horde/Horde.Server/Agents/AgentService.cs
+++ b/Engine/Source/Programs/Horde/Horde.Server/Agents/AgentService.cs
@@ -405,11 +405,13 @@ namespace Horde.Server.Agents
 		/// Waits for a lease to be assigned to an agent
 		/// </summary>
 		/// <param name="agent">The agent to assign a lease to</param>
+		/// <param name="newLeases">Leases that the agent knows about</param>
 		/// <param name="cancellationToken"></param>
 		/// <returns>True if a lease was assigned, false otherwise</returns>
-		public async Task<IAgent?> WaitForLeaseAsync(IAgent? agent, CancellationToken cancellationToken)
+		public async Task<IAgent?> WaitForLeaseAsync(IAgent? agent, IList<HordeCommon.Rpc.Messages.Lease> newLeases, CancellationToken cancellationToken)
 		{
-			while (agent != null)
+			HashSet<string> knownLeases = new HashSet<string>(newLeases.Select(x => x.Id), StringComparer.OrdinalIgnoreCase);
+			while (agent != null && agent.Leases.All(x => knownLeases.Contains(x.Id.ToString())))
 			{
 				if (!agent.SessionExpiresAt.HasValue)
 				{
@@ -605,7 +607,7 @@ namespace Horde.Server.Agents
 			agent = await UpdateSessionAsync(agent, sessionId, status, properties, resources, newLeases);
 			if (agent != null && agent.UpdateIndex == updateIndex && (agent.Leases.Count > 0 || agent.Status != AgentStatus.Stopping))
 			{
-				agent = await WaitForLeaseAsync(agent, cancellationToken);
+				agent = await WaitForLeaseAsync(agent, newLeases, cancellationToken);
 			}
 			return agent;
 		}
