Horde: Add command line arguments to support running compute sample on a live Horde instance.

diff --git a/Engine/Source/Programs/Horde/Horde.Server/Startup.cs b/Engine/Source/Programs/Horde/Horde.Server/Startup.cs
--- a/Engine/Source/Programs/Horde/Horde.Server/Startup.cs
+++ b/Engine/Source/Programs/Horde/Horde.Server/Startup.cs
@@ -647,14 +647,15 @@ namespace Horde.Server
 			// Hosted service that needs to run no matter the run mode of the process (server vs worker)
 			services.AddHostedService(provider => (DowntimeService)provider.GetRequiredService<IDowntimeService>());
 
+			// Always run agent service too; need to be able to listen to Redis for events on any server.
+			services.AddHostedService(provider => provider.GetRequiredService<AgentService>());
+
 			if (settings.IsRunModeActive(RunMode.Worker) && !settings.DatabaseReadOnlyMode)
 			{
 				services.AddHostedService<MongoUpgradeService>();
 
 				services.AddHostedService(provider => provider.GetRequiredService<FleetService>());
 				
-				services.AddHostedService(provider => provider.GetRequiredService<AgentService>());
-
 				if (settings.Commits.ReplicateMetadata)
 				{
 					services.AddHostedService(provider => provider.GetRequiredService<PerforceServiceCache>());
diff --git a/Engine/Source/Programs/Horde/Samples/RemoteClient/RemoteClientApp.cs b/Engine/Source/Programs/Horde/Samples/RemoteClient/RemoteClientApp.cs
--- a/Engine/Source/Programs/Horde/Samples/RemoteClient/RemoteClientApp.cs
+++ b/Engine/Source/Programs/Horde/Samples/RemoteClient/RemoteClientApp.cs
@@ -1,7 +1,9 @@
 // Copyright Epic Games, Inc. All Rights Reserved.
 
+using System.Net.Http.Headers;
 using System.Reflection;
 using EpicGames.Core;
+using EpicGames.Horde.Common;
 using EpicGames.Horde.Compute;
 using EpicGames.Horde.Compute.Clients;
 using EpicGames.Horde.Storage;
@@ -13,25 +15,45 @@ using Microsoft.Extensions.Logging;
 
 namespace RemoteClient
 {
-	class ClientApp
+	class ClientAppOptions
 	{
-		static FileReference CurrentAssembly { get; } = new FileReference(Assembly.GetExecutingAssembly().Location);
-		static DirectoryReference AssemblyDir { get; } = CurrentAssembly.Directory;
-		static DirectoryReference ProjectDir { get; } = DirectoryReference.Combine(AssemblyDir, "../../..");
-		static string ConfigPath { get; } = AssemblyDir.MakeRelativeTo(ProjectDir);
+		[CommandLine("-Server=")]
+		public string? Server { get; set; }
+
+		[CommandLine("-Oidc=")]
+		public string? OidcProvider { get; set; }
+
+		[CommandLine("-Condition=")]
+		public string? Condition { get; set; }
+	}
 
-		static FileReference RemoteServerFile { get; } = FileReference.Combine(ProjectDir, "../RemoteServer", ConfigPath, "RemoteServer.exe");
+	class ClientApp
+	{
+		static FileReference CurrentAssemblyFile { get; } = new FileReference(Assembly.GetExecutingAssembly().Location);
+		static DirectoryReference ClientSourceDir { get; } = DirectoryReference.Combine(CurrentAssemblyFile.Directory, "../../..");
+		static FileReference RemoteServerFile { get; } = FileReference.Combine(ClientSourceDir, "../RemoteServer", CurrentAssemblyFile.Directory.MakeRelativeTo(ClientSourceDir), "RemoteServer.dll");
 
-		static async Task Main()
+		static async Task Main(string[] args)
 		{
-			ILogger logger = Log.Logger;
-			CancellationToken cancellationToken = CancellationToken.None;
+			ILogger logger = new DefaultConsoleLogger(LogLevel.Trace);
+
+			// Parse the command line arguments
+			ClientAppOptions options = new ClientAppOptions();
+			CommandLineArguments arguments = new CommandLineArguments(args);
+			arguments.ApplyTo(options);
+			arguments.CheckAllArgumentsUsed(logger);
 
 			// Create the client to handle our requests
-			await using IComputeClient client = CreateClient(true, logger);
+			await using IComputeClient client = await CreateClientAsync(options, logger);
 
 			// Allocate a worker
-			await using IComputeLease? lease = await client.TryAssignWorkerAsync(new ClusterId("default"), null, cancellationToken);
+			Requirements? requirements = null;
+			if (options.Condition != null)
+			{
+				requirements = new Requirements(Condition.Parse(options.Condition));
+			}
+
+			await using IComputeLease? lease = await client.TryAssignWorkerAsync(new ClusterId("default"), requirements);
 			if (lease == null)
 			{
 				logger.LogInformation("Unable to connect to remote");
@@ -40,63 +62,100 @@ namespace RemoteClient
 
 			// Create a message channel on channel id 0. The Horde Agent always listens on this channel for requests.
 			const int ControlChannelId = 0;
-			await using IComputeMessageChannel channel = await lease.Socket.CreateMessageChannelAsync(ControlChannelId, 4 * 1024 * 1024, logger, cancellationToken);
+			await using IComputeMessageChannel channel = await lease.Socket.CreateMessageChannelAsync(ControlChannelId, 4 * 1024 * 1024, logger);
+			await channel.WaitForAttachAsync();
 
 			// Upload the sandbox
 			MemoryStorageClient storage = new MemoryStorageClient();
 			using (TreeWriter treeWriter = new TreeWriter(storage))
 			{
 				DirectoryNode sandbox = new DirectoryNode();
-				await sandbox.CopyFromDirectoryAsync(RemoteServerFile.Directory.ToDirectoryInfo(), new ChunkingOptions(), treeWriter, null, cancellationToken);
-				NodeHandle handle = await treeWriter.FlushAsync(sandbox, cancellationToken);
-				await channel.UploadFilesAsync("", handle.Locator, storage, cancellationToken);
+				await sandbox.CopyFromDirectoryAsync(RemoteServerFile.Directory.ToDirectoryInfo(), new ChunkingOptions(), treeWriter, null);
+				NodeHandle handle = await treeWriter.FlushAsync(sandbox);
+				await channel.UploadFilesAsync("", handle.Locator, storage);
 			}
 
 			// Run the task remotely in the background and echo the output to the console
-			Task childProcessTask = channel.ExecuteAsync(RemoteServerFile.GetFileName(), new List<string>(), null, null, (string x) => logger.LogInformation("ChildProcess: {Message}", x), cancellationToken);
+			await using (IComputeProcess process = await channel.ExecuteAsync(@"C:\Program Files\dotnet\dotnet.exe", new List<string> { RemoteServerFile.GetFileName() }, null, null))
+			{
+				string? line = await process.ReadLineAsync();
+				logger.LogInformation("[REMOTE] {Line}", line);
+
+				await using BackgroundTask tickTask = BackgroundTask.StartNew(ctx => WriteNumbersAsync(lease.Socket, logger, ctx));
 
+				while ((line = await process.ReadLineAsync()) != null)
+				{
+					logger.LogInformation("[REMOTE] {Line}", line);
+				}
+			}
+		}
+
+		static async Task WriteNumbersAsync(IComputeSocket socket, ILogger logger, CancellationToken cancellationToken)
+		{
 			// Generate data into a buffer attached to channel 1. The remote server will echo them back to us as it receives them, then exit when the channel is complete/closed.
 			const int DataChannelId = 1;
-			using (IComputeBufferWriter writer = await lease.Socket.AttachSendBufferAsync(DataChannelId, 1024 * 1024, cancellationToken))
+			using (IComputeBufferWriter writer = await socket.AttachSendBufferAsync(DataChannelId, 1024 * 1024, cancellationToken))
 			{
 				for (int idx = 0; idx < 100; idx++)
 				{
+					cancellationToken.ThrowIfCancellationRequested();
 					logger.LogInformation("Writing value: {Value}", idx);
 					writer.GetMemory().Span[0] = (byte)idx;
 					writer.Advance(1);
 					await Task.Delay(1000, cancellationToken);
 				}
 			}
-
-			// Wait for the child process to finish
-			await childProcessTask;
 		}
 
-		static IComputeClient CreateClient(bool loopback, ILogger logger)
+		static async Task<IComputeClient> CreateClientAsync(ClientAppOptions options, ILogger logger)
 		{
-			if (loopback)
+			if (options.Server == null)
 			{
 				DirectoryReference sandboxDir = DirectoryReference.Combine(new FileReference(Assembly.GetExecutingAssembly().Location).Directory, "Sandbox");
 				return new LocalComputeClient(2000, sandboxDir, logger);
 			}
 			else
 			{
-				return new ServerComputeClient(new Uri("https://localhost:5001"), null, logger);
+				AuthenticationHeaderValue? authHeader = await GetAuthHeaderAsync(options, logger);
+				return new ServerComputeClient(new Uri(options.Server), authHeader, logger);
 			}
 		}
 
-		static async Task<string> GetBearerToken(DirectoryReference engineDir, DirectoryReference? projectDir, string providerIdentifier)
+		static async Task<AuthenticationHeaderValue?> GetAuthHeaderAsync(ClientAppOptions options, ILogger logger)
 		{
+			if (options.OidcProvider == null)
+			{
+				return null;
+			}
+
+			for (DirectoryReference? currentDir = CurrentAssemblyFile.Directory; currentDir != null; currentDir = currentDir.ParentDirectory)
+			{
+				FileReference buildVersionFile = FileReference.Combine(currentDir, "Build/Build.version");
+				if (FileReference.Exists(buildVersionFile))
+				{
+					string bearerToken = await GetOidcBearerTokenAsync(currentDir, null, options.OidcProvider, logger);
+					return new AuthenticationHeaderValue("Bearer", bearerToken);
+				}
+			}
+
+			throw new Exception($"Unable to find engine directory above {CurrentAssemblyFile}");
+		}
+
+		static async Task<string> GetOidcBearerTokenAsync(DirectoryReference engineDir, DirectoryReference? projectDir, string oidcProvider, ILogger logger)
+		{
+			logger.LogInformation("Performing OIDC token refresh...");
+
 			using ITokenStore tokenStore = TokenStoreFactory.CreateTokenStore();
 			IConfiguration providerConfiguration = ProviderConfigurationFactory.ReadConfiguration(engineDir.ToDirectoryInfo(), projectDir?.ToDirectoryInfo());
-			OidcTokenManager oidcTokenManager = OidcTokenManager.CreateTokenManager(providerConfiguration, tokenStore, new List<string>() { providerIdentifier });
-			OidcTokenInfo result = await oidcTokenManager.Login(providerIdentifier);
+			OidcTokenManager oidcTokenManager = OidcTokenManager.CreateTokenManager(providerConfiguration, tokenStore, new List<string>() { oidcProvider });
+			OidcTokenInfo result = await oidcTokenManager.GetAccessToken(oidcProvider);
 
 			if (result.AccessToken == null)
 			{
-				throw new Exception($"Unable to get access token for {providerIdentifier}");
+				throw new Exception($"Unable to get access token for {oidcProvider}");
 			}
 
+			logger.LogInformation("Received bearer token for {OidcProvider}", oidcProvider);
 			return result.AccessToken;
 		}
 	}
diff --git a/Engine/Source/Programs/Horde/Samples/RemoteServer/RemoteServerApp.cs b/Engine/Source/Programs/Horde/Samples/RemoteServer/RemoteServerApp.cs
--- a/Engine/Source/Programs/Horde/Samples/RemoteServer/RemoteServerApp.cs
+++ b/Engine/Source/Programs/Horde/Samples/RemoteServer/RemoteServerApp.cs
@@ -9,23 +9,19 @@ namespace RemoteServer
 	{
 		static async Task Main()
 		{
-			CancellationToken cancellationToken = CancellationToken.None;
-
 			await using IComputeSocket socket = ComputeSocket.ConnectAsWorker(NullLogger.Instance);
-			Console.WriteLine("Connected to client");
+			using IComputeBufferReader reader = await socket.AttachRecvBufferAsync(1, 1024 * 1024);
+			Console.WriteLine("Connected to client"); // Client will wait for output before sending data on channel 1
 
-			using IComputeBufferReader reader = await socket.AttachRecvBufferAsync(1, 1024 * 1024, cancellationToken);
 			while (!reader.IsComplete)
 			{
 				ReadOnlyMemory<byte> memory = reader.GetMemory();
-				if (memory.Length == 0)
+				if (memory.Length > 0)
 				{
-					await reader.WaitForDataAsync(0, cancellationToken);
-					continue;
+					Console.WriteLine("Read value {0}", memory.Span[0]);
+					reader.Advance(1);
 				}
-
-				Console.WriteLine("Received data: {0}", memory.Span[0]);
-				reader.Advance(1);
+				await reader.WaitForDataAsync(0);
 			}
 		}
 	}
diff --git a/Engine/Source/Programs/Shared/EpicGames.Horde/Compute/ComputeWorker.cs b/Engine/Source/Programs/Shared/EpicGames.Horde/Compute/ComputeWorker.cs
--- a/Engine/Source/Programs/Shared/EpicGames.Horde/Compute/ComputeWorker.cs
+++ b/Engine/Source/Programs/Shared/EpicGames.Horde/Compute/ComputeWorker.cs
@@ -183,7 +183,7 @@ namespace EpicGames.Horde.Compute
 
 				newEnvVars[ComputeSocket.WorkerIpcEnvVar] = ipcChannel.GetStringHandle();
 
-				_logger.LogInformation("Launching {Executable} {Arguments}...", CommandLineArguments.Quote(executable), arguments);
+				_logger.LogInformation("Launching {Executable} {Arguments}", CommandLineArguments.Quote(executable), CommandLineArguments.Join(arguments));
 				await ExecuteProcessInternalAsync(channel, executable, arguments, workingDir, newEnvVars, cancellationToken);
 				_logger.LogInformation("Finished executing process");
 
