Horde: Prevent exceptions parsing trace data from failing the executing job.

diff --git a/Engine/Source/Programs/Horde/Horde.Agent/Execution/JobExecutor.cs b/Engine/Source/Programs/Horde/Horde.Agent/Execution/JobExecutor.cs
--- a/Engine/Source/Programs/Horde/Horde.Agent/Execution/JobExecutor.cs
+++ b/Engine/Source/Programs/Horde/Horde.Agent/Execution/JobExecutor.cs
@@ -1258,96 +1258,13 @@ namespace Horde.Agent.Execution
 			if (DirectoryReference.Exists(telemetryDir))
 			{
 				List<(string, FileReference)> telemetryFiles = new List<(string, FileReference)>();
-
-				List<TraceEventList> telemetryList = new List<TraceEventList>();
-				foreach (FileReference telemetryFile in DirectoryReference.EnumerateFiles(telemetryDir, "*.json"))
+				try
 				{
-					jobLogger.LogInformation("Reading telemetry from {File}", telemetryFile);
-					byte[] data = await FileReference.ReadAllBytesAsync(telemetryFile, cancellationToken);
-
-					TraceEventList telemetry = JsonSerializer.Deserialize<TraceEventList>(data.AsSpan())!;
-					if (telemetry.Spans.Count > 0)
-					{
-						string defaultServiceName = telemetryFile.GetFileNameWithoutAnyExtensions();
-						foreach (TraceEvent span in telemetry.Spans)
-						{
-							span.Service ??= defaultServiceName;
-						}
-						telemetryList.Add(telemetry);
-					}
-
-					telemetryFiles.Add(($"Telemetry/{telemetryFile.GetFileName()}", telemetryFile));
+					await CreateTelemetryFilesAsync(step.Name, telemetryDir, telemetryFiles, jobLogger, cancellationToken);
 				}
-
-				List<TraceEvent> telemetrySpans = new List<TraceEvent>();
-				foreach (TraceEventList telemetry in telemetryList.OrderBy(x => x.Spans.First().StartTime).ThenBy(x => x.Spans.Last().FinishTime))
+				catch (Exception ex)
 				{
-					foreach (TraceEvent span in telemetry.Spans)
-					{
-						if (span.FinishTime - span.StartTime > TimeSpan.FromMilliseconds(1.0))
-						{
-							span.Index = telemetrySpans.Count;
-							telemetrySpans.Add(span);
-						}
-					}
-				}
-
-				FileReference? traceFile = null;
-				if (telemetrySpans.Count > 0)
-				{
-					TraceSpan rootSpan = new TraceSpan();
-					rootSpan.Name = step.Name;
-
-					Stack<TraceSpan> stack = new Stack<TraceSpan>();
-					stack.Push(rootSpan);
-
-					foreach (TraceEvent traceEvent in telemetrySpans.OrderBy(x => x.StartTime).ThenByDescending(x => x.FinishTime).ThenBy(x => x.Index))
-					{
-						TraceSpan newSpan = new TraceSpan();
-						newSpan.Name = traceEvent.Name;
-						newSpan.Service = traceEvent.Service;
-						newSpan.Resource = traceEvent.Resource;
-						newSpan.Start = traceEvent.StartTime.UtcTicks;
-						newSpan.Finish = traceEvent.FinishTime.UtcTicks;
-						if (traceEvent.Metadata != null && traceEvent.Metadata.Count > 0)
-						{
-							newSpan.Properties = traceEvent.Metadata;
-						}
-
-						TraceSpan stackTop = stack.Peek();
-						while (stack.Count > 1 && newSpan.Start >= stackTop.Finish)
-						{
-							stack.Pop();
-							stackTop = stack.Peek();
-						}
-
-						if (stack.Count > 1 && newSpan.Finish > stackTop.Finish)
-						{
-							jobLogger.LogInformation("Trace event name='{Name}', service'{Service}', resource='{Resource}' has invalid finish time ({SpanFinish} < {StackFinish})", newSpan.Name, newSpan.Service, newSpan.Resource, newSpan.Finish, stackTop.Finish);
-							newSpan.Finish = stackTop.Finish;
-						}
-
-						if (stackTop.Children == null)
-						{
-							stackTop.Children = new List<TraceSpan>();
-						}
-
-						stackTop.Children.Add(newSpan);
-						stack.Push(newSpan);
-					}
-
-					rootSpan.Start = rootSpan.Children!.First().Start;
-					rootSpan.Finish = rootSpan.Children!.Last().Finish;
-
-					traceFile = FileReference.Combine(telemetryDir, "Trace.json");
-					using (FileStream stream = FileReference.Open(traceFile, FileMode.Create))
-					{
-						JsonSerializerOptions options = new JsonSerializerOptions { DefaultIgnoreCondition = JsonIgnoreCondition.WhenWritingNull };
-						await JsonSerializer.SerializeAsync(stream, rootSpan, options, cancellationToken);
-					}
-					telemetryFiles.Add(("Trace.json", traceFile));
-
-					CreateTracingData(GlobalTracer.Instance.ActiveSpan, rootSpan);
+					jobLogger.LogWarning(ex, "Unable to parse/upload telemetry files: {Message}", ex.Message);
 				}
 
 				await CreateArtifactsAsync(step.StepId, JobArtifactType.Trace, workspaceDir, telemetryFiles, jobLogger, cancellationToken);
@@ -1433,6 +1350,100 @@ namespace Horde.Agent.Execution
 			return exitCode;
 		}
 
+		async Task CreateTelemetryFilesAsync(string stepName, DirectoryReference telemetryDir, List<(string, FileReference)> telemetryFiles, ILogger jobLogger, CancellationToken cancellationToken)
+		{
+			List<TraceEventList> telemetryList = new List<TraceEventList>();
+			foreach (FileReference telemetryFile in DirectoryReference.EnumerateFiles(telemetryDir, "*.json"))
+			{
+				jobLogger.LogInformation("Reading telemetry from {File}", telemetryFile);
+				byte[] data = await FileReference.ReadAllBytesAsync(telemetryFile, cancellationToken);
+
+				TraceEventList telemetry = JsonSerializer.Deserialize<TraceEventList>(data.AsSpan())!;
+				if (telemetry.Spans.Count > 0)
+				{
+					string defaultServiceName = telemetryFile.GetFileNameWithoutAnyExtensions();
+					foreach (TraceEvent span in telemetry.Spans)
+					{
+						span.Service ??= defaultServiceName;
+					}
+					telemetryList.Add(telemetry);
+				}
+
+				telemetryFiles.Add(($"Telemetry/{telemetryFile.GetFileName()}", telemetryFile));
+			}
+
+			List<TraceEvent> telemetrySpans = new List<TraceEvent>();
+			foreach (TraceEventList telemetry in telemetryList.OrderBy(x => x.Spans.First().StartTime).ThenBy(x => x.Spans.Last().FinishTime))
+			{
+				foreach (TraceEvent span in telemetry.Spans)
+				{
+					if (span.FinishTime - span.StartTime > TimeSpan.FromMilliseconds(1.0))
+					{
+						span.Index = telemetrySpans.Count;
+						telemetrySpans.Add(span);
+					}
+				}
+			}
+
+			FileReference? traceFile = null;
+			if (telemetrySpans.Count > 0)
+			{
+				TraceSpan rootSpan = new TraceSpan();
+				rootSpan.Name = stepName;
+
+				Stack<TraceSpan> stack = new Stack<TraceSpan>();
+				stack.Push(rootSpan);
+
+				foreach (TraceEvent traceEvent in telemetrySpans.OrderBy(x => x.StartTime).ThenByDescending(x => x.FinishTime).ThenBy(x => x.Index))
+				{
+					TraceSpan newSpan = new TraceSpan();
+					newSpan.Name = traceEvent.Name;
+					newSpan.Service = traceEvent.Service;
+					newSpan.Resource = traceEvent.Resource;
+					newSpan.Start = traceEvent.StartTime.UtcTicks;
+					newSpan.Finish = traceEvent.FinishTime.UtcTicks;
+					if (traceEvent.Metadata != null && traceEvent.Metadata.Count > 0)
+					{
+						newSpan.Properties = traceEvent.Metadata;
+					}
+
+					TraceSpan stackTop = stack.Peek();
+					while (stack.Count > 1 && newSpan.Start >= stackTop.Finish)
+					{
+						stack.Pop();
+						stackTop = stack.Peek();
+					}
+
+					if (stack.Count > 1 && newSpan.Finish > stackTop.Finish)
+					{
+						jobLogger.LogInformation("Trace event name='{Name}', service'{Service}', resource='{Resource}' has invalid finish time ({SpanFinish} < {StackFinish})", newSpan.Name, newSpan.Service, newSpan.Resource, newSpan.Finish, stackTop.Finish);
+						newSpan.Finish = stackTop.Finish;
+					}
+
+					if (stackTop.Children == null)
+					{
+						stackTop.Children = new List<TraceSpan>();
+					}
+
+					stackTop.Children.Add(newSpan);
+					stack.Push(newSpan);
+				}
+
+				rootSpan.Start = rootSpan.Children!.First().Start;
+				rootSpan.Finish = rootSpan.Children!.Last().Finish;
+
+				traceFile = FileReference.Combine(telemetryDir, "Trace.json");
+				using (FileStream stream = FileReference.Open(traceFile, FileMode.Create))
+				{
+					JsonSerializerOptions options = new JsonSerializerOptions { DefaultIgnoreCondition = JsonIgnoreCondition.WhenWritingNull };
+					await JsonSerializer.SerializeAsync(stream, rootSpan, options, cancellationToken);
+				}
+				telemetryFiles.Add(("Trace.json", traceFile));
+
+				CreateTracingData(GlobalTracer.Instance.ActiveSpan, rootSpan);
+			}
+		}
+
 		private async Task CreateReportAsync(string stepId, FileReference reportFile, Dictionary<FileReference, string> artifactFileToId, ILogger logger)
 		{
 			byte[] data = await FileReference.ReadAllBytesAsync(reportFile);
