Horde: Use correct mime type for artifact downloads.

diff --git a/Engine/Source/Programs/Horde/Horde.Server/Artifacts/ArtifactsController.cs b/Engine/Source/Programs/Horde/Horde.Server/Artifacts/ArtifactsController.cs
--- a/Engine/Source/Programs/Horde/Horde.Server/Artifacts/ArtifactsController.cs
+++ b/Engine/Source/Programs/Horde/Horde.Server/Artifacts/ArtifactsController.cs
@@ -17,6 +17,7 @@ using Horde.Server.Utilities;
 using Microsoft.AspNetCore.Authorization;
 using Microsoft.AspNetCore.Http;
 using Microsoft.AspNetCore.Mvc;
+using Microsoft.AspNetCore.StaticFiles;
 using Microsoft.Extensions.Options;
 
 namespace Horde.Server.Artifacts
@@ -257,12 +258,12 @@ namespace Horde.Server.Artifacts
 		/// </summary>
 		/// <param name="id">Identifier of the artifact to retrieve</param>
 		/// <param name="path">Path to fetch</param>
-		/// <param name="download">Whether to request the file be downloaded vs displayed inline</param>
+		/// <param name="inline">Whether to request the file be downloaded vs displayed inline</param>
 		/// <param name="cancellationToken">Cancellation token for the operation</param>
 		/// <returns>Information about all the artifacts</returns>
 		[HttpGet]
 		[Route("/api/v2/artifacts/{id}/file")]
-		public async Task<ActionResult<object>> GetFileAsync(ArtifactId id, [FromQuery] string path, [FromQuery] bool download = false, CancellationToken cancellationToken = default)
+		public async Task<ActionResult<object>> GetFileAsync(ArtifactId id, [FromQuery] string path, [FromQuery] bool inline = false, CancellationToken cancellationToken = default)
 		{
 			IArtifact? artifact = await _artifactCollection.GetAsync(id, cancellationToken);
 			if (artifact == null)
@@ -284,14 +285,20 @@ namespace Horde.Server.Artifacts
 				return NotFound($"Unable to find file {path}");
 			}
 
+			string? contentType;
+			if (!new FileExtensionContentTypeProvider().TryGetContentType(path, out contentType))
+			{
+				contentType = "application/octet-stream";
+			}
+
 			Stream stream = fileEntry.AsStream(reader);
-			if (download)
+			if (inline)
 			{
-				return new FileStreamResult(stream, "application/octet-stream") { FileDownloadName = Path.GetFileName(path) };
+				return new InlineFileStreamResult(stream, contentType, Path.GetFileName(path));
 			}
 			else
 			{
-				return new InlineFileStreamResult(stream, "application/octet-stream", Path.GetFileName(path));
+				return new FileStreamResult(stream, contentType) { FileDownloadName = Path.GetFileName(path) };
 			}
 		}
 
@@ -320,7 +327,7 @@ namespace Horde.Server.Artifacts
 			DirectoryNode directory = await directoryRef.ExpandAsync(reader, cancellationToken);
 
 			Stream stream = directory.AsZipStream(reader);
-			return new FileStreamResult(stream, "application/octet-stream") { FileDownloadName = $"{artifact.RefName}.zip" };
+			return new FileStreamResult(stream, "application/zip") { FileDownloadName = $"{artifact.RefName}.zip" };
 		}
 
 		/// <summary>
