Horde: Add conforms to pool view, improve scrolling behavior as well
#skipci
#rnx

diff --git a/Engine/Source/Programs/Horde/HordeDashboard/src/components/PoolView.tsx b/Engine/Source/Programs/Horde/HordeDashboard/src/components/PoolView.tsx
--- a/Engine/Source/Programs/Horde/HordeDashboard/src/components/PoolView.tsx
+++ b/Engine/Source/Programs/Horde/HordeDashboard/src/components/PoolView.tsx
@@ -2,19 +2,20 @@
 
 import { Checkbox, ComboBox, DefaultButton, DetailsList, DetailsListLayoutMode, Dialog, DialogFooter, DialogType, IColumn, IComboBoxOption, ITag, MessageBar, MessageBarType, Position, PrimaryButton, ScrollablePane, ScrollbarVisibility, SelectionMode, Slider, SpinButton, Spinner, SpinnerSize, Stack, TagPicker, Text, TextField } from "@fluentui/react";
 import { observer } from "mobx-react-lite";
+import moment from "moment";
 import React, { useEffect, useState } from "react";
 import { Link, useSearchParams } from "react-router-dom";
 import backend from "../backend";
-import { GetAgentResponse, GetBatchResponse, GetJobTimingResponse, GetPoolResponse, GetStepResponse, JobData, JobQuery, JobState, JobStepBatchState, JobStepState, PoolSizeStrategy, StepData, UpdatePoolRequest } from "../backend/Api";
+import { GetAgentLeaseResponse, GetAgentResponse, GetBatchResponse, GetJobTimingResponse, GetPoolResponse, GetStepResponse, JobData, JobQuery, JobState, JobStepBatchState, JobStepState, LeaseState, PoolSizeStrategy, StepData, UpdatePoolRequest } from "../backend/Api";
 import dashboard from "../backend/Dashboard";
 import { PollBase } from "../backend/PollBase";
 import { projectStore } from "../backend/ProjectStore";
 import { useWindowSize } from "../base/utilities/hooks";
-import { getNiceTime, getStepElapsed, getStepETA, getStepFinishTime, getStepStartTime } from "../base/utilities/timeUtils";
+import { getElapsedString, getNiceTime, getShortNiceTime, getStepElapsed, getStepETA, getStepFinishTime, getStepStartTime } from "../base/utilities/timeUtils";
 import { hordeClasses, linearInterpolate, modeColors } from "../styles/Styles";
 import { Breadcrumbs } from "./Breadcrumbs";
 import { HistoryModal } from "./HistoryModal";
-import { StepStatusIcon } from "./StatusIcon";
+import { LeaseStatusIcon, StepStatusIcon } from "./StatusIcon";
 import { TopNav } from "./TopNav";
 
 
@@ -41,6 +42,7 @@ class PoolHandler extends PollBase {
       this.pendingSteps = new Map();
       this.pendingBatches = [];
       this.jobData = new Map();
+      this.activeConforms = new Map<string, GetAgentLeaseResponse>();
       super.stop();
    }
 
@@ -114,6 +116,17 @@ class PoolHandler extends PollBase {
 
          const minCreateTime = new Date(Math.round((new Date().getTime() / 1000) - (48 * 3600)) * 1000).toISOString();
 
+         this.activeConforms = new Map<string, GetAgentLeaseResponse>();
+
+         agents.forEach(a => {
+            const lease = a.leases?.find(lease => lease.name === "Updating workspaces" && (lease.state === LeaseState.Pending || lease.state === LeaseState.Active))
+            if (lease) {
+               // make sure populated
+               lease.agentId = a.id;
+               this.activeConforms.set(a.id, lease);
+            }
+         })
+
          while (poolStreams.length) {
 
             const batch = poolStreams.slice(0, 5);
@@ -300,6 +313,8 @@ class PoolHandler extends PollBase {
    // agentId => jobId
    agentJobs: Map<string, string> = new Map();
 
+   activeConforms = new Map<string, GetAgentLeaseResponse>();
+
    // agentId => step
    activeSteps: Map<string, GetStepResponse> = new Map();
 
@@ -421,6 +436,9 @@ const StepPanel: React.FC<{ stepState: StepState }> = ({ stepState }) => {
       const node = group.nodes[step.nodeIdx];
 
       if (column.name === "Time Active") {
+         if (stepState === StepState.Pending) {
+            return null;
+         }
          return <Stack horizontalAlign={"end"}><Text>{getStepElapsed(step)}</Text></Stack>;
       }
 
@@ -665,7 +683,130 @@ const BatchPanel: React.FC = () => {
                      columns={columns}
                      setKey="set"
                      layoutMode={DetailsListLayoutMode.justified}
-                     isHeaderVisible={true}
+                     isHeaderVisible={false}
+                     selectionMode={SelectionMode.none}
+                     onRenderItemColumn={onRenderItemColumn}
+                  />
+               </ScrollablePane>
+            </Stack>}
+         </Stack>
+      </Stack>
+   </Stack>)
+};
+
+const ConformPanel: React.FC = () => {
+
+   const pool = handler.pool;
+   if (!pool) {
+      return null;
+   }
+
+   let conforms = Array.from(handler.activeConforms.values());
+
+   conforms = conforms.sort((a, b) => {
+
+      if (a.state === LeaseState.Active && b.state === LeaseState.Pending) {
+         return -1;
+      }
+
+      if (b.state === LeaseState.Active && a.state === LeaseState.Pending) {
+         return 1;
+      }
+
+      return new Date(a.startTime) < new Date(b.startTime) ? -1 : 1;
+
+   });
+
+   type ConformItem = {
+      lease: GetAgentLeaseResponse;
+   }
+
+   const items = conforms.map(c => { return { lease: c } });
+
+   let columns: IColumn[] = [];
+
+   columns = [
+      { key: 'column1', name: 'AgentId', minWidth: 128, maxWidth: 128 },      
+      { key: 'column2', name: 'Type', minWidth: 240, maxWidth: 240 },
+      { key: 'column3', name: 'Duration', minWidth: 128, maxWidth: 128 },
+      { key: 'column4', name: 'Started', minWidth: 128, maxWidth: 128 },
+      { key: 'column5', name: 'ViewLog', minWidth: 140, maxWidth: 140 },
+      { key: 'column6', name: 'Spacer', minWidth: 128, maxWidth: 128 },
+   ];
+
+   const onRenderItemColumn = (item: ConformItem, index?: number, columnIn?: IColumn) => {
+
+      const column = columnIn!;
+      const lease = item.lease;
+
+      if (column.name === "AgentId") {
+         return <Stack horizontal verticalAlign="center">
+            <LeaseStatusIcon lease={lease} />
+            <Text style={{ fontSize: "13px", paddingRight: 12 }}>{lease.agentId}</Text>
+         </Stack>;
+      }
+
+      if (column.name === "Type") {         
+         return <Stack ><Text style={{ fontSize: "13px", paddingRight: 12 }}>{lease.name ?? "Unknown lease type"}</Text></Stack>;
+      }
+
+      if (column.name === "Duration") {
+
+
+         return <Stack horizontalAlign="end"><Text style={{ fontSize: "13px", paddingRight: 12 }}>{getElapsedString(moment(lease.startTime), moment.utc(), true)}</Text></Stack>;
+      }
+
+      if (column.name === "Started") {
+
+         return <Stack horizontalAlign="end"><Text style={{ fontSize: "13px", paddingRight: 12 }}>{getShortNiceTime(lease.startTime, false, true, false)}</Text></Stack>;
+      }
+
+      if (column.name === "ViewLog") {
+         if (!lease.logId) {
+            return null;
+         }
+         return <Stack horizontalAlign="center" style={{ paddingRight: 32 }}>
+            <Link to={`/log/${lease.logId}?leaseId=${lease.id}`} target="_blank">
+               <Stack horizontal verticalAlign="center" tokens={{ childrenGap: 0, padding: 0 }} style={{ width: "100%", height: "100%" }}>
+                  <Text styles={{ root: { margin: '0px', padding: '0px', paddingRight: '8px' } }} className={"view-log-link"}>View Log</Text>
+               </Stack>
+            </Link>
+         </Stack>
+      }
+
+
+      return null;
+   }
+
+   const count = Math.min(items.length, 8);
+   let height = 60 + count * 32;
+
+   return (<Stack>
+      <Stack styles={{ root: { paddingTop: 18, paddingLeft: 12, paddingRight: 12 } }} >
+         <Stack tokens={{ childrenGap: 12 }}>
+            <Stack>
+               <Text variant="mediumPlus" styles={{ root: { fontFamily: "Horde Open Sans SemiBold" } }}>Active Conforms ({items.length})</Text>
+            </Stack>
+            {!!items.length && <Stack style={{ position: "relative", height: height, width: 1280 }}>
+               <ScrollablePane scrollbarVisibility={ScrollbarVisibility.always}>
+                  <DetailsList
+                     styles={{
+                        root: {
+                           selectors: {
+                              '.ms-List-cell': {
+                                 color: modeColors.text,
+                                 minHeight: 22,
+                                 fontSize: 12
+                              }
+                           }
+                        }
+                     }}
+                     compact={true}
+                     items={items}
+                     columns={columns}
+                     setKey="set"
+                     layoutMode={DetailsListLayoutMode.justified}
+                     isHeaderVisible={false}
                      selectionMode={SelectionMode.none}
                      onRenderItemColumn={onRenderItemColumn}
                   />
@@ -676,6 +817,7 @@ const BatchPanel: React.FC = () => {
    </Stack>)
 };
 
+
 const AutoScalerPanel: React.FC = () => {
 
    const [state, setState] = useState<{ error?: string, submitting?: boolean, confirmed?: boolean, modified?: boolean, name?: string, color?: string, autoscale?: boolean, minAgents?: number, reserveAgents?: number, strategy?: PoolSizeStrategy, conformInterval?: number }>({});
@@ -850,7 +992,7 @@ const AutoScalerPanel: React.FC = () => {
          </Stack>
       </Stack>
    </Stack>
-   
+
 }
 
 
@@ -1071,52 +1213,66 @@ export const PoolView: React.FC = observer(() => {
    return <Stack className={hordeClasses.horde}>
       <TopNav />
       <Breadcrumbs items={[{ text: "Agents", link: "/agents" }, { text: breadText }]} />
-      <Stack horizontal>
-         <div key={`windowsize_poolview_${windowSize.width}_${windowSize.height}`} style={{ width: vw / 2 - (1440 / 2), flexShrink: 0, backgroundColor: modeColors.background }} />
-         <Stack tokens={{ childrenGap: 0 }} styles={{ root: { backgroundColor: modeColors.background, width: "100%" } }}>
-            <Stack style={{ maxWidth: 1440, paddingTop: 6, marginLeft: 4, height: 'calc(100vh - 8px)' }}>
-               <Stack horizontal style={{ paddingTop: 8, paddingBottom: 16 }}>
-                  <Stack grow />
-                  <Stack style={{ width: 320 }}>
-                     <TagPicker inputProps={{ placeholder: "Select a pool" }}
-                        defaultSelectedItems={defaultItems}
-                        onResolveSuggestions={filterSuggestedTags}
-                        getTextFromItem={getTextFromItem}
-                        onEmptyResolveSuggestions={(selected) => {
-                           return poolTags;
-                        }}
-
-                        onItemSelected={(item) => {
-
-                           if (!item?.key) {
-                              return null;
-                           }
+      <Stack>
+         <Stack horizontal style={{ flexBasis: 64 }}>
+            <div key={`windowsize_poolview_${windowSize.width}_${windowSize.height}`} style={{ width: vw / 2 - (1440 / 2), flexShrink: 0, backgroundColor: modeColors.background }} />
+            <Stack tokens={{ childrenGap: 0 }} styles={{ root: { backgroundColor: modeColors.background, width: "100%" } }}>
+               <Stack style={{ maxWidth: 1440, paddingTop: 6, marginLeft: 4 }}>
+                  <Stack horizontal style={{ paddingTop: 8, paddingBottom: 16 }}>
+                     <Stack grow />
+                     <Stack style={{ width: 320 }}>
+                        <TagPicker inputProps={{ placeholder: "Select a pool" }}
+                           defaultSelectedItems={defaultItems}
+                           onResolveSuggestions={filterSuggestedTags}
+                           getTextFromItem={getTextFromItem}
+                           onEmptyResolveSuggestions={(selected) => {
+                              return poolTags;
+                           }}
+
+                           onItemSelected={(item) => {
+
+                              if (!item?.key) {
+                                 return null;
+                              }
 
-                           setSearchParams(`?pool=${item.key}`);
+                              setSearchParams(`?pool=${item.key}`, {replace: true});
 
-                           return item;
+                              return item;
 
-                        }}
+                           }}
 
-                        itemLimit={1} />
+                           itemLimit={1} />
+                     </Stack>
                   </Stack>
                </Stack>
-               {!!pool && <Stack className={hordeClasses.raised}>
-                  <Stack style={{ position: "relative", width: "1390px", height: 'calc(100vh - 228px)' }}>
-                     <ScrollablePane scrollbarVisibility={ScrollbarVisibility.always}>
-                        <Stack tokens={{ childrenGap: 18 }}>
-                           <PoolPanel />
-                           <BatchPanel />
-                           <StepPanel stepState={StepState.Active} />
-                           <StepPanel stepState={StepState.Pending} />
+            </Stack>
+         </Stack>
+         <Stack style={{ height: "100%", backgroundColor: modeColors.background }}>
+            <Stack horizontal style={{ height: "100%" }}>
+               <div key={`windowsize_poolview_2_${windowSize.width}_${windowSize.height}`} style={{ width: vw / 2 - (1440 / 2), flexShrink: 0, backgroundColor: modeColors.background }} />
+               {!!pool && <Stack style={{ width: "100%", height: "100%" }}>
+                  <div style={{ marginTop: 8, width: "100%", height: 'calc(100vh - 258px)', position: 'relative' }} data-is-scrollable={true}>
+                     <ScrollablePane scrollbarVisibility={ScrollbarVisibility.always} onScroll={() => { }}>
+
+                        <Stack className={hordeClasses.raised} style={{ width: 1443, height: "fit-content" }}>
+                           <Stack tokens={{ childrenGap: 18 }}>
+                              <PoolPanel />
+                              <StepPanel stepState={StepState.Active} />
+                              <ConformPanel />
+                              <StepPanel stepState={StepState.Pending} />
+                              <BatchPanel />                              
+                           </Stack>
                         </Stack>
                      </ScrollablePane>
-                  </Stack>
+                  </div>
+               </Stack>}
+               {!pool && !!poolId && <Stack horizontalAlign="center" style={{ width: "100%" }}>
+                  <Spinner size={SpinnerSize.large} />
                </Stack>}
-               {!pool && !!poolId && <Spinner size={SpinnerSize.large} />}
             </Stack>
          </Stack>
       </Stack>
    </Stack>
+
 });
 
diff --git a/Engine/Source/Programs/Horde/HordeDashboard/src/components/StatusIcon.tsx b/Engine/Source/Programs/Horde/HordeDashboard/src/components/StatusIcon.tsx
--- a/Engine/Source/Programs/Horde/HordeDashboard/src/components/StatusIcon.tsx
+++ b/Engine/Source/Programs/Horde/HordeDashboard/src/components/StatusIcon.tsx
@@ -4,7 +4,7 @@
 
 import { FontIcon, Stack } from '@fluentui/react';
 import React, { CSSProperties } from 'react';
-import { FindIssueResponse, GetAgentLeaseResponse, GetBatchResponse, GetIssueResponse, GetJobStepRefResponse, GetStepResponse, IssueSeverity, JobStepBatchError, JobStepOutcome, JobStepState, LabelOutcome, LabelState, LeaseOutcome } from '../backend/Api';
+import { FindIssueResponse, GetAgentLeaseResponse, GetBatchResponse, GetIssueResponse, GetJobStepRefResponse, GetStepResponse, IssueSeverity, JobStepBatchError, JobStepOutcome, JobStepState, LabelOutcome, LabelState, LeaseOutcome, LeaseState } from '../backend/Api';
 import dashboard, { StatusColor } from "../backend/Dashboard";
 import { JobLabel } from '../backend/JobDetails';
 
@@ -210,7 +210,7 @@ export const LeaseStatusIcon: React.FC<{ lease: GetAgentLeaseResponse, style?: C
    const icon = "Square";
    let color: string | undefined;
 
-   if (lease.executing) {
+   if (lease.executing || (lease.name === "Updating workspaces" && (lease.state === LeaseState.Active || lease.state === LeaseState.Pending))) {
       color = colors.get(StatusColor.Running);
    }
 
