Horde: Make it so you can open history modal on conform
#rnx
#skipci

diff --git a/Engine/Source/Programs/Horde/HordeDashboard/src/components/PoolView.tsx b/Engine/Source/Programs/Horde/HordeDashboard/src/components/PoolView.tsx
--- a/Engine/Source/Programs/Horde/HordeDashboard/src/components/PoolView.tsx
+++ b/Engine/Source/Programs/Horde/HordeDashboard/src/components/PoolView.tsx
@@ -696,6 +696,8 @@ const BatchPanel: React.FC = () => {
 
 const ConformPanel: React.FC = () => {
 
+   const [lastSelectedAgent, setLastSelectedAgent] = useState<string | undefined>(undefined);
+
    const pool = handler.pool;
    if (!pool) {
       return null;
@@ -737,13 +739,13 @@ const ConformPanel: React.FC = () => {
    const onRenderItemColumn = (item: ConformItem, index?: number, columnIn?: IColumn) => {
 
       const column = columnIn!;
-      const lease = item.lease;
+      const lease = item.lease;      
 
       if (column.name === "AgentId") {
-         return <Stack horizontal verticalAlign="center">
+         return <div style={{ cursor: "pointer" }} onClick={() => { setLastSelectedAgent(lease.agentId) }}><Stack horizontal verticalAlign="center">
             <LeaseStatusIcon lease={lease} />
             <Text style={{ fontSize: "13px", paddingRight: 12 }}>{lease.agentId}</Text>
-         </Stack>;
+         </Stack></div>;
       }
 
       if (column.name === "Type") {         
@@ -782,6 +784,7 @@ const ConformPanel: React.FC = () => {
    let height = 60 + count * 32;
 
    return (<Stack>
+      <HistoryModal agentId={lastSelectedAgent} onDismiss={() => setLastSelectedAgent(undefined)} />
       <Stack styles={{ root: { paddingTop: 18, paddingLeft: 12, paddingRight: 12 } }} >
          <Stack tokens={{ childrenGap: 12 }}>
             <Stack>
@@ -1266,7 +1269,7 @@ export const PoolView: React.FC = observer(() => {
                      </ScrollablePane>
                   </div>
                </Stack>}
-               {!pool && !!poolId && <Stack horizontalAlign="center" style={{ width: "100%" }}>
+               {!pool && !!poolId && <Stack horizontalAlign="center" style={{ width: 1443 }}>
                   <Spinner size={SpinnerSize.large} />
                </Stack>}
             </Stack>
