Horde: Improve pool view formatting for conforms
#rnx
#skipci
#peflight none

diff --git a/Engine/Source/Programs/Horde/HordeDashboard/src/components/PoolView.tsx b/Engine/Source/Programs/Horde/HordeDashboard/src/components/PoolView.tsx
--- a/Engine/Source/Programs/Horde/HordeDashboard/src/components/PoolView.tsx
+++ b/Engine/Source/Programs/Horde/HordeDashboard/src/components/PoolView.tsx
@@ -397,8 +397,8 @@ const StepPanel: React.FC<{ stepState: StepState }> = ({ stepState }) => {
    let columns: IColumn[] = [];
 
    columns = [
-      { key: 'column1', name: 'Job', minWidth: 600, maxWidth: 600 },
       { key: 'column2', name: 'Agent', minWidth: 100, maxWidth: 100 },
+      { key: 'column1', name: 'Job', minWidth: 600, maxWidth: 600 },
       { key: 'column3', name: 'Time Active', minWidth: 100, maxWidth: 100 },
       { key: 'column4', name: 'Start Time', minWidth: 100, maxWidth: 100 },
       { key: 'column5', name: 'Estimated Finish', minWidth: 100, maxWidth: 100 },
@@ -509,15 +509,20 @@ const StepPanel: React.FC<{ stepState: StepState }> = ({ stepState }) => {
          const stepName = `${jobName} - ${node.name}`;
          const stepUrl = `/job/${job.id}?step=${step.id}`;
 
-         return <Stack horizontal>
-            <StepStatusIcon step={step} />
+         return <Stack>
             <Link target="_blank" to={stepUrl}><Text>{stepName}</Text></Link>
          </Stack>;
 
       }
 
       if (column.name === "Agent") {
-         return <div style={{ cursor: "pointer" }} onClick={() => { setLastSelectedAgent(agent.id) }}><Text>{agent.name}</Text></div>
+         return <div style={{ cursor: "pointer" }} onClick={() => { setLastSelectedAgent(agent.id) }}>
+            <Stack horizontal>
+               <StepStatusIcon step={step} />
+               <Text>{agent.name}</Text>
+            </Stack>
+
+         </div>
 
       }
 
@@ -728,46 +733,45 @@ const ConformPanel: React.FC = () => {
    let columns: IColumn[] = [];
 
    columns = [
-      { key: 'column1', name: 'AgentId', minWidth: 128, maxWidth: 128 },      
-      { key: 'column2', name: 'Type', minWidth: 240, maxWidth: 240 },
-      { key: 'column3', name: 'Duration', minWidth: 128, maxWidth: 128 },
-      { key: 'column4', name: 'Started', minWidth: 128, maxWidth: 128 },
-      { key: 'column5', name: 'ViewLog', minWidth: 140, maxWidth: 140 },
-      { key: 'column6', name: 'Spacer', minWidth: 128, maxWidth: 128 },
+      { key: 'column1', name: 'AgentId', minWidth: 100, maxWidth: 100 },
+      { key: 'column2', name: 'Type', minWidth: 600, maxWidth: 600 },
+      { key: 'column3', name: 'Duration', minWidth: 100, maxWidth: 100 },
+      { key: 'column4', name: 'Started', minWidth: 100, maxWidth: 100 },
+      { key: 'column5', name: 'ViewLog', minWidth: 80, maxWidth: 80 }
    ];
 
    const onRenderItemColumn = (item: ConformItem, index?: number, columnIn?: IColumn) => {
 
       const column = columnIn!;
-      const lease = item.lease;      
+      const lease = item.lease;
 
       if (column.name === "AgentId") {
          return <div style={{ cursor: "pointer" }} onClick={() => { setLastSelectedAgent(lease.agentId) }}><Stack horizontal verticalAlign="center">
             <LeaseStatusIcon lease={lease} />
-            <Text style={{ fontSize: "13px", paddingRight: 12 }}>{lease.agentId}</Text>
+            <Text style={{ fontSize: "13px" }}>{lease.agentId}</Text>
          </Stack></div>;
       }
 
-      if (column.name === "Type") {         
-         return <Stack ><Text style={{ fontSize: "13px", paddingRight: 12 }}>{lease.name ?? "Unknown lease type"}</Text></Stack>;
+      if (column.name === "Type") {
+         return <Stack ><Text style={{ fontSize: "13px" }}>{lease.name ?? "Unknown lease type"}</Text></Stack>;
       }
 
       if (column.name === "Duration") {
 
 
-         return <Stack horizontalAlign="end"><Text style={{ fontSize: "13px", paddingRight: 12 }}>{getElapsedString(moment(lease.startTime), moment.utc(), true)}</Text></Stack>;
+         return <Stack horizontalAlign="end"><Text style={{ fontSize: "13px" }}>{getElapsedString(moment(lease.startTime), moment.utc(), true)}</Text></Stack>;
       }
 
       if (column.name === "Started") {
 
-         return <Stack horizontalAlign="end"><Text style={{ fontSize: "13px", paddingRight: 12 }}>{getShortNiceTime(lease.startTime, false, true, false)}</Text></Stack>;
+         return <Stack horizontalAlign="end"><Text style={{ fontSize: "13px" }}>{getShortNiceTime(lease.startTime, false, true, false)}</Text></Stack>;
       }
 
       if (column.name === "ViewLog") {
          if (!lease.logId) {
             return null;
          }
-         return <Stack horizontalAlign="center" style={{ paddingRight: 32 }}>
+         return <Stack horizontalAlign="end" style={{ paddingRight: 32 }}>
             <Link to={`/log/${lease.logId}?leaseId=${lease.id}`} target="_blank">
                <Stack horizontal verticalAlign="center" tokens={{ childrenGap: 0, padding: 0 }} style={{ width: "100%", height: "100%" }}>
                   <Text styles={{ root: { margin: '0px', padding: '0px', paddingRight: '8px' } }} className={"view-log-link"}>View Log</Text>
@@ -1238,7 +1242,7 @@ export const PoolView: React.FC = observer(() => {
                                  return null;
                               }
 
-                              setSearchParams(`?pool=${item.key}`, {replace: true});
+                              setSearchParams(`?pool=${item.key}`, { replace: true });
 
                               return item;
 
@@ -1263,7 +1267,7 @@ export const PoolView: React.FC = observer(() => {
                               <StepPanel stepState={StepState.Active} />
                               <ConformPanel />
                               <StepPanel stepState={StepState.Pending} />
-                              <BatchPanel />                              
+                              <BatchPanel />
                            </Stack>
                         </Stack>
                      </ScrollablePane>
