Horde: Add a C++ sample of a remote worker. Pass -Cpp as an argument to RemoteClient to use it rather than the C# version.

Various other changes:
* Use a single buffer for IPC requests from the worker to host, rather than using the message passing system. Provides less overhead to porting to other languages.
* Rename RemoteServer to RemoteWorker.

diff --git a/Engine/Source/Programs/Horde/Samples/RemoteClient/RemoteClientApp.cs b/Engine/Source/Programs/Horde/Samples/RemoteClient/RemoteClientApp.cs
--- a/Engine/Source/Programs/Horde/Samples/RemoteClient/RemoteClientApp.cs
+++ b/Engine/Source/Programs/Horde/Samples/RemoteClient/RemoteClientApp.cs
@@ -25,13 +25,15 @@ namespace RemoteClient
 
 		[CommandLine("-Condition=")]
 		public string? Condition { get; set; }
+
+		[CommandLine("-Cpp")]
+		public bool UseCppWorker { get; set; }
 	}
 
 	class ClientApp
 	{
 		static FileReference CurrentAssemblyFile { get; } = new FileReference(Assembly.GetExecutingAssembly().Location);
 		static DirectoryReference ClientSourceDir { get; } = DirectoryReference.Combine(CurrentAssemblyFile.Directory, "../../..");
-		static FileReference RemoteServerFile { get; } = FileReference.Combine(ClientSourceDir, "../RemoteServer", CurrentAssemblyFile.Directory.MakeRelativeTo(ClientSourceDir), "RemoteServer.dll");
 
 		static async Task Main(string[] args)
 		{
@@ -60,6 +62,21 @@ namespace RemoteClient
 				return;
 			}
 
+			// Run the worker
+			if (options.UseCppWorker)
+			{
+				FileReference remoteServerFile = FileReference.Combine(ClientSourceDir, "../RemoteWorkerCpp/bin/RemoteServerCpp.exe");
+				await RunRemoteAsync(lease, remoteServerFile.Directory, "RemoteWorkerCpp.exe", new List<string>(), logger);
+			}
+			else
+			{
+				FileReference remoteServerFile = FileReference.Combine(ClientSourceDir, "../RemoteWorker", CurrentAssemblyFile.Directory.MakeRelativeTo(ClientSourceDir), "RemoteWorker.dll");
+				await RunRemoteAsync(lease, remoteServerFile.Directory, @"C:\Program Files\dotnet\dotnet.exe", new List<string> { remoteServerFile.GetFileName() }, logger);
+			}
+		}
+
+		static async Task RunRemoteAsync(IComputeLease lease, DirectoryReference uploadDir, string executable, List<string> arguments, ILogger logger)
+		{
 			// Create a message channel on channel id 0. The Horde Agent always listens on this channel for requests.
 			const int ControlChannelId = 0;
 			await using IComputeMessageChannel channel = await lease.Socket.CreateMessageChannelAsync(ControlChannelId, 4 * 1024 * 1024, logger);
@@ -70,13 +87,13 @@ namespace RemoteClient
 			using (TreeWriter treeWriter = new TreeWriter(storage))
 			{
 				DirectoryNode sandbox = new DirectoryNode();
-				await sandbox.CopyFromDirectoryAsync(RemoteServerFile.Directory.ToDirectoryInfo(), new ChunkingOptions(), treeWriter, null);
+				await sandbox.CopyFromDirectoryAsync(uploadDir.ToDirectoryInfo(), new ChunkingOptions(), treeWriter, null);
 				NodeHandle handle = await treeWriter.FlushAsync(sandbox);
 				await channel.UploadFilesAsync("", handle.Locator, storage);
 			}
 
 			// Run the task remotely in the background and echo the output to the console
-			await using (IComputeProcess process = await channel.ExecuteAsync(@"C:\Program Files\dotnet\dotnet.exe", new List<string> { RemoteServerFile.GetFileName() }, null, null))
+			await using (IComputeProcess process = await channel.ExecuteAsync(executable, arguments, null, null))
 			{
 				string? line = await process.ReadLineAsync();
 				logger.LogInformation("[REMOTE] {Line}", line);
diff --git a/Engine/Source/Programs/Horde/Samples/RemoteServer/RemoteServer.csproj b/Engine/Source/Programs/Horde/Samples/RemoteServer/RemoteServer.csproj
deleted file mode 100644
diff --git a/Engine/Source/Programs/Horde/Samples/RemoteServer/RemoteServerApp.cs b/Engine/Source/Programs/Horde/Samples/RemoteServer/RemoteServerApp.cs
deleted file mode 100644
diff --git a/Engine/Source/Programs/Horde/Samples/RemoteWorker/RemoteWorker.csproj b/Engine/Source/Programs/Horde/Samples/RemoteWorker/RemoteWorker.csproj
new file mode 100644
--- /dev/null
+++ b/Engine/Source/Programs/Horde/Samples/RemoteWorker/RemoteWorker.csproj
@@ -0,0 +1,14 @@
+<Project Sdk="Microsoft.NET.Sdk">
+
+  <PropertyGroup>
+    <OutputType>Exe</OutputType>
+    <TargetFramework>net6.0</TargetFramework>
+    <ImplicitUsings>enable</ImplicitUsings>
+    <Nullable>enable</Nullable>
+  </PropertyGroup>
+
+  <ItemGroup>
+    <ProjectReference Include="..\..\..\Shared\EpicGames.Horde\EpicGames.Horde.csproj" />
+  </ItemGroup>
+
+</Project>
diff --git a/Engine/Source/Programs/Horde/Samples/RemoteWorker/RemoteWorkerApp.cs b/Engine/Source/Programs/Horde/Samples/RemoteWorker/RemoteWorkerApp.cs
new file mode 100644
--- /dev/null
+++ b/Engine/Source/Programs/Horde/Samples/RemoteWorker/RemoteWorkerApp.cs
@@ -0,0 +1,28 @@
+// Copyright Epic Games, Inc. All Rights Reserved.
+
+using EpicGames.Horde.Compute;
+using Microsoft.Extensions.Logging.Abstractions;
+
+namespace RemoteServer
+{
+	class ServerApp
+	{
+		static async Task Main()
+		{
+			await using IComputeSocket socket = ComputeSocket.ConnectAsWorker(NullLogger.Instance);
+			using IComputeBufferReader reader = await socket.AttachRecvBufferAsync(1, 1024 * 1024);
+			Console.WriteLine("Connected to client"); // Client will wait for output before sending data on channel 1
+
+			while (!reader.IsComplete)
+			{
+				ReadOnlyMemory<byte> memory = reader.GetMemory();
+				if (memory.Length > 0)
+				{
+					Console.WriteLine("Read value {0}", memory.Span[0]);
+					reader.Advance(1);
+				}
+				await reader.WaitForDataAsync(0);
+			}
+		}
+	}
+}
\ No newline at end of file
diff --git a/Engine/Source/Programs/Horde/Samples/RemoteWorkerCpp/Main.cpp b/Engine/Source/Programs/Horde/Samples/RemoteWorkerCpp/Main.cpp
new file mode 100644
--- /dev/null
+++ b/Engine/Source/Programs/Horde/Samples/RemoteWorkerCpp/Main.cpp
@@ -0,0 +1,46 @@
+// Copyright Epic Games, Inc. All Rights Reserved.
+
+#include <iostream>
+#include "WorkerSocket.h"
+#include "SharedMemoryBuffer.h"
+
+int main()
+{
+	FWorkerSocket Socket;
+	if (!Socket.Open())
+	{
+		std::cout << "Environment variable not set correctly" << std::endl;
+		return 1;
+	}
+
+	FSharedMemoryBuffer Buffer;
+	if (!Buffer.CreateNew(1024 * 1024))
+	{
+		std::cout << "Couldn't create buffer" << std::endl;
+		return 1;
+	}
+	if (!Socket.TryAttachRecvBuffer(1, &Buffer))
+	{
+		std::cout << "Failed to attach recieve buffer" << std::endl;
+		return 1;
+	}
+
+	std::cout << "Connected to client" << std::endl; // // Client will wait for output before sending data on channel 1
+
+	while (!Buffer.HasFinishedReading())
+	{
+		long long Size;
+		const unsigned char* Memory = Buffer.GetReadMemory(Size);
+
+		if (Size > 0)
+		{
+			std::cout << "Read value " << (int)*Memory << std::endl;
+			Buffer.AdvanceReadPosition(1);
+		}
+
+		Buffer.WaitForData(0);
+	}
+
+	return 0;
+}
+
diff --git a/Engine/Source/Programs/Horde/Samples/RemoteWorkerCpp/SharedMemoryBuffer.cpp b/Engine/Source/Programs/Horde/Samples/RemoteWorkerCpp/SharedMemoryBuffer.cpp
new file mode 100644
--- /dev/null
+++ b/Engine/Source/Programs/Horde/Samples/RemoteWorkerCpp/SharedMemoryBuffer.cpp
@@ -0,0 +1,300 @@
+// Copyright Epic Games, Inc. All Rights Reserved.
+
+#include <windows.h>
+#include <assert.h>
+#include <iostream>
+#include "SharedMemoryBuffer.h"
+
+enum class FSharedMemoryBuffer::EState : int
+{
+	Normal,
+	ReaderWaitingForData, // Once new data is available, set _readerTcs
+	WriterWaitingToFlush, // Once flush has been complete, set _writerTcs
+	Finished
+};
+
+class FSharedMemoryBuffer::FHeader
+{
+public:
+	EState GetState() const
+	{
+		return (EState)InterlockedCompareExchange((long*)&State, 0, 0);
+	}
+
+	bool UpdateState(EState OldState, EState NewState)
+	{
+		return InterlockedCompareExchange((long*)&State, (long)NewState, (long)OldState) == (long)OldState;
+	}
+
+	void Compact()
+	{
+		InterlockedExchange64(&ReadPosition, 0);
+		InterlockedExchange64(&WritePosition, GetReadLength());
+	}
+
+	long long GetReadPosition() const
+	{
+		return InterlockedCompareExchange64(const_cast<long long*>(&ReadPosition), 0, 0);
+	}
+
+	long long AdvanceReadPosition(long long Value)
+	{
+		InterlockedAdd64(&ReadPosition, Value);
+		return InterlockedAdd64(&ReadLength, -Value);
+	}
+
+	long long GetReadLength() const
+	{
+		return InterlockedCompareExchange64(const_cast<long long*>(&ReadLength), 0, 0);
+	}
+
+	long long GetWritePosition() const
+	{
+		return InterlockedCompareExchange64(const_cast<long long*>(&WritePosition), 0, 0);
+	}
+
+	long long AdvanceWritePosition(long long value)
+	{
+		InterlockedAdd64(&WritePosition, value);
+		return InterlockedAdd64(&ReadLength, value);
+	}
+
+	void FinishWriting()
+	{
+		InterlockedExchange((long*)&State, (long)EState::Finished);
+	}
+
+private:
+	EState State;
+	unsigned char Padding[4];
+	long long ReadPosition;
+	long long ReadLength;
+	long long WritePosition;
+};
+
+FSharedMemoryBuffer::FSharedMemoryBuffer()
+	: MemoryMappedFile(nullptr)
+	, Memory(nullptr)
+	, Length(0)
+	, Header(nullptr)
+	, ReaderEvent(nullptr)
+	, WriterEvent(nullptr)
+{
+}
+
+FSharedMemoryBuffer::~FSharedMemoryBuffer()
+{
+	Close();
+}
+
+bool FSharedMemoryBuffer::CreateNew(long long Capacity)
+{
+	Close();
+
+	LARGE_INTEGER CapacityLargeInt;
+	CapacityLargeInt.QuadPart = Capacity;
+
+	ReaderEvent = CreateEvent(nullptr, true, FALSE, nullptr);
+	WriterEvent = CreateEvent(nullptr, true, FALSE, nullptr);
+	MemoryMappedFile = CreateFileMappingW(INVALID_HANDLE_VALUE, nullptr, PAGE_READWRITE, CapacityLargeInt.HighPart, CapacityLargeInt.LowPart, nullptr);
+
+	if (ReaderEvent == nullptr || WriterEvent == nullptr || MemoryMappedFile == nullptr || !Initialize())
+	{
+		Close();
+		return false;
+	}
+
+	return true;
+}
+
+bool FSharedMemoryBuffer::OpenExisting(void* InMemoryMappedFile, void* InReaderEvent, void* InWriterEvent)
+{
+	MemoryMappedFile = InMemoryMappedFile;
+	ReaderEvent = InReaderEvent;
+	WriterEvent = InWriterEvent;
+	return Initialize();
+}
+
+void FSharedMemoryBuffer::Close()
+{
+	Memory = nullptr;
+
+	if (Header != nullptr)
+	{
+		UnmapViewOfFile(Header);
+		Header = nullptr;
+	}
+
+	if (MemoryMappedFile != nullptr)
+	{
+		CloseHandle(MemoryMappedFile);
+		MemoryMappedFile = nullptr;
+	}
+
+	if (WriterEvent != nullptr)
+	{
+		CloseHandle(WriterEvent);
+		WriterEvent = nullptr;
+	}
+
+	if (ReaderEvent != nullptr)
+	{
+		CloseHandle(ReaderEvent);
+		ReaderEvent = nullptr;
+	}
+}
+
+bool FSharedMemoryBuffer::DuplicateHandles(void* TargetProcess, void*& OutMemoryMappedFile, void*& OutReaderEvent, void*& OutWriterEvent)
+{
+	HANDLE CurrentProcess = GetCurrentProcess();
+
+	HANDLE TargetMemoryMappedFile;
+	if (DuplicateHandle(CurrentProcess, MemoryMappedFile, TargetProcess, &TargetMemoryMappedFile, 0, false, DUPLICATE_SAME_ACCESS))
+	{
+		HANDLE TargetWriterEvent;
+		if (DuplicateHandle(CurrentProcess, WriterEvent, TargetProcess, &TargetWriterEvent, 0, false, DUPLICATE_SAME_ACCESS))
+		{
+			HANDLE TargetReaderEvent;
+			if (DuplicateHandle(CurrentProcess, ReaderEvent, TargetProcess, &TargetReaderEvent, 0, false, DUPLICATE_SAME_ACCESS))
+			{
+				OutMemoryMappedFile = TargetMemoryMappedFile;
+				OutReaderEvent = TargetReaderEvent;
+				OutWriterEvent = TargetWriterEvent;
+				return true;
+			}
+			CloseHandle(TargetWriterEvent);
+		}
+		CloseHandle(TargetMemoryMappedFile);
+	}
+
+	return false;
+}
+
+bool FSharedMemoryBuffer::HasFinishedReading() const
+{
+	return Header->GetState() == EState::Finished && Header->GetReadLength() == 0;
+}
+
+void FSharedMemoryBuffer::AdvanceReadPosition(long long Size)
+{
+	if (Header->AdvanceReadPosition(Size) == 0)
+	{
+		SetEvent(WriterEvent);
+	}
+}
+
+const unsigned char* FSharedMemoryBuffer::GetReadMemory(long long& OutSize)
+{
+	OutSize = Header->GetReadLength();
+	return Memory + Header->GetReadPosition();
+}
+
+void FSharedMemoryBuffer::WaitForData(long long Size)
+{
+	while (Header->GetReadLength() == Size)
+	{
+		EState State = Header->GetState();
+		if (State == EState::Finished)
+		{
+			break;
+		}
+		else if (State == EState::Normal)
+		{
+			ResetEvent(ReaderEvent);
+			Header->UpdateState(State, EState::ReaderWaitingForData);
+		}
+		else if (State == EState::ReaderWaitingForData)
+		{
+			WaitForSingleObject(ReaderEvent, INFINITE);
+		}
+		else if (State == EState::WriterWaitingToFlush)
+		{
+			Header->Compact();
+			Header->UpdateState(State, EState::Normal);
+			SetEvent(WriterEvent);
+		}
+		else
+		{
+			assert(false);
+		}
+	}
+}
+
+void FSharedMemoryBuffer::FinishWriting()
+{
+	Header->FinishWriting();
+	SetEvent(ReaderEvent);
+}
+
+void FSharedMemoryBuffer::AdvanceWritePosition(long long size)
+{
+	assert(Header->GetState() != EState::Finished);
+
+	Header->AdvanceWritePosition(size);
+
+	if (Header->UpdateState(EState::ReaderWaitingForData, EState::Normal))
+	{
+		SetEvent(ReaderEvent);
+	}
+}
+
+unsigned char* FSharedMemoryBuffer::GetWriteMemory(long long& OutSize)
+{
+	long long WritePosition = Header->GetWritePosition();
+	OutSize = Length - WritePosition;
+	return Memory + WritePosition;
+}
+
+void FSharedMemoryBuffer::Flush()
+{
+	while (Header->GetReadPosition() > 0)
+	{
+		EState State = Header->GetState();
+		if (State == EState::Finished)
+		{
+			break;
+		}
+		else if (State == EState::Normal)
+		{
+			SetEvent(WriterEvent);
+			Header->UpdateState(State, EState::WriterWaitingToFlush);
+		}
+		else if (State == EState::ReaderWaitingForData)
+		{
+			Header->Compact();
+			break;
+		}
+		else if (State == EState::WriterWaitingToFlush)
+		{
+			WaitForSingleObject(WriterEvent, INFINITE);
+		}
+	}
+}
+
+bool FSharedMemoryBuffer::Initialize()
+{
+	if (ReaderEvent == nullptr || WriterEvent == nullptr || MemoryMappedFile == nullptr)
+	{
+		Close();
+		return false;
+	}
+
+	Header = (FHeader*)MapViewOfFile(MemoryMappedFile, FILE_MAP_ALL_ACCESS, 0, 0, 0);
+	if (Header == nullptr)
+	{
+		Close();
+		return false;
+	}
+
+	Memory = (unsigned char*)(Header + 1);
+
+	MEMORY_BASIC_INFORMATION MemInfo;
+	if (VirtualQuery(Header, &MemInfo, sizeof(MemInfo)) == 0)
+	{
+		Close();
+		return false;
+	}
+
+	Length = MemInfo.RegionSize;
+	return true;
+}
diff --git a/Engine/Source/Programs/Horde/Samples/RemoteWorkerCpp/SharedMemoryBuffer.h b/Engine/Source/Programs/Horde/Samples/RemoteWorkerCpp/SharedMemoryBuffer.h
new file mode 100644
--- /dev/null
+++ b/Engine/Source/Programs/Horde/Samples/RemoteWorkerCpp/SharedMemoryBuffer.h
@@ -0,0 +1,65 @@
+// Copyright Epic Games, Inc. All Rights Reserved.
+
+#pragma once
+
+class FSharedMemoryBuffer
+{
+	enum class EState : int;
+	class FHeader;
+
+public:
+	FSharedMemoryBuffer();
+	~FSharedMemoryBuffer();
+
+	// Initialize the object to a new shared memory buffer of the given size
+	bool CreateNew(long long Capacity);
+
+	// Opens an existing shared memory buffer (typically from handles created in another process)
+	bool OpenExisting(void* MemoryMappedFile, void* ReaderEvent, void* WriterEvent);
+
+	// Close the current buffer and release all allocated resources
+	void Close();
+
+	// Duplicates handles for this buffer into another process
+	bool DuplicateHandles(void* TargetProcess, void*& OutMemoryMappedFile, void*& OutReaderEvent, void*& OutWriterEvent);
+
+
+	/*** Reader Interface ***/
+
+	// Test whether the buffer has finished being written to (ie. FinishWriting() has been called) and all data has been read from it.
+	bool HasFinishedReading() const;
+
+	// Move the read cursor forwards by the given number of bytes
+	void AdvanceReadPosition(long long Size);
+
+	// Gets the next data to be read (and the number of valid bytes accessible from the given pointer)
+	const unsigned char* GetReadMemory(long long& OutSize);
+
+	// Wait for more data to be written to the buffer. The given parameter indicates the current size of the read buffer, used to exit immediately if it's changed since the value was fetched.
+	void WaitForData(long long CurrentLength);
+
+
+	/*** Writer Interface ***/
+
+	// Signal that we've finished writing to this buffer
+	void FinishWriting();
+
+	// Move the write cursor forward by the given number of bytes
+	void AdvanceWritePosition(long long Size);
+
+	// Gets the memory that can be written to, and the available space in it
+	unsigned char* GetWriteMemory(long long& OutSize);
+
+	// Waits until all data has been read from the buffer (or more specifically, that the reader is stalled) and moves any unread data to the start of the buffer.
+	void Flush();
+
+private:
+	bool Initialize();
+
+	void* MemoryMappedFile;
+	long long Length;
+	FHeader* Header;
+	unsigned char* Memory;
+	void* ReaderEvent;
+	void* WriterEvent;
+};
diff --git a/Engine/Source/Programs/Horde/Samples/RemoteWorkerCpp/WorkerSocket.cpp b/Engine/Source/Programs/Horde/Samples/RemoteWorkerCpp/WorkerSocket.cpp
new file mode 100644
--- /dev/null
+++ b/Engine/Source/Programs/Horde/Samples/RemoteWorkerCpp/WorkerSocket.cpp
@@ -0,0 +1,120 @@
+// Copyright Epic Games, Inc. All Rights Reserved.
+
+#include <windows.h>
+#include <iostream>
+#include <assert.h>
+#include "WorkerSocket.h"
+
+const char* const FWorkerSocket::EnvVarName = "UE_HORDE_COMPUTE_IPC";
+
+enum class FWorkerSocket::EIpcMessageType : unsigned char
+{
+	Finish,
+	AttachSendBuffer,
+	AttachRecvBuffer,
+};
+
+struct FWorkerSocket::FAttachMessage
+{
+	EIpcMessageType Type;
+	unsigned char Padding[3];
+	int ChannelId;
+	void* MemoryMappedFile;
+	void* ReaderEvent;
+	void* WriterEvent;
+};
+
+FWorkerSocket::FWorkerSocket()
+{
+	ParentProcess = nullptr;
+}
+
+FWorkerSocket::~FWorkerSocket()
+{
+	Close();
+}
+
+bool FWorkerSocket::Open()
+{
+	Close();
+
+	char Buffer[256];
+	int Length = GetEnvironmentVariableA(EnvVarName, Buffer, sizeof(Buffer) / sizeof(Buffer[0]));
+
+	if (Length <= 0 || Length >= sizeof(Buffer))
+	{
+		return false;
+	}
+
+	const size_t NumHandles = 4;
+	size_t Handles[NumHandles] = { 0, };
+	size_t HandleIdx = 0;
+
+	char* Ptr = Buffer;
+	for (; *Ptr != 0; Ptr++)
+	{
+		if (*Ptr == '.' && HandleIdx + 1 < NumHandles)
+		{
+			HandleIdx++;
+		}
+		else if (*Ptr >= '0' && *Ptr <= '9')
+		{
+			Handles[HandleIdx] = (Handles[HandleIdx] * 10) + (*Ptr - '0');
+		}
+		else
+		{
+			return false;
+		}
+	}
+
+	ParentProcess = (void*)Handles[0];
+
+	if (!IpcBuffer.OpenExisting((void*)Handles[1], (void*)Handles[2], (void*)Handles[3]))
+	{
+		Close();
+		return false;
+	}
+
+	return true;
+}
+
+void FWorkerSocket::Close()
+{
+	ParentProcess = nullptr;
+	IpcBuffer.Close();
+}
+
+bool FWorkerSocket::TryAttachRecvBuffer(int ChannelId, FSharedMemoryBuffer* Buffer)
+{
+	return TryAttachBuffer(EIpcMessageType::AttachRecvBuffer, ChannelId, Buffer);
+}
+
+bool FWorkerSocket::TryAttachSendBuffer(int ChannelId, FSharedMemoryBuffer* Buffer)
+{
+	return TryAttachBuffer(EIpcMessageType::AttachSendBuffer, ChannelId, Buffer);
+}
+
+bool FWorkerSocket::TryAttachBuffer(EIpcMessageType Type, int ChannelId, FSharedMemoryBuffer* Buffer)
+{
+	long long Size;
+
+	FAttachMessage* Message = (FAttachMessage*)IpcBuffer.GetWriteMemory(Size);
+	if (Size < sizeof(FAttachMessage))
+	{
+		return false;
+	}
+
+	memset(Message, 0, sizeof(*Message));
+
+	Message->Type = Type;
+	Message->ChannelId = ChannelId;
+	if (!Buffer->DuplicateHandles(ParentProcess, Message->MemoryMappedFile, Message->ReaderEvent, Message->WriterEvent))
+	{
+		return false;
+	}
+
+	IpcBuffer.AdvanceWritePosition(sizeof(*Message));
+	IpcBuffer.Flush();
+
+	return true;
+}
diff --git a/Engine/Source/Programs/Horde/Samples/RemoteWorkerCpp/WorkerSocket.h b/Engine/Source/Programs/Horde/Samples/RemoteWorkerCpp/WorkerSocket.h
new file mode 100644
--- /dev/null
+++ b/Engine/Source/Programs/Horde/Samples/RemoteWorkerCpp/WorkerSocket.h
@@ -0,0 +1,28 @@
+// Copyright Epic Games, Inc. All Rights Reserved.
+
+#pragma once
+
+#include "SharedMemoryBuffer.h"
+
+class FWorkerSocket
+{
+	static const char* const EnvVarName;
+	enum class EIpcMessageType : unsigned char;
+	struct FAttachMessage;
+
+public:
+	FWorkerSocket();
+	~FWorkerSocket();
+
+	bool Open();
+	void Close();
+
+	bool TryAttachRecvBuffer(int ChannelId, FSharedMemoryBuffer* Buffer);
+	bool TryAttachSendBuffer(int ChannelId, FSharedMemoryBuffer* Buffer);
+
+private:
+	void* ParentProcess;
+	FSharedMemoryBuffer IpcBuffer;
+
+	bool TryAttachBuffer(EIpcMessageType Type, int ChannelId, FSharedMemoryBuffer* Buffer);
+};
diff --git a/Engine/Source/Programs/Shared/EpicGames.Horde/Compute/Buffers/SharedMemoryBuffer.cs b/Engine/Source/Programs/Shared/EpicGames.Horde/Compute/Buffers/SharedMemoryBuffer.cs
--- a/Engine/Source/Programs/Shared/EpicGames.Horde/Compute/Buffers/SharedMemoryBuffer.cs
+++ b/Engine/Source/Programs/Shared/EpicGames.Horde/Compute/Buffers/SharedMemoryBuffer.cs
@@ -112,6 +112,14 @@ namespace EpicGames.Horde.Compute.Buffers
 		{
 		}
 
+		/// <summary>
+		/// Opens a <see cref="SharedMemoryBuffer"/> from handles
+		/// </summary>
+		public SharedMemoryBuffer(IntPtr memoryHandle, IntPtr readerEventHandle, IntPtr writerEventHandle)
+			: this(OpenMemoryMappedFileFromHandle(memoryHandle), new Native.EventHandle(readerEventHandle, true), new Native.EventHandle(writerEventHandle, true))
+		{
+		}
+
 		/// <summary>
 		/// Creates a shared memory buffer from a memory mapped file
 		/// </summary>
@@ -144,7 +152,7 @@ namespace EpicGames.Horde.Compute.Buffers
 		/// Opens an IpcBuffer from a string passed in from another process
 		/// </summary>
 		/// <param name="handle">Descriptor for the buffer to open</param>
-		public static MemoryMappedFile OpenMemoryMappedFileFromHandle(IntPtr handle)
+		static MemoryMappedFile OpenMemoryMappedFileFromHandle(IntPtr handle)
 		{
 			MethodInfo? setHandleInfo = typeof(SafeMemoryMappedFileHandle).GetMethod("SetHandle", BindingFlags.Instance | BindingFlags.NonPublic, new[] { typeof(IntPtr) });
 			if (setHandleInfo == null)
@@ -164,68 +172,37 @@ namespace EpicGames.Horde.Compute.Buffers
 			return (MemoryMappedFile)constructorInfo.Invoke(new object[] { safeHandle });
 		}
 
-		/// <summary>
-		/// Opens a <see cref="SharedMemoryBuffer"/> from handles returned by <see cref="GetIpcHandle()"/>
-		/// </summary>
-		public static SharedMemoryBuffer OpenIpcHandle(string handle)
-		{
-			IntPtr[] handles = handle.Split('.').Select(x => new IntPtr((long)UInt64.Parse(x, NumberStyles.None, null))).ToArray();
-			if (handles.Length != 3)
-			{
-				throw new ArgumentException($"Malformed ipc handle string: {handle}", nameof(handle));
-			}
-
-			MemoryMappedFile memoryMappedFile = OpenMemoryMappedFileFromHandle(handles[0]);
-			Native.EventHandle writtenEvent = new Native.EventHandle(handles[1], true);
-			Native.EventHandle flushedEvent = new Native.EventHandle(handles[2], true);
-
-			return new SharedMemoryBuffer(memoryMappedFile, writtenEvent, flushedEvent);
-		}
-
 		/// <summary>
 		/// Gets a string that can be used to open the same buffer in another process
 		/// </summary>
-		public string GetIpcHandle()
+		public void GetHandles(out IntPtr memoryHandle, out IntPtr readerEventHandle, out IntPtr writerEventHandle)
 		{
-			IntPtr memoryHandle = _memoryMappedFile.SafeMemoryMappedFileHandle.DangerousGetHandle();
-			IntPtr writtenEventHandle = _readerEvent.SafeWaitHandle.DangerousGetHandle();
-			IntPtr flushedEventHandle = _writerEvent.SafeWaitHandle.DangerousGetHandle();
-			return GetIpcHandle(memoryHandle, writtenEventHandle, flushedEventHandle);
+			memoryHandle = _memoryMappedFile.SafeMemoryMappedFileHandle.DangerousGetHandle();
+			readerEventHandle = _readerEvent.SafeWaitHandle.DangerousGetHandle();
+			writerEventHandle = _writerEvent.SafeWaitHandle.DangerousGetHandle();
 		}
 
 		/// <summary>
-		/// Duplicates handles for this buffer into the given process, and returns a string that the target process can use to open it.
+		/// Gets a string that can be used to open the same buffer in another process
 		/// </summary>
-		/// <param name="targetProcessHandle">Handle to the target process</param>
-		/// <returns>String that the target process can pass into a call to <see cref="OpenIpcHandle(String)"/> to open the buffer.</returns>
-		public string GetIpcHandle(IntPtr targetProcessHandle)
+		public void DuplicateHandles(IntPtr targetProcessHandle, out IntPtr targetMemoryHandle, out IntPtr targetReaderEventHandle, out IntPtr targetWriterEventHandle)
 		{
 			IntPtr sourceProcessHandle = Native.GetCurrentProcess();
 
-			IntPtr targetMemoryHandle;
 			if (!Native.DuplicateHandle(sourceProcessHandle, _memoryMappedFile.SafeMemoryMappedFileHandle.DangerousGetHandle(), targetProcessHandle, out targetMemoryHandle, 0, false, Native.DUPLICATE_SAME_ACCESS))
 			{
 				throw new Win32Exception();
 			}
 
-			IntPtr targetWrittenEventHandle;
-			if (!Native.DuplicateHandle(sourceProcessHandle, _readerEvent.SafeWaitHandle.DangerousGetHandle(), targetProcessHandle, out targetWrittenEventHandle, 0, false, Native.DUPLICATE_SAME_ACCESS))
+			if (!Native.DuplicateHandle(sourceProcessHandle, _readerEvent.SafeWaitHandle.DangerousGetHandle(), targetProcessHandle, out targetReaderEventHandle, 0, false, Native.DUPLICATE_SAME_ACCESS))
 			{
 				throw new Win32Exception();
 			}
 
-			IntPtr targetFlushedEventHandle;
-			if(!Native.DuplicateHandle(sourceProcessHandle, _writerEvent.SafeWaitHandle.DangerousGetHandle(), targetProcessHandle, out targetFlushedEventHandle, 0, false, Native.DUPLICATE_SAME_ACCESS))
+			if (!Native.DuplicateHandle(sourceProcessHandle, _writerEvent.SafeWaitHandle.DangerousGetHandle(), targetProcessHandle, out targetWriterEventHandle, 0, false, Native.DUPLICATE_SAME_ACCESS))
 			{
 				throw new Win32Exception();
 			}
-
-			return GetIpcHandle(targetMemoryHandle, targetWrittenEventHandle, targetFlushedEventHandle);
-		}
-
-		static string GetIpcHandle(IntPtr memoryHandle, IntPtr writtenEventHandle, IntPtr flushedEventHandle)
-		{
-			return $"{(ulong)memoryHandle.ToInt64()}.{(ulong)writtenEventHandle.ToInt64()}.{(ulong)flushedEventHandle.ToInt64()}";
 		}
 
 		/// <inheritdoc/>
diff --git a/Engine/Source/Programs/Shared/EpicGames.Horde/Compute/ComputeWorker.cs b/Engine/Source/Programs/Shared/EpicGames.Horde/Compute/ComputeWorker.cs
--- a/Engine/Source/Programs/Shared/EpicGames.Horde/Compute/ComputeWorker.cs
+++ b/Engine/Source/Programs/Shared/EpicGames.Horde/Compute/ComputeWorker.cs
@@ -1,6 +1,7 @@
 // Copyright Epic Games, Inc. All Rights Reserved.
 
 using System;
+using System.Buffers.Binary;
 using System.Collections.Generic;
 using System.Diagnostics;
 using System.Runtime.InteropServices;
@@ -8,6 +9,7 @@ using System.Threading;
 using System.Threading.Tasks;
 using EpicGames.Core;
 using EpicGames.Horde.Compute.Buffers;
+using EpicGames.Horde.Compute.Sockets;
 using EpicGames.Horde.Storage;
 using EpicGames.Horde.Storage.Nodes;
 using Microsoft.Extensions.Caching.Memory;
@@ -169,8 +171,8 @@ namespace EpicGames.Horde.Compute
 
 		async Task ExecuteProcessWindowsAsync(IComputeSocket socket, IComputeMessageChannel channel, string executable, IReadOnlyList<string> arguments, string? workingDir, IReadOnlyDictionary<string, string?>? envVars, CancellationToken cancellationToken)
 		{
-			await using IpcComputeMessageChannel ipcChannel = new IpcComputeMessageChannel(4096, _logger);
-			await using (BackgroundTask ipcTask = BackgroundTask.StartNew(ctx => RunIpcChannel(socket, ipcChannel, ctx)))
+			using SharedMemoryBuffer ipcBuffer = new SharedMemoryBuffer(4096);
+			await using (BackgroundTask ipcTask = BackgroundTask.StartNew(ctx => RunIpcChannel(socket, ipcBuffer, ctx)))
 			{
 				Dictionary<string, string?> newEnvVars = new Dictionary<string, string?>();
 				if (envVars != null)
@@ -181,54 +183,68 @@ namespace EpicGames.Horde.Compute
 					}
 				}
 
-				newEnvVars[ComputeSocket.WorkerIpcEnvVar] = ipcChannel.GetStringHandle();
+				newEnvVars[WorkerComputeSocket.IpcEnvVar] = WorkerComputeSocket.GetIpcEnvVarValue(ipcBuffer);
 
 				_logger.LogInformation("Launching {Executable} {Arguments}", CommandLineArguments.Quote(executable), CommandLineArguments.Join(arguments));
 				await ExecuteProcessInternalAsync(channel, executable, arguments, workingDir, newEnvVars, cancellationToken);
 				_logger.LogInformation("Finished executing process");
 
-				ipcChannel.ForceComplete();
+				Memory<byte> buffer = ipcBuffer.GetWriteMemory();
+				buffer.Span[0] = (byte)IpcMessageType.Finish;
+				ipcBuffer.AdvanceWritePosition(1);
+
 				await ipcTask.Task.WaitAsync(TimeSpan.FromSeconds(5.0), cancellationToken);
 			}
 			_logger.LogInformation("Child process has shut down");
 		}
 
-		async Task RunIpcChannel(IComputeSocket socket, IpcComputeMessageChannel ipcChannel, CancellationToken cancellationToken)
+		async Task RunIpcChannel(IComputeSocket socket, SharedMemoryBuffer ipcBuffer, CancellationToken cancellationToken)
 		{
 			List<IDisposable> buffers = new List<IDisposable>();
 			try
 			{
 				for (; ; )
 				{
-					using IComputeMessage message = await ipcChannel.ReceiveAsync(cancellationToken);
-					switch (message.Type)
+					ReadOnlyMemory<byte> memory = ipcBuffer.GetReadMemory();
+					if (memory.Length == 0)
 					{
-						case ComputeMessageType.None:
+						await ipcBuffer.WaitForDataAsync(0, cancellationToken);
+						continue;
+					}
+
+					IpcMessageType messageType = (IpcMessageType)memory.Span[0];
+					switch (messageType)
+					{
+						case IpcMessageType.Finish:
 							return;
-						case ComputeMessageType.AttachRecvBuffer:
+						case IpcMessageType.AttachRecvBuffer:
 							{
-								AttachRecvBufferRequest attachRecvBuffer = IpcComputeMessageChannel.ParseAttachRecvBuffer(message);
+								IpcAttachBufferRequest attachRecvBuffer = IpcAttachBufferRequest.Read(ipcBuffer);
+
 								_logger.LogInformation("Attaching receive buffer to channel {Id}", attachRecvBuffer.ChannelId);
 #pragma warning disable CA2000 // Dispose objects before losing scope
-								(IComputeBufferReader reader, IComputeBufferWriter writer) = SharedMemoryBuffer.OpenIpcHandle(attachRecvBuffer.Handle).ToShared();
+								SharedMemoryBuffer buffer = new SharedMemoryBuffer(attachRecvBuffer.MemoryHandle, attachRecvBuffer.ReaderEventHandle, attachRecvBuffer.WriterEventHandle);
+								(IComputeBufferReader reader, IComputeBufferWriter writer) = buffer.ToShared();
 								reader.Dispose();
 								await socket.AttachRecvBufferAsync(attachRecvBuffer.ChannelId, writer, cancellationToken);
 #pragma warning restore CA2000 // Dispose objects before losing scope
 							}
 							break;
-						case ComputeMessageType.AttachSendBuffer:
+						case IpcMessageType.AttachSendBuffer:
 							{
-								AttachSendBufferRequest attachSendBuffer = IpcComputeMessageChannel.ParseAttachSendBuffer(message);
+								IpcAttachBufferRequest attachSendBuffer = IpcAttachBufferRequest.Read(ipcBuffer);
+
 								_logger.LogInformation("Attaching send buffer to channel {Id}", attachSendBuffer.ChannelId);
 #pragma warning disable CA2000 // Dispose objects before losing scope
-								(IComputeBufferReader reader, IComputeBufferWriter writer) = SharedMemoryBuffer.OpenIpcHandle(attachSendBuffer.Handle).ToShared();
+								SharedMemoryBuffer buffer = new SharedMemoryBuffer(attachSendBuffer.MemoryHandle, attachSendBuffer.ReaderEventHandle, attachSendBuffer.WriterEventHandle);
+								(IComputeBufferReader reader, IComputeBufferWriter writer) = buffer.ToShared();
 								writer.Dispose();
 								await socket.AttachSendBufferAsync(attachSendBuffer.ChannelId, reader, cancellationToken);
 #pragma warning restore CA2000 // Dispose objects before losing scope
 							}
 							break;
 						default:
-							throw new ComputeInvalidMessageException(message);
+							throw new ComputeInternalException($"Invalid IPC message type: {messageType}");
 					}
 				}
 			}
diff --git a/Engine/Source/Programs/Shared/EpicGames.Horde/Compute/IComputeMessage.cs b/Engine/Source/Programs/Shared/EpicGames.Horde/Compute/IComputeMessage.cs
--- a/Engine/Source/Programs/Shared/EpicGames.Horde/Compute/IComputeMessage.cs
+++ b/Engine/Source/Programs/Shared/EpicGames.Horde/Compute/IComputeMessage.cs
@@ -92,20 +92,6 @@ namespace EpicGames.Horde.Compute
 
 		#endregion
 
-		#region Buffers
-
-		/// <summary>
-		/// Registers a send buffer created by a child process
-		/// </summary>
-		AttachSendBuffer = 0x30,
-
-		/// <summary>
-		/// Registers a receive buffer created by a child process
-		/// </summary>
-		AttachRecvBuffer = 0x31,
-
-		#endregion
-
 		#region Test Requests
 
 		/// <summary>
@@ -180,16 +166,6 @@ namespace EpicGames.Horde.Compute
 	/// </summary>
 	public record struct ReadBlobMessage(BlobLocator Locator, int Offset, int Length);
 
-	/// <summary>
-	/// Request a receive buffer in a worker process
-	/// </summary>
-	public record struct AttachSendBufferRequest(int ChannelId, string Handle);
-
-	/// <summary>
-	/// Request a receive buffer in a worker process
-	/// </summary>
-	public record struct AttachRecvBufferRequest(int ChannelId, string Handle);
-
 	/// <summary>
 	/// Message for running an XOR command
 	/// </summary>
diff --git a/Engine/Source/Programs/Shared/EpicGames.Horde/Compute/IComputeSocket.cs b/Engine/Source/Programs/Shared/EpicGames.Horde/Compute/IComputeSocket.cs
--- a/Engine/Source/Programs/Shared/EpicGames.Horde/Compute/IComputeSocket.cs
+++ b/Engine/Source/Programs/Shared/EpicGames.Horde/Compute/IComputeSocket.cs
@@ -42,33 +42,11 @@ namespace EpicGames.Horde.Compute
 	/// </summary>
 	public static class ComputeSocket
 	{
-		/// <summary>
-		/// Environment variable for channel info for communicating with the worker process
-		/// </summary>
-		public const string WorkerIpcEnvVar = "UE_HORDE_COMPUTE_IPC";
-
 		/// <summary>
 		/// Creates a socket for a worker
 		/// </summary>
 		/// <param name="logger">Logger for diagnostic messages</param>
-		public static IComputeSocket ConnectAsWorker(ILogger logger)
-		{
-			string? handle = Environment.GetEnvironmentVariable(WorkerIpcEnvVar);
-			if (handle == null)
-			{
-				throw new InvalidOperationException($"Environment variable {WorkerIpcEnvVar} is not defined; cannot connect as worker.");
-			}
-
-			try
-			{
-				IpcComputeMessageChannel channel = IpcComputeMessageChannel.FromStringHandle(handle, logger);
-				return new WorkerComputeSocket(channel, logger);
-			}
-			catch(Exception ex)
-			{
-				throw new Exception($"While connecting using '{handle}'", ex);
-			}
-		}
+		public static IComputeSocket ConnectAsWorker(ILogger logger) => WorkerComputeSocket.FromEnvironment(logger);
 	}
 
 	/// <summary>
diff --git a/Engine/Source/Programs/Shared/EpicGames.Horde/Compute/IpcComputeMessageChannel.cs b/Engine/Source/Programs/Shared/EpicGames.Horde/Compute/IpcComputeMessageChannel.cs
deleted file mode 100644
diff --git a/Engine/Source/Programs/Shared/EpicGames.Horde/Compute/IpcMessages.cs b/Engine/Source/Programs/Shared/EpicGames.Horde/Compute/IpcMessages.cs
new file mode 100644
--- /dev/null
+++ b/Engine/Source/Programs/Shared/EpicGames.Horde/Compute/IpcMessages.cs
@@ -0,0 +1,45 @@
+// Copyright Epic Games, Inc. All Rights Reserved.
+
+using System;
+using System.Buffers.Binary;
+using EpicGames.Horde.Compute.Buffers;
+
+namespace EpicGames.Horde.Compute
+{
+	enum IpcMessageType
+	{
+		Finish,
+		AttachSendBuffer,
+		AttachRecvBuffer,
+	}
+
+	record class IpcAttachBufferRequest(IpcMessageType Type, int ChannelId, IntPtr MemoryHandle, IntPtr ReaderEventHandle, IntPtr WriterEventHandle)
+	{
+		public static IpcAttachBufferRequest Read(SharedMemoryBuffer buffer)
+		{
+			ReadOnlySpan<byte> span = buffer.GetReadMemory().Span;
+
+			IpcMessageType type = (IpcMessageType)span[0];
+			int ChannelId = BinaryPrimitives.ReadInt32LittleEndian(span.Slice(4));
+			IntPtr MemoryMappedFileHandle = new IntPtr(BinaryPrimitives.ReadInt64LittleEndian(span.Slice(8)));
+			IntPtr ReaderEventHandle = new IntPtr(BinaryPrimitives.ReadInt64LittleEndian(span.Slice(16)));
+			IntPtr WriterEventHandle = new IntPtr(BinaryPrimitives.ReadInt64LittleEndian(span.Slice(24)));
+			buffer.AdvanceReadPosition(32);
+
+			return new IpcAttachBufferRequest(type, ChannelId, MemoryMappedFileHandle, ReaderEventHandle, WriterEventHandle);
+		}
+
+		public void Write(SharedMemoryBuffer buffer)
+		{
+			Span<byte> span = buffer.GetWriteMemory().Span;
+
+			span[0] = (byte)Type;
+			BinaryPrimitives.WriteInt32LittleEndian(span.Slice(4), ChannelId);
+			BinaryPrimitives.WriteInt64LittleEndian(span.Slice(8), MemoryHandle.ToInt64());
+			BinaryPrimitives.WriteInt64LittleEndian(span.Slice(16), ReaderEventHandle.ToInt64());
+			BinaryPrimitives.WriteInt64LittleEndian(span.Slice(24), WriterEventHandle.ToInt64());
+
+			buffer.AdvanceWritePosition(32);
+		}
+	}
+}
diff --git a/Engine/Source/Programs/Shared/EpicGames.Horde/Compute/Sockets/WorkerComputeSocket.cs b/Engine/Source/Programs/Shared/EpicGames.Horde/Compute/Sockets/WorkerComputeSocket.cs
--- a/Engine/Source/Programs/Shared/EpicGames.Horde/Compute/Sockets/WorkerComputeSocket.cs
+++ b/Engine/Source/Programs/Shared/EpicGames.Horde/Compute/Sockets/WorkerComputeSocket.cs
@@ -1,6 +1,8 @@
 // Copyright Epic Games, Inc. All Rights Reserved.
 
 using System;
+using System.Globalization;
+using System.Linq;
 using System.Threading;
 using System.Threading.Tasks;
 using EpicGames.Horde.Compute.Buffers;
@@ -10,22 +12,67 @@ namespace EpicGames.Horde.Compute.Sockets
 {
 	sealed class WorkerComputeSocket : IComputeSocket
 	{
-		readonly IpcComputeMessageChannel _ipcMessageChannel;
+		/// <summary>
+		/// Environment variable for channel info for communicating with the worker process
+		/// </summary>
+		public const string IpcEnvVar = "UE_HORDE_COMPUTE_IPC";
+
+		readonly IntPtr _parentProcessHandle;
+		readonly SharedMemoryBuffer _ipcBuffer;
 		readonly ILogger _logger;
 
 		/// <summary>
 		/// Constructor
 		/// </summary>
-		/// <param name="ipcMessageChannel">Channel for IPC messages</param>
+		/// <param name="parentProcessHandle">Handle to the parent process</param>
+		/// <param name="ipcBuffer">Channel for IPC messages</param>
 		/// <param name="logger">Logger for trace output</param>
-		public WorkerComputeSocket(IpcComputeMessageChannel ipcMessageChannel, ILogger logger)
+		public WorkerComputeSocket(IntPtr parentProcessHandle, SharedMemoryBuffer ipcBuffer, ILogger logger)
 		{
-			_ipcMessageChannel = ipcMessageChannel;
+			_parentProcessHandle = parentProcessHandle;
+			_ipcBuffer = ipcBuffer;
 			_logger = logger;
 		}
 
+		public static WorkerComputeSocket FromEnvironment(ILogger logger)
+		{
+			string? handle = Environment.GetEnvironmentVariable(IpcEnvVar);
+			if (handle == null)
+			{
+				throw new InvalidOperationException($"Environment variable {IpcEnvVar} is not defined; cannot connect as worker.");
+			}
+
+			IntPtr[] handles = handle.Split('.').Select(x => new IntPtr((long)UInt64.Parse(x, NumberStyles.None, null))).ToArray();
+
+			IntPtr parentProcessHandle = handles[0];
+			SharedMemoryBuffer ipcBuffer = new SharedMemoryBuffer(handles[1], handles[2], handles[3]);
+
+			return new WorkerComputeSocket(parentProcessHandle, ipcBuffer, logger);
+		}
+
+		static IntPtr s_currentProcessHandle;
+
+		public static string GetIpcEnvVarValue(SharedMemoryBuffer ipcBuffer)
+		{
+			if (s_currentProcessHandle == IntPtr.Zero)
+			{
+				IntPtr currentProcessHandle = Native.GetCurrentProcess();
+				Native.DuplicateHandle(currentProcessHandle, currentProcessHandle, currentProcessHandle, out s_currentProcessHandle, 0, true, Native.DUPLICATE_SAME_ACCESS);
+			}
+
+			IntPtr[] handles = new IntPtr[4];
+			handles[0] = s_currentProcessHandle;
+			ipcBuffer.GetHandles(out handles[1], out handles[2], out handles[3]);
+
+			return String.Join(".", handles.Select(x => ((ulong)x.ToInt64()).ToString()));
+		}
+
 		/// <inheritdoc/>
-		public ValueTask DisposeAsync() => _ipcMessageChannel.DisposeAsync();
+		public ValueTask DisposeAsync()
+		{
+			_ipcBuffer.Dispose();
+			return new ValueTask();
+		}
 
 		/// <inheritdoc/>
 		public IComputeBuffer CreateBuffer(long capacity) => new SharedMemoryBuffer(capacity);
@@ -34,16 +81,26 @@ namespace EpicGames.Horde.Compute.Sockets
 		public async ValueTask AttachRecvBufferAsync(int channelId, IComputeBufferWriter writer, CancellationToken cancellationToken)
 		{
 			SharedMemoryBuffer buffer = (SharedMemoryBuffer)writer.Buffer;
-			await _ipcMessageChannel.AttachRecvBufferAsync(channelId, buffer, cancellationToken);
+
+			buffer.DuplicateHandles(_parentProcessHandle, out IntPtr targetMemoryHandle, out IntPtr targetReaderEventHandle, out IntPtr targetWriterEventHandle);
+			IpcAttachBufferRequest request = new IpcAttachBufferRequest(IpcMessageType.AttachRecvBuffer, channelId, targetMemoryHandle, targetReaderEventHandle, targetWriterEventHandle);
+			request.Write(_ipcBuffer);
 			_logger.LogTrace("Attached recv buffer to channel {ChannelId}", channelId);
+
+			await _ipcBuffer.FlushAsync(cancellationToken);
 		}
 
 		/// <inheritdoc/>
 		public async ValueTask AttachSendBufferAsync(int channelId, IComputeBufferReader reader, CancellationToken cancellationToken)
 		{
 			SharedMemoryBuffer buffer = (SharedMemoryBuffer)reader.Buffer;
-			await _ipcMessageChannel.AttachSendBufferAsync(channelId, buffer, cancellationToken);
+
+			buffer.DuplicateHandles(_parentProcessHandle, out IntPtr targetMemoryHandle, out IntPtr targetReaderEventHandle, out IntPtr targetWriterEventHandle);
+			IpcAttachBufferRequest request = new IpcAttachBufferRequest(IpcMessageType.AttachSendBuffer, channelId, targetMemoryHandle, targetReaderEventHandle, targetWriterEventHandle);
+			request.Write(_ipcBuffer);
 			_logger.LogTrace("Attached send buffer to channel {ChannelId}", channelId);
+
+			await _ipcBuffer.FlushAsync(cancellationToken);
 		}
 	}
 }
