Horde: Add a flag in the job response which indicates whether to use the new artifacts endpoint.

diff --git a/Engine/Source/Programs/Horde/Horde.Server/Jobs/JobMessages.cs b/Engine/Source/Programs/Horde/Horde.Server/Jobs/JobMessages.cs
--- a/Engine/Source/Programs/Horde/Horde.Server/Jobs/JobMessages.cs
+++ b/Engine/Source/Programs/Horde/Horde.Server/Jobs/JobMessages.cs
@@ -341,6 +341,11 @@ namespace Horde.Server.Jobs
 		/// </summary>
 		public DateTimeOffset UpdateTime { get; set; }
 
+		/// <summary>
+		/// Whether to use the V2 artifacts endpoint
+		/// </summary>
+		public bool UseArtifactsV2 { get; set; }
+
 		/// <summary>
 		/// Whether issues are being updated by this job
 		/// </summary>
@@ -379,6 +384,7 @@ namespace Horde.Server.Jobs
 			Reports = job.Reports?.ConvertAll(x => new GetReportResponse(x));
 			Arguments = job.Arguments.ToList();
 			UpdateTime = new DateTimeOffset(job.UpdateTimeUtc);
+			UseArtifactsV2 = job.JobOptions?.UseNewTempStorage ?? false;
 			UpdateIssues = job.UpdateIssues;
 		}
 	}
