Horde: Include the agent software version in the /api/v1/server/info response.

diff --git a/Engine/Source/Programs/Horde/Horde.Server/Server/ServerController.cs b/Engine/Source/Programs/Horde/Horde.Server/Server/ServerController.cs
--- a/Engine/Source/Programs/Horde/Horde.Server/Server/ServerController.cs
+++ b/Engine/Source/Programs/Horde/Horde.Server/Server/ServerController.cs
@@ -2,6 +2,10 @@
 
 using System.Diagnostics;
 using System.Reflection;
+using System.Threading.Tasks;
+using Horde.Server.Agents;
+using Horde.Server.Tools;
+using HordeCommon;
 using Microsoft.AspNetCore.Authorization;
 using Microsoft.AspNetCore.Mvc;
 using Microsoft.Extensions.Options;
@@ -16,18 +20,18 @@ namespace Horde.Server.Server
 	[Route("[controller]")]
 	public class ServerController : ControllerBase
 	{
-
-		/// <summary>
-		/// Settings for the server
-		/// </summary>
-		readonly IOptionsMonitor<ServerSettings> _settings;
+		readonly IToolCollection _toolCollection;
+		readonly IClock _clock;
+		readonly IOptionsSnapshot<GlobalConfig> _globalConfig;
 
 		/// <summary>
 		/// Constructor
 		/// </summary>
-		public ServerController(IOptionsMonitor<ServerSettings> settings)
+		public ServerController(IToolCollection toolCollection, IClock clock, IOptionsSnapshot<GlobalConfig> globalConfig)
 		{
-			_settings = settings;
+			_toolCollection = toolCollection;
+			_clock = clock;
+			_globalConfig = globalConfig;
 		}
 		
 		/// <summary>
@@ -47,9 +51,21 @@ namespace Horde.Server.Server
 		[HttpGet]
 		[Route("/api/v1/server/info")]
 		[ProducesResponseType(typeof(GetServerInfoResponse), 200)]
-		public ActionResult<GetServerInfoResponse> GetServerInfo()
+		public async Task<ActionResult<GetServerInfoResponse>> GetServerInfo()
 		{
-			return new GetServerInfoResponse();
+			string? agentVersion = null;
+
+			ITool? tool = await _toolCollection.GetAsync(AgentExtensions.DefaultAgentSoftwareToolId, _globalConfig.Value);
+			if (tool != null)
+			{
+				IToolDeployment? deployment = tool.GetCurrentDeployment(1.0, _clock.UtcNow);
+				if (deployment != null)
+				{
+					agentVersion = deployment.Version;
+				}
+			}
+
+			return new GetServerInfoResponse(agentVersion);
 		}
 	}
 }
diff --git a/Engine/Source/Programs/Horde/Horde.Server/Server/ServerMessages.cs b/Engine/Source/Programs/Horde/Horde.Server/Server/ServerMessages.cs
--- a/Engine/Source/Programs/Horde/Horde.Server/Server/ServerMessages.cs
+++ b/Engine/Source/Programs/Horde/Horde.Server/Server/ServerMessages.cs
@@ -17,6 +17,11 @@ namespace Horde.Server.Server
 		/// </summary>
         public string ServerVersion { get; set; }
 
+		/// <summary>
+		/// The current agent version string
+		/// </summary>
+		public string? AgentVersion { get; set; }
+
         /// <summary>
         /// The operating system server is hosted on
         /// </summary>
@@ -25,10 +30,11 @@ namespace Horde.Server.Server
 		/// <summary>
 		/// Constructor
 		/// </summary>
-		public GetServerInfoResponse()
+		public GetServerInfoResponse(string? agentVersion)
         {
             FileVersionInfo versionInfo = FileVersionInfo.GetVersionInfo(Assembly.GetExecutingAssembly().Location);		
-			ServerVersion = versionInfo.ProductVersion ?? String.Empty;			
+			ServerVersion = versionInfo.ProductVersion ?? String.Empty;
+			AgentVersion = agentVersion;
 			OsDescription = RuntimeInformation.OSDescription;			
 		}
 	}
