Horde: Do not redirect routes under /api to the default SPA files.

diff --git a/Engine/Source/Programs/Horde/Horde.Server/Startup.cs b/Engine/Source/Programs/Horde/Horde.Server/Startup.cs
--- a/Engine/Source/Programs/Horde/Horde.Server/Startup.cs
+++ b/Engine/Source/Programs/Horde/Horde.Server/Startup.cs
@@ -1010,9 +1010,9 @@ namespace Horde.Server
 
 			if (DirectoryReference.Exists(dashboardDir)) 
 			{
-				app.UseSpaStaticFiles();
+				app.UseWhen(IsSpaRequest, builder => builder.UseSpaStaticFiles());
 			}
-						
+
 			app.UseRouting();
 
 			app.UseAuthentication();
@@ -1034,10 +1034,7 @@ namespace Horde.Server
 
 			if (DirectoryReference.Exists(dashboardDir)) 
 			{
-				app.UseSpa(spa =>
-				{ 
-					spa.Options.SourcePath = "DashboardApp";        
-				});
+				app.MapWhen(IsSpaRequest, builder => builder.UseSpa(spa => spa.Options.SourcePath = "DashboardApp"));
 			}
 
 			if (settings.Value.OpenBrowser && RuntimeInformation.IsOSPlatform(OSPlatform.Windows))
@@ -1046,6 +1043,11 @@ namespace Horde.Server
 			}
 		}
 
+		static bool IsSpaRequest(HttpContext context)
+		{
+			return !context.Request.Path.StartsWithSegments("/api", StringComparison.OrdinalIgnoreCase);
+		}
+
 		static void LaunchBrowser(IApplicationBuilder app)
 		{
 			IServerAddressesFeature? feature = app.ServerFeatures.Get<IServerAddressesFeature>();
