Horde: Restore the default page for the server. This got lost in the Horde.Build -> Horde.Server rename.

diff --git a/Engine/Source/Programs/Horde/Horde.Server/DashboardApp/index.html b/Engine/Source/Programs/Horde/Horde.Server/DashboardApp/index.html
--- a/Engine/Source/Programs/Horde/Horde.Server/DashboardApp/index.html
+++ b/Engine/Source/Programs/Horde/Horde.Server/DashboardApp/index.html
@@ -1 +1,18 @@
-<!doctype html><html lang="en"><head><meta charset="utf-8"/><link rel="icon" href="/favicon.svg"/><meta name="viewport" content="width=device-width,initial-scale=1"/><meta name="theme-color" content="#000000"/><meta name="description" content="Web site created using create-react-app"/><link rel="apple-touch-icon" href="/logo192.png"/><link rel="manifest" href="/manifest.json"/><title>Horde</title><script defer="defer" src="/static/js/main.92c6e7ff.js"></script><link href="/static/css/main.7a06f76f.css" rel="stylesheet"></head><body><noscript>You need to enable JavaScript to run this app.</noscript><div id="root"></div></body></html>
\ No newline at end of file
+<!DOCTYPE html>
+<html>
+<head>
+	<title>Horde Server</title>
+	<style>
+		body { font-family: 'Segoe UI', 'Roboto', arial, sans-serif; }
+		p { margin:20px; font-size:13px; }
+		h1 { margin:20px; font-size:32px; font-weight:200; }
+	</style>
+	<link xmlns="http://www.w3.org/1999/xhtml" rel="shortcut icon" href="favicon.svg" />
+</head>
+<body>
+	<h1>Horde Server</h1>
+	<p>To have the Horde server host the full dashboard app, copy the built dashboard files into the DashboardApp folder.</p>
+	<p><a href="account/">Account Page</a></p>
+	<p><a href="swagger/">API Documentation</a></p>
+</body>
+</html>
