Horde: Adjust breadcrumb font size based on text length
#skipci
#rnx

diff --git a/Engine/Source/Programs/Horde/HordeDashboard/src/components/Breadcrumbs.tsx b/Engine/Source/Programs/Horde/HordeDashboard/src/components/Breadcrumbs.tsx
--- a/Engine/Source/Programs/Horde/HordeDashboard/src/components/Breadcrumbs.tsx
+++ b/Engine/Source/Programs/Horde/HordeDashboard/src/components/Breadcrumbs.tsx
@@ -71,7 +71,20 @@ export const Breadcrumbs: React.FC<{ items: BreadcrumbItem[], title?: string, su
       </Link>;
    });
 
-   const bottomElement = <Text className={classes.crumb} styles={{ root: { color: modeColors.text, fontSize: 28 } }}>{last.text}</Text>;
+   const bottomLength = last.text?.length ?? 0;
+   let bottomFontSize = 28;
+
+   if (bottomLength > 32) {
+      bottomFontSize = 24;
+   }
+   if (bottomLength > 48) {
+      bottomFontSize = 22;
+   }
+   if (bottomLength > 64) {
+      bottomFontSize = 20;
+   }
+
+   const bottomElement = <Text className={classes.crumb} styles={{ root: { color: modeColors.text, fontSize: bottomFontSize } }}>{last.text}</Text>;
 
    const alert = notices.alertText;
 
