Horde: Add version info and the new artifact browser types to api
#skipci
#rnx

diff --git a/Engine/Source/Programs/Horde/HordeDashboard/src/backend/Api.ts b/Engine/Source/Programs/Horde/HordeDashboard/src/backend/Api.ts
--- a/Engine/Source/Programs/Horde/HordeDashboard/src/backend/Api.ts
+++ b/Engine/Source/Programs/Horde/HordeDashboard/src/backend/Api.ts
@@ -1113,6 +1113,9 @@ export type GetJobResponse = {
 	/** Whether to update issues based on the outcome of this job */
 	updateIssues?: boolean;
 
+	/** Whether to use the V2 artifacts endpoint */
+	useArtifactsV2?: boolean;
+
 	/**  Custom permissions for this object */
 	acl?: GetAclResponse;
 
@@ -1418,6 +1421,59 @@ export type GetArtifactZipRequest = {
 }
 
 
+// Artifacts V2
+
+export type ArtifactContextType = "step-trace" | "step-output" | "step-saved";
+
+/// Request to create a zip file with artifact data
+export type CreateZipRequest = {
+	/// Filter lines for the zip. Uses standard <see cref="FileFilter"/> syntax.
+	filter: string[];
+}
+
+
+/** Describes an artifact */
+export type GetArtifactResponseV2 = {
+	
+	id: string;
+
+	type: ArtifactContextType;
+
+	keys: string[]
+
+}
+
+/** Result of an artifact search */
+export type FindArtifactsResponse =  {	
+	/** List of artifacts matching the search criteria*/
+	artifacts: GetArtifactResponseV2[];
+}
+
+/** Describes a file within an artifact */
+export type GetArtifactFileEntryResponse = {
+
+	name: string;
+	length: number;
+	ioHash: string;
+}
+
+/** Describes a directory within an artifact */
+export type GetArtifactDirectoryEntryResponse = {
+	name: string;
+	length: number;
+	ioHash: string;
+}
+
+/** Describes a directory within an artifact */
+export type GetArtifactDirectoryResponse = {
+
+	/** Names of sub-directories */
+	directories: GetArtifactDirectoryEntryResponse[];
+
+	/** Files within the directory */
+	files: GetArtifactFileEntryResponse[];
+}
+
 /**Parameters required to update a log file */
 export type UpdateLogFileRequest = {
 
@@ -3942,6 +3998,9 @@ export type GetServerInfoResponse = {
 	/// Server version info
 	serverVersion: string;
 
+	/// The current agent version
+	agentVersion?: string;
+
 	/// The operating system server is hosted on
 	osDescription: string;
 
diff --git a/Engine/Source/Programs/Horde/HordeDashboard/src/components/VersionModal.tsx b/Engine/Source/Programs/Horde/HordeDashboard/src/components/VersionModal.tsx
--- a/Engine/Source/Programs/Horde/HordeDashboard/src/components/VersionModal.tsx
+++ b/Engine/Source/Programs/Horde/HordeDashboard/src/components/VersionModal.tsx
@@ -2,14 +2,13 @@
 import { DetailsList, DetailsListLayoutMode, IColumn, IconButton, Label, Modal, PrimaryButton, SelectionMode, Stack, Text } from '@fluentui/react';
 import React, { useState } from 'react';
 import backend from '../backend';
-import { GetAgentSoftwareChannelResponse, GetServerInfoResponse } from '../backend/Api';
-import dashboard from '../backend/Dashboard';
+import { GetServerInfoResponse } from '../backend/Api';
 import { hordeClasses } from '../styles/Styles';
 
 
 export const VersionModal: React.FC<{ show: boolean, onClose: () => void }> = ({ show, onClose }) => {
 
-   const [version, setVersion] = useState<{ serverInfo?: GetServerInfoResponse, agentInfo?: GetAgentSoftwareChannelResponse, querying?: boolean }>({});
+   const [version, setVersion] = useState<{ serverInfo?: GetServerInfoResponse, querying?: boolean }>({});
 
    type GeneralItem = {
       name: string;
@@ -25,19 +24,10 @@ export const VersionModal: React.FC<{ show: boolean, onClose: () => void }> = ({
       setVersion({ querying: true });
 
       (async () => {
-         let agentInfo: GetAgentSoftwareChannelResponse | undefined;
 
          const serverInfo = await backend.getServerInfo();
 
-         try {
-            if (dashboard.hordeAdmin) {
-               agentInfo = await backend.getAgentSoftwareChannel();
-            }
-         } catch (error) {
-            console.error(`Unable to get agent version: ${error}`);
-         }
-
-         setVersion({ serverInfo: serverInfo, agentInfo: agentInfo })
+         setVersion({ serverInfo: serverInfo})
       })()
 
    }
@@ -66,8 +56,8 @@ export const VersionModal: React.FC<{ show: boolean, onClose: () => void }> = ({
 
    versionItems.push({ name: "Server", value: version.serverInfo?.serverVersion ?? "Querying" });
 
-   if (dashboard.hordeAdmin) {
-      versionItems.push({ name: "Agent", value: version.agentInfo?.version ?? "Querying" });
+   if (version.serverInfo?.agentVersion) {
+      versionItems.push({ name: "Agent", value: version.serverInfo?.agentVersion });
    }
 
    const onRenderItemColumn = (item: GeneralItem, index?: number, columnIn?: IColumn) => {
