Horde: Add settings flags for Perforce replication and log service

There's likely a memory leak in native code, causing memory usage to grow over time. This only happens for the processes running in worker mode. Having the ability to turn these on/off allow for testing in production and compare mem usage of worker processes.

diff --git a/Engine/Source/Programs/Horde/Horde.Server/ServerSettings.cs b/Engine/Source/Programs/Horde/Horde.Server/ServerSettings.cs
--- a/Engine/Source/Programs/Horde/Horde.Server/ServerSettings.cs
+++ b/Engine/Source/Programs/Horde/Horde.Server/ServerSettings.cs
@@ -227,9 +227,9 @@ namespace Horde.Server
 		public bool ReplicateMetadata { get; set; } = true;
 
 		/// <summary>
-		/// Whether to mirror commit metadata to the database
+		/// Whether to mirror commit data to storage
 		/// </summary>
-		public bool ReplicateContent { get; set; } = false;
+		public bool ReplicateContent { get; set; } = true;
 
 		/// <summary>
 		/// Options for how objects are packed together
@@ -505,6 +505,11 @@ namespace Horde.Server
 		/// Whether to log requests to the UpdateSession and QueryServerState RPC endpoints
 		/// </summary>
 		public bool LogSessionRequests { get; set; } = false;
+		
+		/// <summary>
+		/// Whether to enable the hosted LogService running background jobs
+		/// </summary>
+		public bool EnableLogService { get; set; } = true;
 
 		/// <summary>
 		/// Default fleet manager to use (when not specified by pool)
diff --git a/Engine/Source/Programs/Horde/Horde.Server/Startup.cs b/Engine/Source/Programs/Horde/Horde.Server/Startup.cs
--- a/Engine/Source/Programs/Horde/Horde.Server/Startup.cs
+++ b/Engine/Source/Programs/Horde/Horde.Server/Startup.cs
@@ -653,35 +653,42 @@ namespace Horde.Server
 			if (settings.IsRunModeActive(RunMode.Worker) && !settings.DatabaseReadOnlyMode)
 			{
 				services.AddHostedService<MongoUpgradeService>();
-
 				services.AddHostedService(provider => provider.GetRequiredService<FleetService>());
+				services.AddHostedService(provider => provider.GetRequiredService<ConsistencyService>());
+				services.AddHostedService(provider => provider.GetRequiredService<IssueService>());
+				services.AddHostedService<IssueReportService>();
+				services.AddHostedService(provider => (NotificationService)provider.GetRequiredService<INotificationService>());
+				services.AddHostedService<MetricService>();
+				services.AddHostedService(provider => provider.GetRequiredService<PerforceLoadBalancer>());
+				services.AddHostedService<PoolUpdateService>();
+				services.AddHostedService<TelemetryService>();
+				services.AddHostedService(provider => provider.GetRequiredService<DeviceService>());
+				services.AddHostedService(provider => provider.GetRequiredService<TestDataService>());
 				
 				if (settings.Commits.ReplicateMetadata)
 				{
 					services.AddHostedService(provider => provider.GetRequiredService<PerforceServiceCache>());
 				}
 
-				services.AddHostedService(provider => provider.GetRequiredService<ConsistencyService>());
-				services.AddHostedService(provider => provider.GetRequiredService<IssueService>());
-				services.AddHostedService<IssueReportService>();
-				services.AddHostedService(provider => (LogFileService)provider.GetRequiredService<ILogFileService>());
-				services.AddHostedService(provider => (NotificationService)provider.GetRequiredService<INotificationService>());
-				services.AddHostedService(provider => provider.GetRequiredService<ReplicationService>());
+				if (settings.Commits.ReplicateContent)
+				{
+					services.AddHostedService(provider => provider.GetRequiredService<ReplicationService>());
+				}
+
+				if (settings.EnableLogService)
+				{
+					services.AddHostedService(provider => (LogFileService)provider.GetRequiredService<ILogFileService>());
+				}
+				
 				if (!settings.DisableSchedules)
 				{
 					services.AddHostedService(provider => provider.GetRequiredService<ScheduleService>());
 				}
 
-				services.AddHostedService<MetricService>();
-				services.AddHostedService(provider => provider.GetRequiredService<PerforceLoadBalancer>());
-				services.AddHostedService<PoolUpdateService>();
 				if (settings.SlackToken != null)
 				{
 					services.AddHostedService(provider => provider.GetRequiredService<SlackNotificationSink>());
 				}
-				services.AddHostedService<TelemetryService>();
-				services.AddHostedService(provider => provider.GetRequiredService<DeviceService>());
-				services.AddHostedService(provider => provider.GetRequiredService<TestDataService>());
 			}
 
 			services.AddHostedService(provider => provider.GetRequiredService<IExternalIssueService>());
