Horde: Reduce Docker/OCI image size by removing unused native libs

diff --git a/Engine/Source/Programs/Horde/Horde.Server/Dockerfile b/Engine/Source/Programs/Horde/Horde.Server/Dockerfile
--- a/Engine/Source/Programs/Horde/Horde.Server/Dockerfile
+++ b/Engine/Source/Programs/Horde/Horde.Server/Dockerfile
@@ -57,6 +57,12 @@ RUN redis-server --save "" --appendonly no --daemonize yes && \
 	dotnet test Source/Programs/Shared/EpicGames.Serialization.Tests/EpicGames.Serialization.Tests.csproj --blame-hang-timeout 2m --logger console;verbosity=normal && \
 	dotnet test Source/Programs/Horde/Horde.Server.Tests/Horde.Server.Tests.csproj --blame-hang-timeout 2m --logger console;verbosity=normal
 
+# Remove native libs not used on Linux x86_64
+RUN rm -rf /app/out/runtimes/osx* && \
+    rm -rf /app/out/runtimes/win* && \
+    rm -rf /app/out/runtimes/linux-arm* && \
+    rm -rf /app/out/runtimes/linux/native/libgrpc_csharp_ext.x86.so
+
 # Create the runtime image
 FROM mcr.microsoft.com/dotnet/aspnet:6.0 AS runtime
 WORKDIR /app
