Horde: Add support for executing a job batch inside a container on Linux

diff --git a/Engine/Source/Programs/Horde/Horde.Agent/Execution/JobExecutor.cs b/Engine/Source/Programs/Horde/Horde.Agent/Execution/JobExecutor.cs
--- a/Engine/Source/Programs/Horde/Horde.Agent/Execution/JobExecutor.cs
+++ b/Engine/Source/Programs/Horde/Horde.Agent/Execution/JobExecutor.cs
@@ -467,7 +467,7 @@ namespace Horde.Agent.Execution
 				arguments.Append($" CopyUAT -WithLauncher -TargetDir=\"{buildDir}\"");
 			}
 
-			int result = await ExecuteAutomationToolAsync(step, workspaceDir, arguments.ToString(), useP4, logger, cancellationToken);
+			int result = await ExecuteAutomationToolAsync(step, workspaceDir, sharedStorageDir, arguments.ToString(), useP4, logger, cancellationToken);
 			if (result != 0)
 			{
 				return false;
@@ -732,7 +732,7 @@ namespace Horde.Agent.Execution
 					arguments.AppendArgument("-SharedStorageDir=", sharedStorageDir.FullName);
 				}
 				
-				bool result = await ExecuteAutomationToolAsync(step, workspaceDir, arguments.ToString(), useP4, logger, cancellationToken) == 0;
+				bool result = await ExecuteAutomationToolAsync(step, workspaceDir, sharedStorageDir, arguments.ToString(), useP4, logger, cancellationToken) == 0;
 				await UploadXgeMonitorFilesAsync(step, logger, cancellationToken);
 				return result;
 			}
@@ -835,7 +835,7 @@ namespace Horde.Agent.Execution
 			}
 
 			// Run UAT
-			if (await ExecuteAutomationToolAsync(step, workspaceDir, arguments, useP4, logger, cancellationToken) != 0)
+			if (await ExecuteAutomationToolAsync(step, workspaceDir, null, arguments, useP4, logger, cancellationToken) != 0)
 			{
 				return false;
 			}
@@ -999,7 +999,7 @@ namespace Horde.Agent.Execution
 			return true;
 		}
 
-		protected async Task<int> ExecuteAutomationToolAsync(BeginStepResponse step, DirectoryReference workspaceDir, string arguments, bool? useP4, ILogger logger, CancellationToken cancellationToken)
+		protected async Task<int> ExecuteAutomationToolAsync(BeginStepResponse step, DirectoryReference workspaceDir, DirectoryReference? sharedStorageDir, string? arguments, bool? useP4, ILogger logger, CancellationToken cancellationToken)
 		{
 			int result;
 			using IScope scope = GlobalTracer.Instance.BuildSpan("BuildGraph").StartActive();
@@ -1016,7 +1016,7 @@ namespace Horde.Agent.Execution
 
 			if (RuntimeInformation.IsOSPlatform(OSPlatform.Windows))
 			{
-				result = await ExecuteCommandAsync(step, workspaceDir, Environment.GetEnvironmentVariable("COMSPEC") ?? "cmd.exe", $"/C \"\"{workspaceDir}\\Engine\\Build\\BatchFiles\\RunUAT.bat\" {arguments}\"", logger, cancellationToken);
+				result = await ExecuteCommandAsync(step, workspaceDir, sharedStorageDir, Environment.GetEnvironmentVariable("COMSPEC") ?? "cmd.exe", $"/C \"\"{workspaceDir}\\Engine\\Build\\BatchFiles\\RunUAT.bat\" {arguments}\"", logger, cancellationToken);
 			}
 			else if (RuntimeInformation.IsOSPlatform(OSPlatform.Linux))
 			{
@@ -1027,11 +1027,11 @@ namespace Horde.Agent.Execution
 					args = $"\"{workspaceDir}/Engine/Build/BatchFiles/RunWineUAT.sh\" {arguments}";
 				}
 				
-				result = await ExecuteCommandAsync(step, workspaceDir, "/bin/bash", args, logger, cancellationToken);
+				result = await ExecuteCommandAsync(step, workspaceDir, sharedStorageDir, "/bin/bash", args, logger, cancellationToken);
 			}
 			else if (RuntimeInformation.IsOSPlatform(OSPlatform.OSX))
 			{
-				result = await ExecuteCommandAsync(step, workspaceDir, "/bin/sh", $"\"{workspaceDir}/Engine/Build/BatchFiles/RunUAT.sh\" {arguments}", logger, cancellationToken);
+				result = await ExecuteCommandAsync(step, workspaceDir, sharedStorageDir, "/bin/sh", $"\"{workspaceDir}/Engine/Build/BatchFiles/RunUAT.sh\" {arguments}", logger, cancellationToken);
 			}
 			else
 			{
@@ -1163,7 +1163,86 @@ namespace Horde.Agent.Execution
 			}
 		}
 		
-		async Task<int> ExecuteCommandAsync(BeginStepResponse step, DirectoryReference workspaceDir, string fileName, string arguments, ILogger jobLogger, CancellationToken cancellationToken)
+		/// <summary>
+		/// Execute a process inside a Linux container
+		/// </summary>
+		/// <param name="arguments">Arguments</param>
+		/// <param name="mountDirs">Directories to mount inside the container for read/write</param>
+		/// <param name="newEnvironment">Environment variables</param>
+		/// <param name="filter">Log parser</param>
+		/// <param name="logger">Logger</param>
+		/// <param name="cancellationToken">Cancellation token for the call</param>
+		/// <returns></returns>
+		/// <exception cref="Exception"></exception>
+		async Task<int> ExecuteProcessInContainerAsync(string arguments, List<DirectoryReference> mountDirs, IReadOnlyDictionary<string, string>? newEnvironment, LogParser filter, ILogger logger, CancellationToken cancellationToken)
+		{
+			if (!RuntimeInformation.IsOSPlatform(OSPlatform.Linux))
+			{
+				throw new Exception("Only Linux is supported for executing a process inside a container");
+			}
+
+			if (String.IsNullOrEmpty(_jobOptions.Container.ImageUrl))
+			{
+				throw new Exception("Image URL is null or empty");
+			}
+
+			// Default to "docker" executable and assume it's available on the path
+			string executable = _jobOptions.Container.ContainerEngineExecutable ?? "docker";
+
+			uint linuxUid = LinuxInterop.getuid();
+			uint linuxGid = LinuxInterop.getgid();
+			
+			List<string> containerArgs = new()
+				{
+					"run",
+					"--tty", // Allocate a pseudo-TTY
+					$"--name horde-job-{_jobId}-{_batchId}", // Better name for debugging purposes
+					"--rm", // Ensure container is removed after run
+					$"--user {linuxUid}:{linuxGid}" // Run container as current user (important for mounted dirs)
+				};
+
+			foreach (DirectoryReference mountDir in mountDirs)
+			{
+				containerArgs.Add($"--volume {mountDir.FullName}:{mountDir.FullName}:rw");
+			}
+
+			if (newEnvironment != null)
+			{
+				string envFilePath = Path.GetTempFileName();
+				StringBuilder sb = new();
+				foreach ((string key, string value) in newEnvironment)
+				{
+					sb.AppendLine($"{key}={value}");
+				}
+				await File.WriteAllTextAsync(envFilePath, sb.ToString(), cancellationToken);
+				containerArgs.Add("--env-file=" + envFilePath);
+			}
+
+			if (!String.IsNullOrEmpty(_jobOptions.Container.ExtraArguments))
+			{
+				containerArgs.Add(_jobOptions.Container.ExtraArguments);
+			}
+			
+			containerArgs.Add(_jobOptions.Container.ImageUrl);
+			string containerArgStr = String.Join(' ', containerArgs);
+			arguments = containerArgStr + " " + arguments;
+			
+			logger.LogInformation("Executing {File} {Arguments} in container", executable.QuoteArgument(), arguments);
+			
+			// Skip forwarding of env vars as they are explicitly set above as arguments to container run
+			return await ExecuteProcessAsync(executable, arguments, new Dictionary<string, string>(), filter, logger, cancellationToken);
+		}
+
+		private static List<DirectoryReference> GetContainerMountDirs(DirectoryReference workspaceDir, DirectoryReference? sharedStorageDir, IReadOnlyDictionary<string, string> envVars)
+		{
+			List<DirectoryReference> dirs = new();
+			dirs.Add(workspaceDir);
+			if (sharedStorageDir != null) dirs.Add(sharedStorageDir);
+			if (envVars.TryGetValue("UE_SDKS_ROOT", out string? autoSdkDirPath)) dirs.Add(new DirectoryReference(autoSdkDirPath));
+			return dirs;
+		}
+		
+		async Task<int> ExecuteCommandAsync(BeginStepResponse step, DirectoryReference workspaceDir, DirectoryReference? sharedStorageDir, string fileName, string arguments, ILogger jobLogger, CancellationToken cancellationToken)
 		{
 			// Combine all the supplied environment variables together
 			Dictionary<string, string> newEnvVars = new Dictionary<string, string>(_envVars, StringComparer.Ordinal);
@@ -1246,7 +1325,15 @@ namespace Horde.Agent.Execution
 				await ExecuteCleanupScriptAsync(cleanupScript, filter, jobLogger);
 				try
 				{
-					exitCode = await ExecuteProcessAsync(fileName, arguments, newEnvVars, filter, jobLogger, cancellationToken);
+					if (RuntimeInformation.IsOSPlatform(OSPlatform.Linux) && _jobOptions.Container.Enabled is true)
+					{
+						List<DirectoryReference> mountDirs = GetContainerMountDirs(workspaceDir, sharedStorageDir, newEnvVars);
+						exitCode = await ExecuteProcessInContainerAsync(arguments, mountDirs, newEnvVars, filter, jobLogger, cancellationToken);
+					}
+					else
+					{
+						exitCode = await ExecuteProcessAsync(fileName, arguments, newEnvVars, filter, jobLogger, cancellationToken);	
+					}
 				}
 				finally
 				{
diff --git a/Engine/Source/Programs/Horde/Horde.Agent/Utility/LinuxInterop.cs b/Engine/Source/Programs/Horde/Horde.Agent/Utility/LinuxInterop.cs
new file mode 100644
--- /dev/null
+++ b/Engine/Source/Programs/Horde/Horde.Agent/Utility/LinuxInterop.cs
@@ -0,0 +1,25 @@
+// Copyright Epic Games, Inc. All Rights Reserved.
+
+using System.Runtime.InteropServices;
+
+namespace Horde.Agent.Utility;
+
+/// <summary>
+/// Linux specific P/Invoke calls
+/// </summary>
+static class LinuxInterop
+{
+	/// <summary>
+	/// Get user identity
+	/// </summary>
+	/// <returns>Real user ID of the calling process</returns>
+	[DllImport("libc", SetLastError = true)]
+	internal static extern uint getuid();
+
+	/// <summary>
+	/// Get group identity
+	/// </summary>
+	/// <returns>Real group ID of the calling process</returns>
+	[DllImport("libc", SetLastError = true)]
+	internal static extern uint getgid();
+}
