Horde: Allow job schedule times to be updated even when the schedule is not enabled.

diff --git a/Engine/Source/Programs/Horde/Horde.Server/Jobs/Schedules/ScheduleService.cs b/Engine/Source/Programs/Horde/Horde.Server/Jobs/Schedules/ScheduleService.cs
--- a/Engine/Source/Programs/Horde/Horde.Server/Jobs/Schedules/ScheduleService.cs
+++ b/Engine/Source/Programs/Horde/Horde.Server/Jobs/Schedules/ScheduleService.cs
@@ -246,7 +246,7 @@ namespace Horde.Server.Jobs.Schedules
 			}
 
 			ITemplateSchedule? schedule = templateRef.Schedule;
-			if (schedule == null)
+			if (schedule == null || !schedule.Config.Enabled)
 			{
 				return false;
 			}
diff --git a/Engine/Source/Programs/Horde/Horde.Server/Streams/StreamConfig.cs b/Engine/Source/Programs/Horde/Horde.Server/Streams/StreamConfig.cs
--- a/Engine/Source/Programs/Horde/Horde.Server/Streams/StreamConfig.cs
+++ b/Engine/Source/Programs/Horde/Horde.Server/Streams/StreamConfig.cs
@@ -972,15 +972,12 @@ namespace Horde.Server.Streams
 		public DateTime? GetNextTriggerTimeUtc(DateTime lastTimeUtc, TimeZoneInfo timeZone)
 		{
 			DateTime? nextTriggerTimeUtc = null;
-			if (Enabled)
+			foreach (SchedulePatternConfig pattern in Patterns)
 			{
-				foreach (SchedulePatternConfig pattern in Patterns)
+				DateTime patternTriggerTime = pattern.GetNextTriggerTimeUtc(lastTimeUtc, timeZone);
+				if (nextTriggerTimeUtc == null || patternTriggerTime < nextTriggerTimeUtc)
 				{
-					DateTime patternTriggerTime = pattern.GetNextTriggerTimeUtc(lastTimeUtc, timeZone);
-					if (nextTriggerTimeUtc == null || patternTriggerTime < nextTriggerTimeUtc)
-					{
-						nextTriggerTimeUtc = patternTriggerTime;
-					}
+					nextTriggerTimeUtc = patternTriggerTime;
 				}
 			}
 			return nextTriggerTimeUtc;
