Horde: Add null checks for job options.

diff --git a/Engine/Source/Programs/Horde/Horde.Agent/Execution/JobExecutor.cs b/Engine/Source/Programs/Horde/Horde.Agent/Execution/JobExecutor.cs
--- a/Engine/Source/Programs/Horde/Horde.Agent/Execution/JobExecutor.cs
+++ b/Engine/Source/Programs/Horde/Horde.Agent/Execution/JobExecutor.cs
@@ -1325,7 +1325,7 @@ namespace Horde.Agent.Execution
 				await ExecuteCleanupScriptAsync(cleanupScript, filter, jobLogger);
 				try
 				{
-					if (RuntimeInformation.IsOSPlatform(OSPlatform.Linux) && _jobOptions.Container.Enabled is true)
+					if (RuntimeInformation.IsOSPlatform(OSPlatform.Linux) && _jobOptions?.Container?.Enabled is true)
 					{
 						List<DirectoryReference> mountDirs = GetContainerMountDirs(workspaceDir, sharedStorageDir, newEnvVars);
 						exitCode = await ExecuteProcessInContainerAsync(arguments, mountDirs, newEnvVars, filter, jobLogger, cancellationToken);
