Horde: Artifacts Browserfor new artifacts storage backend
#rnx
#slipci

diff --git a/Engine/Source/Programs/Horde/HordeDashboard/src/backend/Backend.ts b/Engine/Source/Programs/Horde/HordeDashboard/src/backend/Backend.ts
--- a/Engine/Source/Programs/Horde/HordeDashboard/src/backend/Backend.ts
+++ b/Engine/Source/Programs/Horde/HordeDashboard/src/backend/Backend.ts
@@ -1,7 +1,7 @@
 // Copyright Epic Games, Inc. All Rights Reserved.
 
 import templateCache from '../backend/TemplateCache';
-import { AgentData, AgentQuery, ArtifactData, AuditLogEntry, AuditLogQuery, BatchUpdatePoolRequest, ChangeSummaryData, CreateDeviceRequest, CreateDeviceResponse, CreateExternalIssueRequest, CreateExternalIssueResponse, CreateJobRequest, CreateJobResponse, CreateNoticeRequest, CreatePoolRequest, CreateSoftwareResponse, CreateSubscriptionRequest, CreateSubscriptionResponse, DashboardPreference, DevicePoolTelemetryQuery, DeviceTelemetryQuery, EventData, FindIssueResponse, FindJobTimingsResponse, GetAgentSoftwareChannelResponse, GetArtifactZipRequest, GetDashboardConfigResponse, GetDevicePlatformResponse, GetDevicePoolResponse, GetDevicePoolTelemetryResponse, GetDeviceReservationResponse, GetDeviceResponse, GetDeviceTelemetryResponse, GetExternalIssueProjectResponse, GetExternalIssueResponse, GetGraphResponse, GetIssueStreamResponse, GetJobsTabResponse, GetJobStepRefResponse, GetJobStepTraceResponse, GetJobTimingResponse, GetLogEventResponse, GetNoticeResponse, GetNotificationResponse, GetPerforceServerStatusResponse, GetPoolResponse, GetServerInfoResponse, GetServerSettingsResponse, GetSoftwareResponse, GetSubscriptionResponse, GetTestDataDetailsResponse, GetTestDataRefResponse, GetTestMetaResponse, GetTestResponse, GetTestsRequest, GetTestStreamResponse, GetToolSummaryResponse, GetUserResponse, GetUtilizationTelemetryResponse, GlobalConfig, IssueData, IssueQuery, IssueQueryV2, JobData, JobQuery, JobsTabColumnType, JobStepOutcome, JobStreamQuery, JobTimingsQuery, LeaseData, LogData, LogLineData, PoolData, ProjectData, ScheduleData, ScheduleQuery, SearchLogFileResponse, ServerUpdateResponse, SessionData, StreamData, TabType, TestData, UpdateAgentRequest, UpdateDeviceRequest, UpdateGlobalConfigRequest, UpdateIssueRequest, UpdateJobRequest, UpdateLeaseRequest, UpdateNoticeRequest, UpdateNotificationsRequest, UpdatePoolRequest, UpdateServerSettingsRequest, UpdateStepRequest, UpdateStepResponse, UpdateTemplateRefRequest, UpdateUserRequest, UsersQuery } from './Api';
+import { AgentData, AgentQuery, ArtifactData, AuditLogEntry, AuditLogQuery, BatchUpdatePoolRequest, ChangeSummaryData, CreateDeviceRequest, CreateDeviceResponse, CreateExternalIssueRequest, CreateExternalIssueResponse, CreateJobRequest, CreateJobResponse, CreateNoticeRequest, CreatePoolRequest, CreateSoftwareResponse, CreateSubscriptionRequest, CreateSubscriptionResponse, CreateZipRequest, DashboardPreference, DevicePoolTelemetryQuery, DeviceTelemetryQuery, EventData, FindArtifactsResponse, FindIssueResponse, FindJobTimingsResponse, GetAgentSoftwareChannelResponse, GetArtifactDirectoryResponse, GetArtifactZipRequest, GetDashboardConfigResponse, GetDevicePlatformResponse, GetDevicePoolResponse, GetDevicePoolTelemetryResponse, GetDeviceReservationResponse, GetDeviceResponse, GetDeviceTelemetryResponse, GetExternalIssueProjectResponse, GetExternalIssueResponse, GetGraphResponse, GetIssueStreamResponse, GetJobsTabResponse, GetJobStepRefResponse, GetJobStepTraceResponse, GetJobTimingResponse, GetLogEventResponse, GetNoticeResponse, GetNotificationResponse, GetPerforceServerStatusResponse, GetPoolResponse, GetServerInfoResponse, GetServerSettingsResponse, GetSoftwareResponse, GetSubscriptionResponse, GetTestDataDetailsResponse, GetTestDataRefResponse, GetTestMetaResponse, GetTestResponse, GetTestsRequest, GetTestStreamResponse, GetToolSummaryResponse, GetUserResponse, GetUtilizationTelemetryResponse, GlobalConfig, IssueData, IssueQuery, IssueQueryV2, JobData, JobQuery, JobsTabColumnType, JobStepOutcome, JobStreamQuery, JobTimingsQuery, LeaseData, LogData, LogLineData, PoolData, ProjectData, ScheduleData, ScheduleQuery, SearchLogFileResponse, ServerUpdateResponse, SessionData, StreamData, TabType, TestData, UpdateAgentRequest, UpdateDeviceRequest, UpdateGlobalConfigRequest, UpdateIssueRequest, UpdateJobRequest, UpdateLeaseRequest, UpdateNoticeRequest, UpdateNotificationsRequest, UpdatePoolRequest, UpdateServerSettingsRequest, UpdateStepRequest, UpdateStepResponse, UpdateTemplateRefRequest, UpdateUserRequest, UsersQuery } from './Api';
 import dashboard from './Dashboard';
 import { ChallengeStatus, Fetch } from './Fetch';
 import graphCache, { GraphQuery } from './GraphCache';
@@ -526,6 +526,35 @@ export class Backend {
         });
     }
 
+    getJobArtifactsV2(ids?: string[], keys?: string[]): Promise<FindArtifactsResponse> {
+
+        const uniqueIds = Array.from(new Set(ids ?? []));
+        const uniqueKeys = Array.from(new Set(keys ?? []));
+
+        if (!uniqueIds.length && !uniqueKeys.length) {
+            throw new Error(`Must provide at least 1 id or key`);
+        }
+
+        return new Promise<FindArtifactsResponse>((resolve, reject) => {
+
+            this.backend.get(`/api/v2/artifacts`, { params: { id: uniqueIds, key: uniqueKeys } })
+                .then(response => { resolve(response.data); })
+                .catch(reason => reject(reason));
+        })
+    }
+
+
+    getBrowseArtifacts(id: string, path?: string): Promise<GetArtifactDirectoryResponse> {
+
+        return new Promise<GetArtifactDirectoryResponse>((resolve, reject) => {
+
+            this.backend.get(`/api/v2/artifacts/${id}/browse`, { params: { path: path } })
+                .then(response => { resolve(response.data); })
+                .catch(reason => reject(reason));
+        })
+    }
+
+
     getJobArtifacts(jobId: string, stepId?: string): Promise<ArtifactData[]> {
 
         const params: any = {
@@ -549,6 +578,46 @@ export class Backend {
         });
     }
 
+    downloadArtifactV2(artifactId: string, path: string, filename: string): Promise<boolean> {
+
+        return new Promise<any>((resolve, reject) => {
+            this.backend.get(`/api/v2/artifacts/${artifactId}/file`, { params: { path: path }, responseBlob: true }).then(response => {
+                const url = window.URL.createObjectURL(response.data);
+                const link = document.createElement('a');
+                link.href = url;
+                link.setAttribute('download', filename);
+                link.click();
+                setTimeout(() => {
+                    window.URL.revokeObjectURL(url);
+                }, 100);
+                resolve(true);
+            }).catch(reason => {
+                reject(reason);
+            });
+        });
+    }
+
+    downloadArtifactZipV2(artifactId: string, request: CreateZipRequest, filename: string): Promise<boolean> {
+
+        return new Promise<any>((resolve, reject) => {
+            this.backend.post(`/api/v2/artifacts/${artifactId}/zip`, request, { responseBlob: true }).then(response => {
+                const url = window.URL.createObjectURL(response.data);
+                const link = document.createElement('a');
+                link.href = url;
+                link.setAttribute('download', filename);
+                link.click();
+                setTimeout(() => {
+                    window.URL.revokeObjectURL(url);
+                }, 100);
+                resolve(true);
+            }).catch(reason => {
+                reject(reason);
+            });
+        });
+
+    }
+
+
     downloadJobArtifacts(request: GetArtifactZipRequest): Promise<boolean> {
         return new Promise<any>((resolve, reject) => {
             this.backend.post(`/api/v1/artifacts/zip`, request, { responseBlob: true }).then(response => {
@@ -1122,7 +1191,7 @@ export class Backend {
     }
 
     getTests(testIds: string[]): Promise<GetTestResponse[]> {
-        const request: GetTestsRequest = { testIds: testIds };        
+        const request: GetTestsRequest = { testIds: testIds };
         return new Promise<GetTestResponse[]>((resolve, reject) => {
             this.backend.post(`/api/v2/testdata/tests`, request).then((value) => {
                 resolve(value.data as GetTestResponse[]);
@@ -1544,7 +1613,7 @@ export class Backend {
             });
         });
     }
-    
+
     // update a notice
     updateNotice(request: UpdateNoticeRequest): Promise<void> {
         return new Promise<void>((resolve, reject) => {
diff --git a/Engine/Source/Programs/Horde/HordeDashboard/src/components/jobDetailsV2/JobDetailArtifactsV3.tsx b/Engine/Source/Programs/Horde/HordeDashboard/src/components/jobDetailsV2/JobDetailArtifactsV3.tsx
new file mode 100644
--- /dev/null
+++ b/Engine/Source/Programs/Horde/HordeDashboard/src/components/jobDetailsV2/JobDetailArtifactsV3.tsx
@@ -0,0 +1,561 @@
+// Copyright Epic Games, Inc. All Rights Reserved.
+
+import { DetailsList, DetailsListLayoutMode, Dropdown, Link as FluentLink, FontIcon, IColumn, IDropdownOption, IconButton, Modal, PrimaryButton, ScrollablePane, Selection, SelectionMode, SelectionZone, Spinner, SpinnerSize, Stack, Text, mergeStyleSets } from "@fluentui/react";
+import { action, makeObservable, observable } from "mobx";
+import { observer } from "mobx-react-lite";
+import { useEffect, useState } from "react";
+import backend from "../../backend";
+import { ArtifactContextType, GetArtifactDirectoryEntryResponse, GetArtifactDirectoryResponse, GetArtifactFileEntryResponse, GetArtifactResponseV2 } from "../../backend/Api";
+import { hordeClasses } from "../../styles/Styles";
+import { JobDetailsV2 } from "./JobDetailsViewCommon";
+
+
+enum BrowserType {
+   Directory,
+   File,
+   NavigateUp
+}
+
+type BrowserItem = {
+   key: string;
+   text: string;
+   icon?: string;
+   size?: number;
+   type: BrowserType;
+}
+
+
+class ArtifactsHandler {
+
+   constructor(details: JobDetailsV2, stepId: string) {
+      makeObservable(this);
+      this.details = details;
+      this.stepId = stepId;
+
+      this.set();
+   }
+
+   @observable
+   updated = 0
+
+   @action
+   updateReady() {
+      this.updated++;
+   }
+
+   async set() {
+
+      const details = this.details;
+      const jobData = details.jobData!;
+
+      const artifacts = details.stepArtifacts.get(this.stepId);
+      if (!artifacts) {
+         console.error(`Missing artifacts for job: ${jobData.id} step: ${this.stepId}`);
+         return;
+      }
+
+      this.artifacts = artifacts;
+
+      this.contexts = [];
+
+      let a = artifacts.find(a => a.type === "step-saved");
+      if (a) {
+         this.contexts.push("step-saved");
+         this.context = "step-saved";
+      }
+
+      a = artifacts.find(a => a.type === "step-output");
+      if (a) {
+         this.contexts.push("step-output");
+         this.context = this.context ?? "step-saved";
+      }
+
+      a = artifacts.find(a => a.type === "step-trace");
+      if (a) {
+         this.contexts.push("step-trace");
+         this.context = this.context ?? "step-trace";
+      }
+
+      if (!this.context) {
+         this.updateReady();
+         return;
+      }
+
+      a = artifacts.find(a => a.type === this.context)!;
+      this.artifact = a;
+
+      this.browse = await backend.getBrowseArtifacts(a.id);
+
+      this.updateReady();
+
+   }
+
+   hasContext(c: ArtifactContextType) {
+
+      return !!this.artifacts?.find(a => a.type === c);
+
+   }
+
+   async setContext(c: ArtifactContextType) {
+
+
+      let a = this.artifacts?.find(a => a.type === c)!;
+      if (!a) {
+         return;
+      }
+
+      this.context = c;
+
+      this.artifact = a;
+      this.path = undefined;
+      this.browse = await backend.getBrowseArtifacts(a.id);
+      this.updateReady();
+
+   }
+
+   async browseTo(path: string) {
+
+      if (!this.artifact) {
+         return;
+      }
+
+      this.browse = await backend.getBrowseArtifacts(this.artifact.id, path);
+
+      this.path = path;
+
+      this.updateReady();
+
+   }
+
+   get currentSelection(): { filesSelected: number, directoriesSelected: number, size: number, items: BrowserItem[] } {
+
+      let result = {
+         filesSelected: 0,
+         directoriesSelected: 0,
+         size: 0,
+         items: [] as BrowserItem[]
+      }
+
+      const browse = this.browse;
+      if (!browse) {
+         return result;
+      }
+
+      let selection = this.selection?.getSelection() as (BrowserItem[] | undefined);
+
+      if (!selection?.length) {
+         selection = [];
+
+         browse.directories.forEach(d => {
+            selection!.push({ key: d.ioHash, text: d.name, icon: "Folder", type: BrowserType.Directory, size: d.length });
+         });
+
+         browse.files.forEach(d => {
+            selection!.push({ key: d.ioHash, text: d.name, icon: "Document", type: BrowserType.File, size: d.length });
+         });
+      }
+
+      selection.forEach(b => {
+         result.items.push(b);
+         if (b.type === BrowserType.Directory) {
+            result.directoriesSelected++;
+            result.size += b.size ?? 0;
+         }
+         if (b.type === BrowserType.File) {
+            result.filesSelected++;
+            result.size += b.size ?? 0;
+         }
+      });
+
+      return result;
+
+   }
+
+   get directories(): GetArtifactDirectoryEntryResponse[] {
+      return this.browse?.directories ?? [];
+   }
+
+   get files(): GetArtifactFileEntryResponse[] {
+      return this.browse?.files ?? [];
+   }
+
+   clear() {
+      this.selection = undefined;
+      this.selectionCallback = undefined;
+      this.path = undefined;
+      this.browse = undefined;
+      this.artifact = undefined;
+      this.artifacts = undefined;
+      this.context = undefined;
+      this.contexts = undefined;
+      this.stepId = "";
+   }
+
+
+   details: JobDetailsV2;
+
+   selectionCallback?: () => void;
+
+   path?: string;
+
+   selection?: Selection = new Selection({ onSelectionChanged: () => { if (this.selectionCallback) this.selectionCallback() } });
+
+   browse?: GetArtifactDirectoryResponse;
+
+   artifact?: GetArtifactResponseV2;
+   artifacts?: GetArtifactResponseV2[];
+
+   contexts?: ArtifactContextType[];
+
+   context?: ArtifactContextType;
+
+   stepId: string;
+
+
+}
+
+const styles = mergeStyleSets({
+   list: {
+      selectors: {
+         '.ms-DetailsRow': {
+            border: "1px solid #EDEBE9"
+         }
+      }
+   }
+});
+
+
+const BrowseBreadCrumbs: React.FC<{ handler: ArtifactsHandler }> = observer(({ handler }) => {
+
+   // subscribe
+   if (handler.updated) { }
+
+   const fontSize = 16;
+
+   if (!handler.path) {
+      return <Stack><Text style={{ fontSize: fontSize, paddingLeft: 4, paddingRight: 4 }}>/
+      </Text>
+      </Stack>
+   }
+
+   const ppath = handler.path.split("/");
+
+   const elements = ppath.map((e, index) => {
+      const path = ppath.slice(0, index + 1).join("/");
+      return <Stack horizontal onClick={() => handler.browseTo(path)} style={{ cursor: "pointer" }}>
+         <Stack>
+            <Text style={{ fontSize: fontSize }}>{e}</Text>
+         </Stack>
+         {index !== (ppath.length - 1) && <Stack style={{ paddingLeft: 4, paddingRight: 4 }}><Text style={{ fontSize: fontSize }}>/</Text></Stack>}
+      </Stack>
+   });
+
+   elements.unshift(<Stack style={{ cursor: "pointer" }} onClick={() => handler.browseTo("")}>
+      <Stack style={{ paddingLeft: 4, paddingRight: 4 }}>
+         <Text style={{ fontSize: fontSize }}>/</Text>
+      </Stack>
+   </Stack>);
+
+   return <Stack horizontal>{elements}</Stack>
+
+})
+
+function formatBytes(bytes: number, decimals = 2) {
+   if (!+bytes) return '0 Bytes'
+
+   const k = 1024
+   const dm = decimals < 0 ? 0 : decimals
+   const sizes = ['Bytes', 'KB', 'MB', 'GB', 'TB', 'PB', 'EB', 'ZB', 'YB']
+
+   const i = Math.floor(Math.log(bytes) / Math.log(k))
+
+   return `${parseFloat((bytes / Math.pow(k, i)).toFixed(dm))} ${sizes[i]}`
+}
+
+
+const DownloadButton: React.FC<{ handler: ArtifactsHandler }> = observer(({ handler }) => {
+
+   const [downloading, setDownloading] = useState(false);
+   const [selectKey, setSelectionKey] = useState(0);
+
+
+   // subscribe
+   if (handler.updated) { };
+
+   const browse = handler.browse;
+
+   if (!browse) {
+      return null;
+   }
+
+   handler.selectionCallback = () => {
+      setSelectionKey(selectKey + 1);
+   }
+
+   let buttonText = "Download";
+
+
+   const selection = handler.currentSelection;
+
+
+   let sizeText = "0KB";
+   if (selection.size) {
+      sizeText = formatBytes(selection.size, (selection.size < (1024 * 1024)) ? 0 : 1)
+   }
+
+   if (selection.directoriesSelected || selection.filesSelected > 1) {
+      buttonText = `Download Zip (${sizeText})`;
+   } else if (selection.filesSelected === 1) {
+      buttonText = `Download File (${sizeText})`;
+   }
+
+   return <Stack>
+      <PrimaryButton styles={{ root: { fontFamily: 'Horde Open Sans Semibold !important' } }} disabled={!selection.size} onClick={async () => {
+
+         const selection = handler.currentSelection.items;
+
+         if (!selection?.length || !handler.artifact) {
+            return;
+         }
+
+         // download a single file
+         if (selection.length === 1) {
+            const item = selection[0] as BrowserItem;
+            if (item.type === BrowserType.File) {
+
+               try {
+                  setDownloading(true);
+                  await backend.downloadArtifactV2(handler.artifact.id, (handler.path ? handler.path + "/" : "") + item.text, item.text);
+               } catch (err) {
+                  console.error(err);
+               } finally {
+                  setDownloading(false);
+               }
+
+               return;
+            }
+         }
+
+         const path = handler.path ? handler.path + "/" : "";
+
+         const filters = selection.map(s => {
+            const item = s as BrowserItem;
+
+            if (item.type === BrowserType.NavigateUp) {
+               return "";
+            }
+
+            if (item.type === BrowserType.Directory) {
+               return `${path}${item.text}/...`;
+            }
+            return `${path}${item.text}`;
+         }).filter(f => !!f);
+
+         let context = handler.context;
+         let contextName = "step";
+         if (context === "step-output") {
+            contextName = "output";
+         }
+
+         if (context === "step-trace") {
+            contextName = "trace";
+         }
+
+         const filename = "horde-" + contextName + "-artifacts-" + handler.details!.jobId! + '-' + handler.stepId + ".zip";
+
+         try {
+            setDownloading(true);
+            await backend.downloadArtifactZipV2(handler.artifact.id, { filter: filters }, filename);
+         } catch (err) {
+            console.error(err);
+         } finally {
+            setDownloading(false);
+         }
+
+      }}>{buttonText}{!!downloading && <Spinner style={{ paddingLeft: 12 }} size={SpinnerSize.medium} />}</PrimaryButton>
+   </Stack>
+
+});
+
+const JobDetailArtifactsInner: React.FC<{ handler: ArtifactsHandler }> = observer(({ handler }) => {
+
+   // subscribe
+   if (handler.updated) { }
+
+   const browse = handler.browse;
+
+   if (!browse) {
+      return null;
+   }
+
+   const coptions: IDropdownOption[] = [];
+   if (handler.hasContext("step-saved")) {
+      coptions.push({ key: 'step-saved', text: 'Step' })
+   }
+   if (handler.hasContext("step-output")) {
+      coptions.push({ key: 'step-output', text: 'Output' })
+   }
+
+   if (handler.hasContext("step-trace")) {
+      coptions.push({ key: 'step-trace', text: 'Trace' })
+   }
+
+   const items: BrowserItem[] = [];
+
+   if (handler.path?.length) {
+      items.push({ key: "navigate up", text: "..", type: BrowserType.NavigateUp });
+   }
+
+   browse.directories.forEach(d => {
+      items.push({ key: d.ioHash, text: d.name, icon: "Folder", type: BrowserType.Directory, size: d.length });
+   });
+
+   browse.files.forEach(d => {
+      items.push({ key: d.ioHash, text: d.name, icon: "Document", type: BrowserType.File, size: d.length });
+   });
+
+   const columns: IColumn[] = [
+      { key: 'column1', name: 'Name', minWidth: 640, maxWidth: 640, isResizable: false, isPadded: false },
+      { key: 'column2', name: 'Size', minWidth: 128, maxWidth: 128, isResizable: false, isPadded: false },
+   ];
+
+   const renderItem = (item: any, index?: number, column?: IColumn) => {
+
+      if (!column) {
+         return null;
+      }
+
+      if (column.name === "Size") {
+         if (!item.size) {
+            return null;
+         }
+         return <Stack>
+            <Text>{formatBytes(item.size, (item.size < (1024 * 1024)) ? 0 : 1)}</Text>
+         </Stack>
+
+      }
+
+      if (column.name === "Name") {
+
+         let href: string | undefined;
+         let target: string | undefined;
+
+         if (item.type === BrowserType.File) {
+            target = "_blank";
+            const path = encodeURI(handler.path + "/" + item.text);
+            const server = backend.serverUrl;
+            href = `${server}/api/v2/artifacts/${handler.artifact!.id}/file?path=${path}&inline=true`;
+         }
+
+         return <FluentLink href={href} target={target} style={{ color: "unset", height: "100%", width: "100%" }}><Stack style={{ cursor: "pointer" }} onClick={(ev) => {
+
+            if (item.type === BrowserType.Directory) {
+               const nbrowse = handler.path ? `${handler.path}/${item.text}` : item.text;
+               handler.browseTo(nbrowse);
+            }
+            if (item.type === BrowserType.NavigateUp && handler.path) {
+               const nbrowse = handler.path.split("/")
+               nbrowse.pop();
+               handler.browseTo(nbrowse.join("/"));
+            }
+
+         }}>
+            {item.type !== BrowserType.NavigateUp && <Stack data-selection-disabled horizontal tokens={{ childrenGap: 8 }}>
+               <Stack>
+                  <FontIcon style={{ fontSize: 16 }} iconName={item.icon} />
+               </Stack>
+               <Stack>
+                  <Text>{item.text}</Text>
+               </Stack>
+            </Stack>}
+            {item.type === BrowserType.NavigateUp && <Stack data-selection-disabled style={{ height: "100%", width: "100%" }}>
+               <Text>..</Text>
+            </Stack>
+            }
+         </Stack>
+         </FluentLink>
+      }
+
+      return null;
+
+   }
+
+   return <Stack tokens={{ childrenGap: 12 }}>
+      <Stack horizontal verticalAlign="center" tokens={{ childrenGap: 18 }} style={{ paddingBottom: 12 }}>
+         <Stack>
+            <Dropdown className={hordeClasses.modal} style={{ width: 128 }} options={coptions} selectedKey={handler.context} onChange={(ev, option) => {
+               if (option?.key) {
+                  handler.setContext(option.key as ArtifactContextType);
+               }
+            }} />
+         </Stack>
+         <Stack>
+            <BrowseBreadCrumbs handler={handler} />
+         </Stack>
+         <Stack grow />
+         <DownloadButton handler={handler} />
+      </Stack >
+      <Stack style={{ height: 492 + 160, position: "relative" }}>
+         <ScrollablePane style={{ height: 492 + 160 }}>
+            <SelectionZone selection={handler.selection!}>
+               <DetailsList
+                  className={styles.list}
+                  isHeaderVisible={false}
+                  compact={true}
+                  items={items}
+                  columns={columns}
+                  layoutMode={DetailsListLayoutMode.fixedColumns}
+                  selectionMode={SelectionMode.multiple}
+                  selection={handler.selection}
+                  selectionPreservedOnEmptyClick={true}
+                  //onItemInvoked={this._onItemInvoked} <--- double click*/
+                  onRenderItemColumn={renderItem}
+               />
+            </SelectionZone>
+
+         </ScrollablePane>
+      </Stack>
+   </Stack >
+
+   //
+})
+
+
+
+export const JobArtifactsModal: React.FC<{ jobDetails: JobDetailsV2; stepId: string, onClose: () => void }> = ({ jobDetails, stepId, onClose }) => {
+
+   const [handler] = useState(new ArtifactsHandler(jobDetails, stepId));
+
+   useEffect(() => {
+      return () => {
+         handler?.clear();
+      };
+   }, [handler]);
+
+   return <Stack>
+      <Modal isOpen={true} isBlocking={true} topOffsetFixed={true} styles={{ main: { padding: 8, width: 1180, height: 820, hasBeenOpened: false, top: "80px", position: "absolute" } }} onDismiss={() => onClose()} className={hordeClasses.modal}>
+         <Stack styles={{ root: { paddingTop: 18, paddingRight: 12 } }}>
+            <Stack style={{ paddingLeft: 24, paddingRight: 24 }}>
+               <Stack tokens={{ childrenGap: 12 }} style={{ height: 800 }}>
+                  <Stack horizontal>
+                     <Stack>
+                        <Text variant="mediumPlus" styles={{ root: { fontFamily: "Horde Open Sans SemiBold" } }}>Artifacts</Text>
+                     </Stack>
+                     <Stack grow />
+                     <Stack horizontalAlign="end">
+                        <IconButton
+                           iconProps={{ iconName: 'Cancel' }}
+                           onClick={() => { onClose() }}
+                        />
+                     </Stack>
+
+                  </Stack>
+                  <Stack styles={{ root: { paddingLeft: 4, paddingRight: 0, paddingTop: 8, paddingBottom: 4 } }}>
+                     <JobDetailArtifactsInner handler={handler} />
+                  </Stack>
+               </Stack>
+            </Stack>
+         </Stack>
+      </Modal>
+   </Stack>;
+};
\ No newline at end of file
diff --git a/Engine/Source/Programs/Horde/HordeDashboard/src/components/jobDetailsV2/JobDetailsViewCommon.tsx b/Engine/Source/Programs/Horde/HordeDashboard/src/components/jobDetailsV2/JobDetailsViewCommon.tsx
--- a/Engine/Source/Programs/Horde/HordeDashboard/src/components/jobDetailsV2/JobDetailsViewCommon.tsx
+++ b/Engine/Source/Programs/Horde/HordeDashboard/src/components/jobDetailsV2/JobDetailsViewCommon.tsx
@@ -5,7 +5,7 @@ import moment from "moment";
 import React, { useState } from "react";
 import { useLocation, useNavigate } from "react-router-dom";
 import backend from "../../backend";
-import { BatchData, GetJobTimingResponse, GetLabelResponse, GetTemplateRefResponse, GroupData, JobData, JobState, LabelState, NodeData, ReportPlacement, StepData, StreamData } from "../../backend/Api";
+import { BatchData, GetArtifactResponseV2, GetJobTimingResponse, GetLabelResponse, GetTemplateRefResponse, GroupData, JobData, JobState, LabelState, NodeData, ReportPlacement, StepData, StreamData } from "../../backend/Api";
 import { JobLabel } from "../../backend/JobDetails";
 import { PollBase } from "../../backend/PollBase";
 import { projectStore } from "../../backend/ProjectStore";
@@ -717,6 +717,9 @@ export class JobDetailsV2 extends PollBase {
    nodes: NodeData[] = [];
    batches: BatchData[] = [];
    template?: GetTemplateRefResponse;
+   
+   // stepId => artifacts
+   stepArtifacts = new Map<string, GetArtifactResponseV2[]>();
 
    // artifact id => report contents (markdown)
    private reportData = new Map<string, string>();
diff --git a/Engine/Source/Programs/Horde/HordeDashboard/src/components/jobDetailsV2/JobOperationsBar.tsx b/Engine/Source/Programs/Horde/HordeDashboard/src/components/jobDetailsV2/JobOperationsBar.tsx
--- a/Engine/Source/Programs/Horde/HordeDashboard/src/components/jobDetailsV2/JobOperationsBar.tsx
+++ b/Engine/Source/Programs/Horde/HordeDashboard/src/components/jobDetailsV2/JobOperationsBar.tsx
@@ -13,6 +13,7 @@ import { NewBuild } from '../NewBuild';
 import { NotificationDropdown } from '../NotificationDropdown';
 import { PauseStepModal } from '../StepPauseModal';
 import { AbortJobModal } from './AbortJobModal';
+import { JobArtifactsModal } from './JobDetailArtifactsV3';
 import { JobDetailsV2 } from './JobDetailsViewCommon';
 import { StepRetryModal, StepRetryType } from './StepRetryModal';
 
@@ -28,6 +29,8 @@ export const JobOperations: React.FC<{ jobDetails: JobDetailsV2 }> = observer(({
    const navigate = useNavigate();
    const [abortShown, setAbortShown] = useState(false);
    const [editShown, setEditShown] = useState(false);
+   const [artifactsShown, setArtifactsShown] = useState(false);
+
    const [parametersState, setParametersState] = useState(query.get("newbuild") ? ParameterState.Clone : ParameterState.Hidden);
 
    const stepId = query.get("step") ? query.get("step")! : undefined;
@@ -44,6 +47,8 @@ export const JobOperations: React.FC<{ jobDetails: JobDetailsV2 }> = observer(({
       return null;
    }
 
+   const stepArtifacts = jobDetails.stepArtifacts.get(stepId ?? "");
+
    const abortDisabled = jobData.state === JobState.Complete;
    const runAgainDisabled = false; /*jobDetails.jobdata?.state !== JobState.Complete*/
 
@@ -117,6 +122,7 @@ export const JobOperations: React.FC<{ jobDetails: JobDetailsV2 }> = observer(({
    }
 
    return <Stack>
+      {artifactsShown && !!step && <JobArtifactsModal stepId={step.id} jobDetails={jobDetails} onClose={() => { setArtifactsShown(false); }} />}
       <AbortJobModal jobDetails={jobDetails} show={abortShown} onClose={() => { setAbortShown(false); }} />
       <EditJobModal jobData={jobDetails.jobData} show={editShown} onClose={() => { setEditShown(false); }} />
       <NewBuild streamId={jobDetails.stream!.id} jobDetails={jobDetails} readOnly={parametersState === ParameterState.Parameters}
@@ -141,6 +147,12 @@ export const JobOperations: React.FC<{ jobDetails: JobDetailsV2 }> = observer(({
             </Stack>
             </Link>}
 
+            {!!step && !!stepArtifacts?.length && <div onClick={() => { setArtifactsShown(true) }}>
+               <Stack styles={{ root: { paddingTop: 4 } }}>
+                  <CommandBarButton disabled={viewLogDisabled} className={hordeClasses.commandBarSmall} styles={{ root: { padding: "10px 8px 10px 8px" } }} iconProps={{ iconName: "CloudDownload" }} text="Artifacts" />
+               </Stack>
+            </div>}
+
             <Stack className={hordeClasses.commandBarSmall} styles={{ root: { paddingTop: 4 } }}>
                {notifications}
             </Stack>
@@ -160,6 +172,7 @@ export const JobOperations: React.FC<{ jobDetails: JobDetailsV2 }> = observer(({
                   <CommandBarButton disabled={viewLogDisabled} className={hordeClasses.commandBarSmall} styles={{ root: { padding: "10px 8px 10px 8px" } }} iconProps={{ iconName: "AlignLeft" }} text="View Log" />
                </Stack>
             </Link>}
+
          </Stack>
       </Stack>
 
diff --git a/Engine/Source/Programs/Horde/HordeDashboard/src/components/jobDetailsV2/StepDetailView.tsx b/Engine/Source/Programs/Horde/HordeDashboard/src/components/jobDetailsV2/StepDetailView.tsx
--- a/Engine/Source/Programs/Horde/HordeDashboard/src/components/jobDetailsV2/StepDetailView.tsx
+++ b/Engine/Source/Programs/Horde/HordeDashboard/src/components/jobDetailsV2/StepDetailView.tsx
@@ -2,7 +2,7 @@
 
 import { Stack } from "@fluentui/react";
 import { observer } from "mobx-react-lite";
-import { useState } from "react";
+import { useEffect, useState } from "react";
 import { useNavigate } from "react-router-dom";
 import { hordeClasses } from "../../styles/Styles";
 import { HistoryModal } from "../HistoryModal";
@@ -17,6 +17,8 @@ import { StepsPanelV2 } from "./JobDetailViewSteps";
 import { StepSummaryPanel } from "./StepDetailSummary";
 import { StepTestReportPanel } from "./StepDetailTestPanel";
 import { StepErrorPanel } from "./StepErrorPanel";
+import backend from "../../backend";
+import { GetArtifactResponseV2 } from "../../backend/Api";
 
 
 class StepDetailDataView extends JobDataView {
@@ -31,9 +33,74 @@ class StepDetailDataView extends JobDataView {
          return;
       }
 
+      this.queryArtifacts();
+
+      this.updateReady();
+
+   }
+
+   async queryArtifacts() {
+
+      if (this.artifacts) {
+         return;
+      }
+
+      if (!this.details?.jobData) {
+         return;
+      }
+
+      if (!this.details.jobData.useArtifactsV2) {
+         this.artifacts = [];
+         return;
+      }
+
+      const stepId = this.stepId!;
+
+      let artifacts = this.details.stepArtifacts.get(stepId);
+      if (!artifacts) {
+         const key = `job:${this.details!.jobData!.id}/step:${this.stepId}`;
+         try {
+            const v = await backend.getJobArtifactsV2(undefined, [key]);
+            artifacts = v.artifacts;
+            this.details.stepArtifacts.set(stepId, artifacts);
+
+            // need to notify things like the operations bar
+            this.details.externalUpdate();
+         } catch (err) {
+            console.error(err);
+         } finally {
+            artifacts = [];
+         }
+         this.artifacts = artifacts;
+      }
+
+
+   }
+
+   async set(stepId: string) {
+
+      if (this.stepId === stepId) {
+         return;
+      }
+
+      this.stepId = stepId;
+
+      this.queryArtifacts();
+
       this.updateReady();
+
+   }
+
+   clear() {
+      this.stepId = undefined;
+      this.artifacts = undefined;
+      super.clear();
    }
 
+   artifacts?: GetArtifactResponseV2[];
+
+   stepId?: string;
+
 }
 
 JobDetailsV2.registerDataView("StepDetailDataView", (details: JobDetailsV2) => new StepDetailDataView(details));
@@ -46,10 +113,19 @@ const StepDetailViewInner: React.FC<{ jobDetails: JobDetailsV2, stepId: string }
 
    const dataView = jobDetails.getDataView<StepDetailDataView>("StepDetailDataView");
 
-   dataView.initialize();
+   useEffect(() => {
+      return () => {
+         dataView?.clear();
+      };
+   }, [dataView]);
+
 
    dataView.subscribe();
 
+   dataView.initialize();
+
+   dataView.set(stepId);
+
    if (!historyAgentId && query.get("agentId")) {
       historyAgentId = query.get("agentId")!;
    }
@@ -78,7 +154,7 @@ const StepDetailViewInner: React.FC<{ jobDetails: JobDetailsV2, stepId: string }
          <StepErrorPanel jobDetails={jobDetails} stepId={stepId} showErrors={false} />
       </Stack>
       <Stack>
-         <StepsPanelV2 jobDetails={jobDetails} depStepId={stepId}/>
+         <StepsPanelV2 jobDetails={jobDetails} depStepId={stepId} />
       </Stack>
       <Stack>
          <StepTestReportPanel jobDetails={jobDetails} stepId={stepId} />
@@ -86,9 +162,9 @@ const StepDetailViewInner: React.FC<{ jobDetails: JobDetailsV2, stepId: string }
       {<Stack>
          <StepHistoryPanel jobDetails={jobDetails} stepId={stepId} />
       </Stack>}
-      <Stack>
+      {!jobData.useArtifactsV2 && <Stack>
          <JobDetailArtifactsV2 jobDetails={jobDetails} stepId={stepId} />
-      </Stack>
+      </Stack>}
       {!!step && <Stack>
          <TimelinePanel jobDetails={jobDetails} stepId={stepId} />
       </Stack>}
diff --git a/Engine/Source/Programs/Horde/HordeDashboard/src/styles/Styles.tsx b/Engine/Source/Programs/Horde/HordeDashboard/src/styles/Styles.tsx
--- a/Engine/Source/Programs/Horde/HordeDashboard/src/styles/Styles.tsx
+++ b/Engine/Source/Programs/Horde/HordeDashboard/src/styles/Styles.tsx
@@ -122,7 +122,10 @@ registerIcons({
       "More": IconContents.MORE,
       "Share": IconContents.SHARE,
       "Wrench": IconContents.WRENCH,
-      "Dashboard": IconContents.DASHBOARD
+      "Dashboard": IconContents.DASHBOARD,
+      "Document": IconContents.DOCUMENT,
+      "Folder": IconContents.FOLDER_CLOSE,
+      "CloudDownload" : IconContents.CLOUD_DOWNLOAD
    }
 });
 
