Horde: Always show artifacts button, disable when no artifacts
#skipci
#rnx

diff --git a/Engine/Source/Programs/Horde/HordeDashboard/src/components/jobDetailsV2/JobOperationsBar.tsx b/Engine/Source/Programs/Horde/HordeDashboard/src/components/jobDetailsV2/JobOperationsBar.tsx
--- a/Engine/Source/Programs/Horde/HordeDashboard/src/components/jobDetailsV2/JobOperationsBar.tsx
+++ b/Engine/Source/Programs/Horde/HordeDashboard/src/components/jobDetailsV2/JobOperationsBar.tsx
@@ -147,9 +147,9 @@ export const JobOperations: React.FC<{ jobDetails: JobDetailsV2 }> = observer(({
             </Stack>
             </Link>}
 
-            {!!step && !!stepArtifacts?.length && <div onClick={() => { setArtifactsShown(true) }}>
+            {!!jobData?.useArtifactsV2 && <div onClick={() => { setArtifactsShown(true) }}>
                <Stack styles={{ root: { paddingTop: 4 } }}>
-                  <CommandBarButton disabled={viewLogDisabled} className={hordeClasses.commandBarSmall} styles={{ root: { padding: "10px 8px 10px 8px" } }} iconProps={{ iconName: "CloudDownload" }} text="Artifacts" />
+                  <CommandBarButton disabled={!stepArtifacts?.length} className={hordeClasses.commandBarSmall} styles={{ root: { padding: "10px 8px 10px 8px" } }} iconProps={{ iconName: "CloudDownload" }} text="Artifacts" />
                </Stack>
             </div>}
 
