Horde: Update C++ remote execution sample to work with new named buffers.

diff --git a/Engine/Source/Programs/Horde/Samples/RemoteWorkerCpp/ComputeChannel.cpp b/Engine/Source/Programs/Horde/Samples/RemoteWorkerCpp/ComputeChannel.cpp
new file mode 100644
--- /dev/null
+++ b/Engine/Source/Programs/Horde/Samples/RemoteWorkerCpp/ComputeChannel.cpp
@@ -0,0 +1,99 @@
+// Copyright Epic Games, Inc. All Rights Reserved.
+
+#include <windows.h>
+#include <iostream>
+#include <assert.h>
+#include "ComputeChannel.h"
+
+const char* const FComputeChannel::EnvVarName = "UE_HORDE_COMPUTE_IPC";
+
+FComputeChannel::FComputeChannel()
+{
+}
+
+FComputeChannel::~FComputeChannel()
+{
+	Close();
+}
+
+bool FComputeChannel::Open()
+{
+	Close();
+
+	char EnvVar[MAX_PATH];
+	int Length = GetEnvironmentVariableA(EnvVarName, EnvVar, sizeof(EnvVar) / sizeof(EnvVar[0]));
+
+	if (Length <= 0 || Length >= sizeof(EnvVar))
+	{
+		std::cout << "No env var" << std::endl;
+		Close();
+		return false;
+	}
+
+	char Buffer[MAX_PATH];
+
+	sprintf_s(Buffer, "%s_SEND", EnvVar);
+	if (!SendBuffer.OpenExisting(Buffer))
+	{
+		std::cout << "No send: " << Buffer << std::endl;
+		Close();
+		return false;
+	}
+
+	sprintf_s(Buffer, "%s_RECV", EnvVar);
+	if (!RecvBuffer.OpenExisting(Buffer))
+	{
+		std::cout << "No recv: " << Buffer << std::endl;
+		Close();
+		return false;
+	}
+
+	return true;
+}
+
+void FComputeChannel::Close()
+{
+	SendBuffer.Close();
+	RecvBuffer.Close();
+}
+
+void FComputeChannel::Send(const void* Data, size_t Length)
+{
+	while (Length > 0)
+	{
+		size_t Size;
+		unsigned char* WriteMemory = SendBuffer.GetWriteMemory(Size);
+
+		if (Size == 0)
+		{
+			SendBuffer.WaitToWrite(0);
+			continue;
+		}
+
+		size_t CopyLength = (Size < Length) ? Size : Length;
+		memcpy(WriteMemory, Data, CopyLength);
+		SendBuffer.AdvanceWritePosition(CopyLength);
+
+		Data = (unsigned char*)Data + CopyLength;
+		Length -= CopyLength;
+	}
+}
+
+size_t FComputeChannel::Receive(void* Data, size_t Length)
+{
+	size_t Size;
+	const unsigned char* ReadMemory = RecvBuffer.GetReadMemory(Size);
+
+	while (Size == 0 && !RecvBuffer.IsComplete())
+	{
+		RecvBuffer.WaitToRead(0);
+		ReadMemory = RecvBuffer.GetReadMemory(Size);
+	}
+
+	size_t CopyLength = (Size < Length) ? Size : Length;
+	memcpy(Data, ReadMemory, CopyLength);
+	RecvBuffer.AdvanceReadPosition(CopyLength);
+
+	return CopyLength;
+}
+
diff --git a/Engine/Source/Programs/Horde/Samples/RemoteWorkerCpp/ComputeChannel.h b/Engine/Source/Programs/Horde/Samples/RemoteWorkerCpp/ComputeChannel.h
new file mode 100644
--- /dev/null
+++ b/Engine/Source/Programs/Horde/Samples/RemoteWorkerCpp/ComputeChannel.h
@@ -0,0 +1,24 @@
+// Copyright Epic Games, Inc. All Rights Reserved.
+
+#pragma once
+
+#include "SharedMemoryBuffer.h"
+
+class FComputeChannel
+{
+	static const char* const EnvVarName;
+
+public:
+	FComputeChannel();
+	~FComputeChannel();
+
+	bool Open();
+	void Close();
+
+	void Send(const void* Data, size_t Length);
+	size_t Receive(void* Data, size_t Length);
+
+private:
+	FSharedMemoryBuffer RecvBuffer;
+	FSharedMemoryBuffer SendBuffer;
+};
diff --git a/Engine/Source/Programs/Horde/Samples/RemoteWorkerCpp/Main.cpp b/Engine/Source/Programs/Horde/Samples/RemoteWorkerCpp/Main.cpp
--- a/Engine/Source/Programs/Horde/Samples/RemoteWorkerCpp/Main.cpp
+++ b/Engine/Source/Programs/Horde/Samples/RemoteWorkerCpp/Main.cpp
@@ -1,44 +1,38 @@
 // Copyright Epic Games, Inc. All Rights Reserved.
 
 #include <iostream>
-#include "WorkerSocket.h"
+#include "ComputeChannel.h"
 #include "SharedMemoryBuffer.h"
 
 int main()
 {
-	FWorkerSocket Socket;
-	if (!Socket.Open())
+	FComputeChannel Channel;
+	if (!Channel.Open())
 	{
 		std::cout << "Environment variable not set correctly" << std::endl;
 		return 1;
 	}
 
-	FSharedMemoryBuffer Buffer;
-	if (!Buffer.CreateNew(1024 * 1024))
-	{
-		std::cout << "Couldn't create buffer" << std::endl;
-		return 1;
-	}
-	if (!Socket.TryAttachRecvBuffer(1, &Buffer))
-	{
-		std::cout << "Failed to attach recieve buffer" << std::endl;
-		return 1;
-	}
-
 	std::cout << "Connected to client" << std::endl; // // Client will wait for output before sending data on channel 1
 
-	while (!Buffer.HasFinishedReading())
-	{
-		long long Size;
-		const unsigned char* Memory = Buffer.GetReadMemory(Size);
+	size_t Length = 0;
+	char Buffer[4];
 
-		if (Size > 0)
+	for (;;)
+	{
+		size_t Read = Channel.Receive(Buffer + Length, sizeof(Buffer) - Length);
+		if (Read == 0)
 		{
-			std::cout << "Read value " << (int)*Memory << std::endl;
-			Buffer.AdvanceReadPosition(1);
+			return 0;
 		}
 
-		Buffer.WaitForData(0);
+		Length += Read;
+
+		if (Length >= 4)
+		{
+			std::cout << "Read value " << *(int*)Buffer << std::endl;
+			Length = 0;
+		}
 	}
 
 	return 0;
diff --git a/Engine/Source/Programs/Horde/Samples/RemoteWorkerCpp/SharedMemoryBuffer.cpp b/Engine/Source/Programs/Horde/Samples/RemoteWorkerCpp/SharedMemoryBuffer.cpp
--- a/Engine/Source/Programs/Horde/Samples/RemoteWorkerCpp/SharedMemoryBuffer.cpp
+++ b/Engine/Source/Programs/Horde/Samples/RemoteWorkerCpp/SharedMemoryBuffer.cpp
@@ -5,80 +5,105 @@
 #include <iostream>
 #include "SharedMemoryBuffer.h"
 
-enum class FSharedMemoryBuffer::EState : int
+enum class FSharedMemoryBuffer::EWriteState : int
 {
-	Normal,
-	ReaderWaitingForData, // Once new data is available, set _readerTcs
-	WriterWaitingToFlush, // Once flush has been complete, set _writerTcs
-	Finished
+	// Chunk is still being appended to
+	Writing = 0,
+
+	// Writer has moved to the next chunk
+	MovedToNext = 2,
+
+	// This chunk marks the end of the stream
+	Complete = 3,
 };
 
-class FSharedMemoryBuffer::FHeader
+struct FSharedMemoryBuffer::FChunkState
 {
-public:
-	EState GetState() const
-	{
-		return (EState)InterlockedCompareExchange((long*)&State, 0, 0);
-	}
+	const long long Value;
 
-	bool UpdateState(EState OldState, EState NewState)
+	// Constructor
+	FChunkState(long long InValue)
+		: Value(InValue)
 	{
-		return InterlockedCompareExchange((long*)&State, (long)NewState, (long)OldState) == (long)OldState;
 	}
 
-	void Compact()
+	// Constructor
+	FChunkState(EWriteState WriteState, int ReaderFlags, int Length)
+		: Value((unsigned long long)Length | ((unsigned long long)ReaderFlags << 31) | ((unsigned long long)WriteState << 62))
 	{
-		InterlockedExchange64(&ReadPosition, 0);
-		InterlockedExchange64(&WritePosition, GetReadLength());
 	}
 
-	long long GetReadPosition() const
+	// Written length of this chunk
+	int GetLength() const { return (int)(Value & 0x7fffffff); }
+
+	// Set of flags which are set for each reader that still has to read from a chunk
+	int GetReaderFlags() const { return (int)((Value >> 31) & 0x7fffffff); }
+
+	// State of the writer
+	EWriteState GetWriteState() const { return (EWriteState)(Value >> 62); }
+
+	// Test whether a particular reader is still referencing the chunk
+	bool HasReaderFlag() const { return (Value & (1UL << 31)) != 0; }
+
+	// Read the state value from memory
+	static FChunkState Read(volatile long long* StateValue)
 	{
-		return InterlockedCompareExchange64(const_cast<long long*>(&ReadPosition), 0, 0);
+		return FChunkState(InterlockedCompareExchange64(StateValue, 0, 0));
 	}
 
-	long long AdvanceReadPosition(long long Value)
+	// Append data to the chunk
+	static void Append(volatile long long* StateValue, long long length)
 	{
-		InterlockedAdd64(&ReadPosition, Value);
-		return InterlockedAdd64(&ReadLength, -Value);
+		InterlockedAdd64(StateValue, length);
 	}
 
-	long long GetReadLength() const
+	// Mark the chunk as being written to
+	static void StartWriting(volatile long long* StateValue, int numReaders)
 	{
-		return InterlockedCompareExchange64(const_cast<long long*>(&ReadLength), 0, 0);
+		InterlockedExchange64(StateValue, FChunkState(EWriteState::Writing, (1 << numReaders) - 1, 0).Value);
 	}
 
-	long long GetWritePosition() const
+	// Move to the next chunk
+	static void MoveToNext(volatile long long* StateValue)
 	{
-		return InterlockedCompareExchange64(const_cast<long long*>(&WritePosition), 0, 0);
+		InterlockedOr64(StateValue, FChunkState(EWriteState::MovedToNext, 0, 0).Value);
 	}
 
-	long long AdvanceWritePosition(long long value)
+	// Move to the next chunk
+	static void MarkComplete(volatile long long* StateValue)
 	{
-		InterlockedAdd64(&WritePosition, value);
-		return InterlockedAdd64(&ReadLength, value);
+		InterlockedOr64(StateValue, FChunkState(EWriteState::Complete, 0, 0).Value);
 	}
 
-	void FinishWriting()
+	// Clear the reader flag
+	static void FinishReading(volatile long long* StateValue)
 	{
-		InterlockedExchange((long*)&State, (long)EState::Finished);
+		InterlockedAnd64(StateValue, ~(1LL << (0 + 31)));
 	}
+};
 
-private:
-	EState State;
-	unsigned char Padding[4];
-	long long ReadPosition;
-	long long ReadLength;
-	long long WritePosition;
+struct FSharedMemoryBuffer::FHeader
+{
+	int NumChunks;
+	int ChunkLength;
 };
 
+////////////////////////////
+
 FSharedMemoryBuffer::FSharedMemoryBuffer()
 	: MemoryMappedFile(nullptr)
-	, Memory(nullptr)
-	, Length(0)
 	, Header(nullptr)
 	, ReaderEvent(nullptr)
 	, WriterEvent(nullptr)
+	// Reader
+	, ReadChunkIdx(-1)
+	, ReadOffset(0)
+	, ReadChunkStatePtr(nullptr)
+	, ReadChunkDataPtr(nullptr)
+	// Writer
+	, WriteChunkIdx(-1)
+	, WriteChunkStatePtr(nullptr)
+	, WriteChunkDataPtr(nullptr)
 {
 }
 
@@ -87,38 +112,58 @@ FSharedMemoryBuffer::~FSharedMemoryBuffer()
 	Close();
 }
 
-bool FSharedMemoryBuffer::CreateNew(long long Capacity)
+bool FSharedMemoryBuffer::OpenExisting(const char* Name)
 {
-	Close();
+	char NameBuffer[MAX_PATH];
+	sprintf_s(NameBuffer, "%s_M", Name);
 
-	LARGE_INTEGER CapacityLargeInt;
-	CapacityLargeInt.QuadPart = Capacity;
+	MemoryMappedFile = OpenFileMappingA(FILE_MAP_ALL_ACCESS, TRUE, NameBuffer);
+	if (MemoryMappedFile == nullptr)
+	{
+		std::cout << "No filemap" << NameBuffer << std::endl;
+		Close();
+		return false;
+	}
 
-	ReaderEvent = CreateEvent(nullptr, true, FALSE, nullptr);
-	WriterEvent = CreateEvent(nullptr, true, FALSE, nullptr);
-	MemoryMappedFile = CreateFileMappingW(INVALID_HANDLE_VALUE, nullptr, PAGE_READWRITE, CapacityLargeInt.HighPart, CapacityLargeInt.LowPart, nullptr);
+	Header = (FHeader*)MapViewOfFile(MemoryMappedFile, FILE_MAP_ALL_ACCESS, 0, 0, 0);
+	if (Header == nullptr)
+	{
+		std::cout << "No filemap view" << std::endl;
+		Close();
+		return false;
+	}
 
-	if (ReaderEvent == nullptr || WriterEvent == nullptr || MemoryMappedFile == nullptr || !Initialize())
+	sprintf_s(NameBuffer, "%s_R", Name);
+	ReaderEvent = OpenEventA(SYNCHRONIZE | EVENT_MODIFY_STATE, TRUE, NameBuffer);
+	if (ReaderEvent == nullptr)
 	{
+		std::cout << "No read event" << std::endl;
 		Close();
 		return false;
 	}
 
-	return true;
-}
+	sprintf_s(NameBuffer, "%s_W", Name);
+	WriterEvent = OpenEventA(SYNCHRONIZE | EVENT_MODIFY_STATE, TRUE, NameBuffer);
+	if (WriterEvent == nullptr)
+	{
+		std::cout << "No write event" << std::endl;
+		Close();
+		return false;
+	}
 
-bool FSharedMemoryBuffer::OpenExisting(void* InMemoryMappedFile, void* InReaderEvent, void* InWriterEvent)
-{
-	MemoryMappedFile = InMemoryMappedFile;
-	ReaderEvent = InReaderEvent;
-	WriterEvent = InWriterEvent;
-	return Initialize();
+	ReadChunkIdx = 0;
+	ReadChunkDataPtr = GetChunkDataPtr(ReadChunkIdx);
+	ReadChunkStatePtr = GetChunkStatePtr(ReadChunkIdx);
+
+	WriteChunkIdx = 0;
+	WriteChunkDataPtr = GetChunkDataPtr(WriteChunkIdx);
+	WriteChunkStatePtr = GetChunkStatePtr(WriteChunkIdx);
+
+	return true;
 }
 
 void FSharedMemoryBuffer::Close()
 {
-	Memory = nullptr;
-
 	if (Header != nullptr)
 	{
 		UnmapViewOfFile(Header);
@@ -142,159 +187,152 @@ void FSharedMemoryBuffer::Close()
 		CloseHandle(ReaderEvent);
 		ReaderEvent = nullptr;
 	}
-}
 
-bool FSharedMemoryBuffer::DuplicateHandles(void* TargetProcess, void*& OutMemoryMappedFile, void*& OutReaderEvent, void*& OutWriterEvent)
-{
-	HANDLE CurrentProcess = GetCurrentProcess();
+	ReadChunkIdx = -1;
+	ReadOffset = 0;
+	ReadChunkStatePtr = nullptr;
+	ReadChunkDataPtr = nullptr;
 
-	HANDLE TargetMemoryMappedFile;
-	if (DuplicateHandle(CurrentProcess, MemoryMappedFile, TargetProcess, &TargetMemoryMappedFile, 0, false, DUPLICATE_SAME_ACCESS))
-	{
-		HANDLE TargetWriterEvent;
-		if (DuplicateHandle(CurrentProcess, WriterEvent, TargetProcess, &TargetWriterEvent, 0, false, DUPLICATE_SAME_ACCESS))
-		{
-			HANDLE TargetReaderEvent;
-			if (DuplicateHandle(CurrentProcess, ReaderEvent, TargetProcess, &TargetReaderEvent, 0, false, DUPLICATE_SAME_ACCESS))
-			{
-				OutMemoryMappedFile = TargetMemoryMappedFile;
-				OutReaderEvent = TargetReaderEvent;
-				OutWriterEvent = TargetWriterEvent;
-				return true;
-			}
-			CloseHandle(TargetWriterEvent);
-		}
-		CloseHandle(TargetMemoryMappedFile);
-	}
-
-	return false;
+	WriteChunkIdx = -1;
+	WriteChunkStatePtr = nullptr;
+	WriteChunkDataPtr = nullptr;
 }
 
-bool FSharedMemoryBuffer::HasFinishedReading() const
+bool FSharedMemoryBuffer::IsComplete() const
 {
-	return Header->GetState() == EState::Finished && Header->GetReadLength() == 0;
+	FChunkState State = FChunkState::Read(ReadChunkStatePtr);
+	return State.GetWriteState() == EWriteState::Complete && ReadOffset == State.GetLength();
 }
 
-void FSharedMemoryBuffer::AdvanceReadPosition(long long Size)
+void FSharedMemoryBuffer::AdvanceReadPosition(size_t Size)
 {
-	if (Header->AdvanceReadPosition(Size) == 0)
-	{
-		SetEvent(WriterEvent);
-	}
+	ReadOffset += Size;
 }
 
-const unsigned char* FSharedMemoryBuffer::GetReadMemory(long long& OutSize)
+const unsigned char* FSharedMemoryBuffer::GetReadMemory(size_t& OutSize)
 {
-	OutSize = Header->GetReadLength();
-	return Memory + Header->GetReadPosition();
+	FChunkState State = FChunkState::Read(ReadChunkStatePtr);
+	if (State.HasReaderFlag())
+	{
+		OutSize = State.GetLength() - ReadOffset;
+		return ReadChunkDataPtr + ReadOffset;
+	}
+	else
+	{
+		OutSize = 0;
+		return nullptr;
+	}
 }
 
-void FSharedMemoryBuffer::WaitForData(long long Size)
+void FSharedMemoryBuffer::WaitToRead(size_t CurrentLength)
 {
-	while (Header->GetReadLength() == Size)
+	for (; ; )
 	{
-		EState State = Header->GetState();
-		if (State == EState::Finished)
+		FChunkState State = FChunkState::Read(ReadChunkStatePtr);
+
+		if (!State.HasReaderFlag())
 		{
-			break;
+			// Wait until the current chunk is readable
+			ResetEvent(ReaderEvent);
+			if (!FChunkState::Read(ReadChunkStatePtr).HasReaderFlag())
+			{
+				WaitForSingleObject(ReaderEvent, INFINITE);
+			}
 		}
-		else if (State == EState::Normal)
+		else if (ReadOffset + CurrentLength < State.GetLength() || State.GetWriteState() == EWriteState::Complete)
 		{
-			ResetEvent(ReaderEvent);
-			Header->UpdateState(State, EState::ReaderWaitingForData);
+			// Still have data to read from this chunk
+			break;
 		}
-		else if (State == EState::ReaderWaitingForData)
+		else if (ReadOffset + CurrentLength >= State.GetLength() && State.GetWriteState() == EWriteState::Writing)
 		{
-			WaitForSingleObject(ReaderEvent, INFINITE);
+			// Wait until there is more data in the chunk
+			ResetEvent(ReaderEvent);
+			if (FChunkState::Read(ReadChunkStatePtr).Value == State.Value)
+			{
+				WaitForSingleObject(ReaderEvent, INFINITE);
+			}
 		}
-		else if (State == EState::WriterWaitingToFlush)
+		else if (State.GetWriteState() == EWriteState::MovedToNext)
 		{
-			Header->Compact();
-			Header->UpdateState(State, EState::Normal);
+			// Move to the next chunk
+			FChunkState::FinishReading(GetChunkStatePtr(ReadChunkIdx));
 			SetEvent(WriterEvent);
+
+			ReadChunkIdx++;
+			if (ReadChunkIdx == Header->NumChunks)
+			{
+				ReadChunkIdx = 0;
+			}
+			ReadOffset = 0;
 		}
 		else
 		{
-			assert(false);
+			// Still need to read data from the current buffer
+			break;
 		}
 	}
 }
 
-void FSharedMemoryBuffer::FinishWriting()
+
+void FSharedMemoryBuffer::MarkComplete()
 {
-	Header->FinishWriting();
+	FChunkState::MarkComplete(ReadChunkStatePtr);
 	SetEvent(ReaderEvent);
 }
 
-void FSharedMemoryBuffer::AdvanceWritePosition(long long size)
+void FSharedMemoryBuffer::AdvanceWritePosition(size_t Size)
 {
-	assert(Header->GetState() != EState::Finished);
-
-	Header->AdvanceWritePosition(size);
-
-	if (Header->UpdateState(EState::ReaderWaitingForData, EState::Normal))
-	{
-		SetEvent(ReaderEvent);
-	}
+	FChunkState::Append(ReadChunkStatePtr, Size);
+	SetEvent(ReaderEvent);
 }
 
-unsigned char* FSharedMemoryBuffer::GetWriteMemory(long long& OutSize)
+unsigned char* FSharedMemoryBuffer::GetWriteMemory(size_t& OutSize)
 {
-	long long WritePosition = Header->GetWritePosition();
-	OutSize = Length - WritePosition;
-	return Memory + WritePosition;
+	FChunkState State = FChunkState::Read(ReadChunkStatePtr);
+	OutSize = State.GetLength() - ReadOffset;
+	return ReadChunkDataPtr + ReadOffset;
 }
 
-void FSharedMemoryBuffer::Flush()
+void FSharedMemoryBuffer::WaitToWrite(size_t CurrentLength)
 {
-	while (Header->GetReadPosition() > 0)
+	for (; ; )
 	{
-		EState State = Header->GetState();
-		if (State == EState::Finished)
+		size_t Length;
+		GetWriteMemory(Length);
+
+		if (Length != CurrentLength)
 		{
 			break;
 		}
-		else if (State == EState::Normal)
-		{
-			SetEvent(WriterEvent);
-			Header->UpdateState(State, EState::WriterWaitingToFlush);
-		}
-		else if (State == EState::ReaderWaitingForData)
+
+		FChunkState::MoveToNext(WriteChunkStatePtr);
+
+		WriteChunkIdx++;
+		if (WriteChunkIdx == Header->NumChunks)
 		{
-			Header->Compact();
-			break;
+			WriteChunkIdx = 0;
 		}
-		else if (State == EState::WriterWaitingToFlush)
+
+		WriteChunkDataPtr = GetChunkDataPtr(WriteChunkIdx);
+		WriteChunkStatePtr = GetChunkStatePtr(WriteChunkIdx);
+
+		while (FChunkState::Read(WriteChunkStatePtr).GetReaderFlags() != 0)
 		{
 			WaitForSingleObject(WriterEvent, INFINITE);
+			ResetEvent(WriterEvent);
 		}
+
+		FChunkState::StartWriting(WriteChunkStatePtr, 1);
 	}
 }
 
-bool FSharedMemoryBuffer::Initialize()
+unsigned char* FSharedMemoryBuffer::GetChunkDataPtr(int chunkIdx) const
 {
-	if (ReaderEvent == nullptr || WriterEvent == nullptr || MemoryMappedFile == nullptr)
-	{
-		Close();
-		return false;
-	}
-
-	Header = (FHeader*)MapViewOfFile(MemoryMappedFile, FILE_MAP_ALL_ACCESS, 0, 0, 0);
-	if (Header == nullptr)
-	{
-		Close();
-		return false;
-	}
-
-	Memory = (unsigned char*)(Header + 1);
-
-	MEMORY_BASIC_INFORMATION MemInfo;
-	if (VirtualQuery(Header, &MemInfo, sizeof(MemInfo)) == 0)
-	{
-		Close();
-		return false;
-	}
-
-	Length = MemInfo.RegionSize;
-	return true;
+	return (unsigned char*)(Header + 1) + (sizeof(long long) * Header->NumChunks) + (chunkIdx * Header->ChunkLength);
 }
+
+volatile long long* FSharedMemoryBuffer::GetChunkStatePtr(int chunkIdx) const
+{
+	return (volatile long long*)(Header + 1) + chunkIdx;
+}
\ No newline at end of file
diff --git a/Engine/Source/Programs/Horde/Samples/RemoteWorkerCpp/SharedMemoryBuffer.h b/Engine/Source/Programs/Horde/Samples/RemoteWorkerCpp/SharedMemoryBuffer.h
--- a/Engine/Source/Programs/Horde/Samples/RemoteWorkerCpp/SharedMemoryBuffer.h
+++ b/Engine/Source/Programs/Horde/Samples/RemoteWorkerCpp/SharedMemoryBuffer.h
@@ -4,62 +4,68 @@
 
 class FSharedMemoryBuffer
 {
-	enum class EState : int;
-	class FHeader;
+	enum class EWriteState : int;
+	struct FChunkState;
+	struct FHeader;
 
 public:
 	FSharedMemoryBuffer();
 	~FSharedMemoryBuffer();
 
-	// Initialize the object to a new shared memory buffer of the given size
-	bool CreateNew(long long Capacity);
-
 	// Opens an existing shared memory buffer (typically from handles created in another process)
-	bool OpenExisting(void* MemoryMappedFile, void* ReaderEvent, void* WriterEvent);
+	bool OpenExisting(const char* Name);
 
 	// Close the current buffer and release all allocated resources
 	void Close();
 
-	// Duplicates handles for this buffer into another process
-	bool DuplicateHandles(void* TargetProcess, void*& OutMemoryMappedFile, void*& OutReaderEvent, void*& OutWriterEvent);
-
 
 	/*** Reader Interface ***/
 
 	// Test whether the buffer has finished being written to (ie. FinishWriting() has been called) and all data has been read from it.
-	bool HasFinishedReading() const;
+	bool IsComplete() const;
 
 	// Move the read cursor forwards by the given number of bytes
-	void AdvanceReadPosition(long long Size);
+	void AdvanceReadPosition(size_t Size);
 
 	// Gets the next data to be read (and the number of valid bytes accessible from the given pointer)
-	const unsigned char* GetReadMemory(long long& OutSize);
+	const unsigned char* GetReadMemory(size_t& OutSize);
 
 	// Wait for more data to be written to the buffer. The given parameter indicates the current size of the read buffer, used to exit immediately if it's changed since the value was fetched.
-	void WaitForData(long long CurrentLength);
+	void WaitToRead(size_t CurrentLength);
 
 
 	/*** Writer Interface ***/
 
 	// Signal that we've finished writing to this buffer
-	void FinishWriting();
+	void MarkComplete();
 
 	// Move the write cursor forward by the given number of bytes
-	void AdvanceWritePosition(long long Size);
+	void AdvanceWritePosition(size_t Size);
 
 	// Gets the memory that can be written to, and the available space in it
-	unsigned char* GetWriteMemory(long long& OutSize);
+	unsigned char* GetWriteMemory(size_t& OutSize);
 
 	// Waits until all data has been read from the buffer (or more specifically, that the reader is stalled) and moves any unread data to the start of the buffer.
-	void Flush();
+	void WaitToWrite(size_t currentLength);
 
 private:
-	bool Initialize();
+	unsigned char* GetChunkDataPtr(int chunkIdx) const;
+	volatile long long* GetChunkStatePtr(int chunkIdx) const;
+
+	const int ReaderIdx = 0;
 
 	void* MemoryMappedFile;
-	long long Length;
 	FHeader* Header;
-	unsigned char* Memory;
 	void* ReaderEvent;
 	void* WriterEvent;
+
+	int ReadChunkIdx;
+	size_t ReadOffset;
+	volatile long long* ReadChunkStatePtr;
+	unsigned char* ReadChunkDataPtr;
+
+	int WriteChunkIdx;
+	volatile long long* WriteChunkStatePtr;
+	unsigned char* WriteChunkDataPtr;
 };
+
diff --git a/Engine/Source/Programs/Horde/Samples/RemoteWorkerCpp/WorkerSocket.cpp b/Engine/Source/Programs/Horde/Samples/RemoteWorkerCpp/WorkerSocket.cpp
deleted file mode 100644
diff --git a/Engine/Source/Programs/Horde/Samples/RemoteWorkerCpp/WorkerSocket.h b/Engine/Source/Programs/Horde/Samples/RemoteWorkerCpp/WorkerSocket.h
deleted file mode 100644
