Horde: Remove debug output.

diff --git a/Engine/Source/Programs/Horde/Samples/RemoteWorkerCpp/ComputeChannel.cpp b/Engine/Source/Programs/Horde/Samples/RemoteWorkerCpp/ComputeChannel.cpp
--- a/Engine/Source/Programs/Horde/Samples/RemoteWorkerCpp/ComputeChannel.cpp
+++ b/Engine/Source/Programs/Horde/Samples/RemoteWorkerCpp/ComputeChannel.cpp
@@ -25,7 +25,6 @@ bool FComputeChannel::Open()
 
 	if (Length <= 0 || Length >= sizeof(EnvVar))
 	{
-		std::cout << "No env var" << std::endl;
 		Close();
 		return false;
 	}
@@ -35,7 +34,6 @@ bool FComputeChannel::Open()
 	sprintf_s(Buffer, "%s_SEND", EnvVar);
 	if (!SendBuffer.OpenExisting(Buffer))
 	{
-		std::cout << "No send: " << Buffer << std::endl;
 		Close();
 		return false;
 	}
@@ -43,7 +41,6 @@ bool FComputeChannel::Open()
 	sprintf_s(Buffer, "%s_RECV", EnvVar);
 	if (!RecvBuffer.OpenExisting(Buffer))
 	{
-		std::cout << "No recv: " << Buffer << std::endl;
 		Close();
 		return false;
 	}
diff --git a/Engine/Source/Programs/Horde/Samples/RemoteWorkerCpp/SharedMemoryBuffer.cpp b/Engine/Source/Programs/Horde/Samples/RemoteWorkerCpp/SharedMemoryBuffer.cpp
--- a/Engine/Source/Programs/Horde/Samples/RemoteWorkerCpp/SharedMemoryBuffer.cpp
+++ b/Engine/Source/Programs/Horde/Samples/RemoteWorkerCpp/SharedMemoryBuffer.cpp
@@ -120,7 +120,6 @@ bool FSharedMemoryBuffer::OpenExisting(const char* Name)
 	MemoryMappedFile = OpenFileMappingA(FILE_MAP_ALL_ACCESS, TRUE, NameBuffer);
 	if (MemoryMappedFile == nullptr)
 	{
-		std::cout << "No filemap" << NameBuffer << std::endl;
 		Close();
 		return false;
 	}
@@ -128,7 +127,6 @@ bool FSharedMemoryBuffer::OpenExisting(const char* Name)
 	Header = (FHeader*)MapViewOfFile(MemoryMappedFile, FILE_MAP_ALL_ACCESS, 0, 0, 0);
 	if (Header == nullptr)
 	{
-		std::cout << "No filemap view" << std::endl;
 		Close();
 		return false;
 	}
@@ -137,7 +135,6 @@ bool FSharedMemoryBuffer::OpenExisting(const char* Name)
 	ReaderEvent = OpenEventA(SYNCHRONIZE | EVENT_MODIFY_STATE, TRUE, NameBuffer);
 	if (ReaderEvent == nullptr)
 	{
-		std::cout << "No read event" << std::endl;
 		Close();
 		return false;
 	}
@@ -146,7 +143,6 @@ bool FSharedMemoryBuffer::OpenExisting(const char* Name)
 	WriterEvent = OpenEventA(SYNCHRONIZE | EVENT_MODIFY_STATE, TRUE, NameBuffer);
 	if (WriterEvent == nullptr)
 	{
-		std::cout << "No write event" << std::endl;
 		Close();
 		return false;
 	}
