Horde: Fix worker shutting down after reading is complete.

diff --git a/Engine/Source/Programs/Horde/Samples/RemoteClient/RemoteClientApp.cs b/Engine/Source/Programs/Horde/Samples/RemoteClient/RemoteClientApp.cs
--- a/Engine/Source/Programs/Horde/Samples/RemoteClient/RemoteClientApp.cs
+++ b/Engine/Source/Programs/Horde/Samples/RemoteClient/RemoteClientApp.cs
@@ -114,7 +114,7 @@ namespace RemoteClient
 			const int DataChannelId = 1;
 
 			byte[] buffer = new byte[4];
-			for (int idx = 0; idx < 100; idx++)
+			for (int idx = 0; idx < 3; idx++)
 			{
 				cancellationToken.ThrowIfCancellationRequested();
 				logger.LogInformation("Writing value: {Value}", idx);
@@ -122,6 +122,8 @@ namespace RemoteClient
 				await socket.SendAsync(DataChannelId, buffer, cancellationToken);
 				await Task.Delay(1000, cancellationToken);
 			}
+
+			await socket.MarkCompleteAsync(DataChannelId, cancellationToken);
 		}
 
 		static async Task<IComputeClient> CreateClientAsync(ClientAppOptions options, ILogger logger)
diff --git a/Engine/Source/Programs/Horde/Samples/RemoteWorker/RemoteWorkerApp.cs b/Engine/Source/Programs/Horde/Samples/RemoteWorker/RemoteWorkerApp.cs
--- a/Engine/Source/Programs/Horde/Samples/RemoteWorker/RemoteWorkerApp.cs
+++ b/Engine/Source/Programs/Horde/Samples/RemoteWorker/RemoteWorkerApp.cs
@@ -2,16 +2,15 @@
 
 using System.Buffers.Binary;
 using EpicGames.Horde.Compute;
-using Microsoft.Extensions.Logging.Abstractions;
 
-namespace RemoteServer
+namespace RemoteWorker
 {
-	class ServerApp
+	class WorkerApp
 	{
 		static async Task Main()
 		{
 			using IComputeChannel channel = ComputeChannel.ConnectAsWorker();
-			Console.WriteLine("Connected to client"); // Client will wait for output before sending data on channel 1
+			Console.WriteLine("Connected to client");
 
 			byte[] data = new byte[4];
 			while(await channel.TryReceiveMessageAsync(data))
@@ -19,6 +18,8 @@ namespace RemoteServer
 				int value = BinaryPrimitives.ReadInt32LittleEndian(data);
 				Console.WriteLine("Read value {0}", value);
 			}
+
+			Console.WriteLine("Exiting worker");
 		}
 	}
 }
\ No newline at end of file
diff --git a/Engine/Source/Programs/Horde/Samples/RemoteWorkerCpp/SharedMemoryBuffer.cpp b/Engine/Source/Programs/Horde/Samples/RemoteWorkerCpp/SharedMemoryBuffer.cpp
--- a/Engine/Source/Programs/Horde/Samples/RemoteWorkerCpp/SharedMemoryBuffer.cpp
+++ b/Engine/Source/Programs/Horde/Samples/RemoteWorkerCpp/SharedMemoryBuffer.cpp
@@ -40,7 +40,7 @@ struct FSharedMemoryBuffer::FChunkState
 	int GetReaderFlags() const { return (int)((Value >> 31) & 0x7fffffff); }
 
 	// State of the writer
-	EWriteState GetWriteState() const { return (EWriteState)(Value >> 62); }
+	EWriteState GetWriteState() const { return (EWriteState)((unsigned long long)Value >> 62); }
 
 	// Test whether a particular reader is still referencing the chunk
 	bool HasReaderFlag() const { return (Value & (1UL << 31)) != 0; }
diff --git a/Engine/Source/Programs/Shared/EpicGames.Horde/Compute/ComputeBufferBase.cs b/Engine/Source/Programs/Shared/EpicGames.Horde/Compute/ComputeBufferBase.cs
--- a/Engine/Source/Programs/Shared/EpicGames.Horde/Compute/ComputeBufferBase.cs
+++ b/Engine/Source/Programs/Shared/EpicGames.Horde/Compute/ComputeBufferBase.cs
@@ -256,7 +256,7 @@ namespace EpicGames.Horde.Compute
 			{
 				if (currentLength >= _chunks[_chunkIdx].Memory.Length)
 				{
-					throw new NotImplementedException("Will never succeed");
+					throw new ArgumentException("Current length is already complete buffer size; cannot allocate more.", nameof(currentLength));
 				}
 
 				for (; ; )
diff --git a/Engine/Source/Programs/Shared/EpicGames.Horde/Compute/IComputeBuffer.cs b/Engine/Source/Programs/Shared/EpicGames.Horde/Compute/IComputeBuffer.cs
--- a/Engine/Source/Programs/Shared/EpicGames.Horde/Compute/IComputeBuffer.cs
+++ b/Engine/Source/Programs/Shared/EpicGames.Horde/Compute/IComputeBuffer.cs
@@ -97,7 +97,7 @@ namespace EpicGames.Horde.Compute
 			for (; ; )
 			{
 				ReadOnlyMemory<byte> readMemory = reader.GetMemory();
-				if (readMemory.Length > 0)
+				if (reader.IsComplete || readMemory.Length > 0)
 				{
 					int length = Math.Min(readMemory.Length, buffer.Length);
 					readMemory.Slice(0, length).CopyTo(buffer);
