Horde: Add retries to marking agents for shutdown

Automatic retrying was removed in CL 24811764.

diff --git a/Engine/Source/Programs/Horde/Horde.Server/Agents/Fleet/Providers/AwsFleetManager.cs b/Engine/Source/Programs/Horde/Horde.Server/Agents/Fleet/Providers/AwsFleetManager.cs
--- a/Engine/Source/Programs/Horde/Horde.Server/Agents/Fleet/Providers/AwsFleetManager.cs
+++ b/Engine/Source/Programs/Horde/Horde.Server/Agents/Fleet/Providers/AwsFleetManager.cs
@@ -256,16 +256,35 @@ namespace Horde.Server.Agents.Fleet.Providers
 
 			foreach (IAgent agent in agentsLimitedByCount)
 			{
-				IAuditLogChannel<AgentId> agentLogger = agentCollection.GetLogger(agent.Id);
+				await TryRequestShutdown(agentCollection, pool, agent);
+			}
+		}
+
+		private static async Task<bool> TryRequestShutdown(IAgentCollection agentCollection, IPool pool, IAgent agent)
+		{
+			IAuditLogChannel<AgentId> agentLogger = agentCollection.GetLogger(agent.Id);
+
+			const int MaxRetries = 5;
+			for (int retryCount = 0; retryCount < MaxRetries; retryCount++)
+			{
 				if (await agentCollection.TryUpdateSettingsAsync(agent, requestShutdown: true, shutdownReason: "Autoscaler") != null)
 				{
-					agentLogger.LogInformation("Marked {AgentId} in pool {PoolName} for shutdown due to autoscaling (currently {NumLeases} leases outstanding)", agent.Id, pool.Name, agent.Leases.Count);
+					agentLogger.LogInformation("Marked {AgentId} in pool {PoolName} for shutdown due to autoscaling (currently {NumLeases} leases outstanding, {NumRetries} retries)", agent.Id, pool.Name, agent.Leases.Count, retryCount);
+					return true;
 				}
-				else
+
+				IAgent? updatedAgent = await agentCollection.GetAsync(agent.Id);
+				if (updatedAgent == null)
 				{
-					agentLogger.LogError("Unable to mark agent {AgentId} in pool {PoolName} for shutdown due to autoscaling", agent.Id, pool.Name);
+					agentLogger.LogError("Unable to mark agent {AgentId} in pool {PoolName} for shutdown due to autoscaling. Agent no longer exists", agent.Id, pool.Name);
+					return false;
 				}
+
+				agent = updatedAgent;
 			}
+
+			agentLogger.LogError("Unable to mark agent {AgentId} in pool {PoolName} for shutdown due to autoscaling. Retries exhausted", agent.Id, pool.Name);
+			return false;
 		}
 
 		/// <inheritdoc/>
