Horde: Artifact browser work
#rnx
#skipci

diff --git a/Engine/Source/Programs/Horde/HordeDashboard/src/components/jobDetailsV2/JobDetailArtifactsV3.tsx b/Engine/Source/Programs/Horde/HordeDashboard/src/components/jobDetailsV2/JobDetailArtifactsV3.tsx
--- a/Engine/Source/Programs/Horde/HordeDashboard/src/components/jobDetailsV2/JobDetailArtifactsV3.tsx
+++ b/Engine/Source/Programs/Horde/HordeDashboard/src/components/jobDetailsV2/JobDetailArtifactsV3.tsx
@@ -1,12 +1,12 @@
 // Copyright Epic Games, Inc. All Rights Reserved.
 
-import { DetailsList, DetailsListLayoutMode, Dropdown, Link as FluentLink, FontIcon, IColumn, IDropdownOption, IconButton, Modal, PrimaryButton, ScrollablePane, Selection, SelectionMode, SelectionZone, Spinner, SpinnerSize, Stack, Text, mergeStyleSets } from "@fluentui/react";
+import { DetailsList, DetailsListLayoutMode, DetailsRow, Dropdown, Link as FluentLink, FontIcon, IColumn, IDetailsListProps, IDropdownOption, Icon, IconButton, Modal, PrimaryButton, ScrollablePane, Selection, SelectionMode, SelectionZone, Spinner, SpinnerSize, Stack, Text, mergeStyleSets } from "@fluentui/react";
 import { action, makeObservable, observable } from "mobx";
 import { observer } from "mobx-react-lite";
 import { useEffect, useState } from "react";
 import backend from "../../backend";
 import { ArtifactContextType, GetArtifactDirectoryEntryResponse, GetArtifactDirectoryResponse, GetArtifactFileEntryResponse, GetArtifactResponseV2 } from "../../backend/Api";
-import { hordeClasses } from "../../styles/Styles";
+import { hordeClasses, modeColors } from "../../styles/Styles";
 import { JobDetailsV2 } from "./JobDetailsViewCommon";
 
 
@@ -198,7 +198,7 @@ class ArtifactsHandler {
 
    path?: string;
 
-   selection?: Selection = new Selection({ onSelectionChanged: () => { if (this.selectionCallback) this.selectionCallback() } });
+   selection?: Selection = new Selection({ canSelectItem: (item: any) => { return item.type !== BrowserType.NavigateUp; }, onSelectionChanged: () => { if (this.selectionCallback) this.selectionCallback() } });
 
    browse?: GetArtifactDirectoryResponse;
 
@@ -210,15 +210,18 @@ class ArtifactsHandler {
    context?: ArtifactContextType;
 
    stepId: string;
-
-
 }
 
 const styles = mergeStyleSets({
    list: {
       selectors: {
-         '.ms-DetailsRow': {
-            border: "1px solid #EDEBE9"
+         '.ms-List-cell': {
+            borderTop: "1px solid #EDEBE9",
+            borderRight: "1px solid #EDEBE9",
+            borderLeft: "1px solid #EDEBE9"
+         },
+         '.ms-List-cell:nth-last-child(-n + 1)': {
+            borderBottom: "1px solid #EDEBE9"
          }
       }
    }
@@ -230,7 +233,7 @@ const BrowseBreadCrumbs: React.FC<{ handler: ArtifactsHandler }> = observer(({ h
    // subscribe
    if (handler.updated) { }
 
-   const fontSize = 16;
+   const fontSize = 13;
 
    if (!handler.path) {
       return <Stack><Text style={{ fontSize: fontSize, paddingLeft: 4, paddingRight: 4 }}>/
@@ -244,9 +247,9 @@ const BrowseBreadCrumbs: React.FC<{ handler: ArtifactsHandler }> = observer(({ h
       const path = ppath.slice(0, index + 1).join("/");
       return <Stack horizontal onClick={() => handler.browseTo(path)} style={{ cursor: "pointer" }}>
          <Stack>
-            <Text style={{ fontSize: fontSize }}>{e}</Text>
+            <Text style={{ fontSize: fontSize, fontFamily: "Horde Open Sans Light" }}>{e}</Text>
          </Stack>
-         {index !== (ppath.length - 1) && <Stack style={{ paddingLeft: 4, paddingRight: 4 }}><Text style={{ fontSize: fontSize }}>/</Text></Stack>}
+         {index !== (ppath.length - 1) && <Stack style={{ paddingLeft: 4, paddingRight: 4 }}><Text style={{ fontSize: fontSize, fontFamily: "Horde Open Sans Light" }}>/</Text></Stack>}
       </Stack>
    });
 
@@ -415,8 +418,10 @@ const JobDetailArtifactsInner: React.FC<{ handler: ArtifactsHandler }> = observe
    });
 
    const columns: IColumn[] = [
-      { key: 'column1', name: 'Name', minWidth: 640, maxWidth: 640, isResizable: false, isPadded: false },
-      { key: 'column2', name: 'Size', minWidth: 128, maxWidth: 128, isResizable: false, isPadded: false },
+      { key: 'column1', name: 'Name', minWidth: 768, maxWidth: 768, isResizable: false, isPadded: false },
+      { key: 'column2', name: 'View', minWidth: 32, maxWidth: 32, isResizable: false, isPadded: false },
+      { key: 'column3', name: 'Download', minWidth: 32, maxWidth: 32, isResizable: false, isPadded: false },
+      { key: 'column4', name: 'Size', minWidth: 128, maxWidth: 128, isResizable: false, isPadded: false },
    ];
 
    const renderItem = (item: any, index?: number, column?: IColumn) => {
@@ -429,12 +434,39 @@ const JobDetailArtifactsInner: React.FC<{ handler: ArtifactsHandler }> = observe
          if (!item.size) {
             return null;
          }
-         return <Stack>
+         return <Stack horizontalAlign="end">
             <Text>{formatBytes(item.size, (item.size < (1024 * 1024)) ? 0 : 1)}</Text>
          </Stack>
 
       }
 
+      if (column.name === "View") {
+         if (item.type !== BrowserType.File) {
+            return null;
+         }         
+         const path = encodeURI(handler.path + "/" + item.text);
+         const server = backend.serverUrl;
+         const href = `${server}/api/v2/artifacts/${handler.artifact!.id}/file?path=${path}&inline=true`;
+
+         return <Stack>
+            <IconButton href={href} target="_blank" style={{ height: 20, color: "#106EBE" }} iconProps={{ iconName: "Eye", styles: { root: { fontSize: "15px" } } }} />
+         </Stack>
+      }
+
+      if (column.name === "Download") {
+         if (item.type !== BrowserType.File) {
+            return null;
+         }
+         const path = encodeURI(handler.path + "/" + item.text);
+         const server = backend.serverUrl;
+         const href = `${server}/api/v2/artifacts/${handler.artifact!.id}/file?path=${path}`;
+
+
+         return <Stack>
+            <IconButton href={href} styles={{ root: { height: 20, color: "#106EBE" } }} iconProps={{ iconName: "CloudDownload", styles: { root: { fontSize: "14px" } } }} />
+         </Stack>
+      }
+
       if (column.name === "Name") {
 
          let href: string | undefined;
@@ -508,7 +540,7 @@ const JobDetailArtifactsInner: React.FC<{ handler: ArtifactsHandler }> = observe
                   selectionMode={SelectionMode.multiple}
                   selection={handler.selection}
                   selectionPreservedOnEmptyClick={true}
-                  //onItemInvoked={this._onItemInvoked} <--- double click*/
+                  //onItemInvoked={this._onItemInvoked} <--- double click*/                  
                   onRenderItemColumn={renderItem}
                />
             </SelectionZone>
diff --git a/Engine/Source/Programs/Horde/HordeDashboard/src/components/jobDetailsV2/JobOperationsBar.tsx b/Engine/Source/Programs/Horde/HordeDashboard/src/components/jobDetailsV2/JobOperationsBar.tsx
--- a/Engine/Source/Programs/Horde/HordeDashboard/src/components/jobDetailsV2/JobOperationsBar.tsx
+++ b/Engine/Source/Programs/Horde/HordeDashboard/src/components/jobDetailsV2/JobOperationsBar.tsx
@@ -147,7 +147,7 @@ export const JobOperations: React.FC<{ jobDetails: JobDetailsV2 }> = observer(({
             </Stack>
             </Link>}
 
-            {!!jobData?.useArtifactsV2 && <div onClick={() => { setArtifactsShown(true) }}>
+            {!!stepId && !!jobData?.useArtifactsV2 && <div onClick={() => { setArtifactsShown(true) }}>
                <Stack styles={{ root: { paddingTop: 4 } }}>
                   <CommandBarButton disabled={!stepArtifacts?.length} className={hordeClasses.commandBarSmall} styles={{ root: { padding: "10px 8px 10px 8px" } }} iconProps={{ iconName: "CloudDownload" }} text="Artifacts" />
                </Stack>
diff --git a/Engine/Source/Programs/Horde/HordeDashboard/src/styles/Styles.tsx b/Engine/Source/Programs/Horde/HordeDashboard/src/styles/Styles.tsx
--- a/Engine/Source/Programs/Horde/HordeDashboard/src/styles/Styles.tsx
+++ b/Engine/Source/Programs/Horde/HordeDashboard/src/styles/Styles.tsx
@@ -125,7 +125,8 @@ registerIcons({
       "Dashboard": IconContents.DASHBOARD,
       "Document": IconContents.DOCUMENT,
       "Folder": IconContents.FOLDER_CLOSE,
-      "CloudDownload" : IconContents.CLOUD_DOWNLOAD
+      "CloudDownload": IconContents.CLOUD_DOWNLOAD,
+      "Eye" : IconContents.EYE_OPEN
    }
 });
 
