Horde: Explicit filtering for errors, warnings, skips, also display GiB instead of GB for agents
#rnx
#skipci

diff --git a/Engine/Source/Programs/Horde/HordeDashboard/src/components/AgentView.tsx b/Engine/Source/Programs/Horde/HordeDashboard/src/components/AgentView.tsx
--- a/Engine/Source/Programs/Horde/HordeDashboard/src/components/AgentView.tsx
+++ b/Engine/Source/Programs/Horde/HordeDashboard/src/components/AgentView.tsx
@@ -1926,8 +1926,8 @@ export const AgentView: React.FC = observer(() => {
                   let diskCapacity: string | number | null = getAgentCapability(agent, "DiskTotalSize");
                   toCheck = "Unknown";
                   if (diskFree !== null && diskCapacity !== null) {
-                     diskFree = Number(diskFree) / 1000000000;
-                     diskCapacity = Number(diskCapacity) / 1000000000;
+                     diskFree = Number(diskFree) /  1073741824;
+                     diskCapacity = Number(diskCapacity) / 1073741824;
                      const percentage = (diskCapacity - diskFree) / diskCapacity;
                      toCheck = (diskCapacity - (percentage * diskCapacity)).toFixed(0);
                   }
@@ -2491,13 +2491,13 @@ export const AgentView: React.FC = observer(() => {
             let descVal = "Unknown";
             let descNumber = 0;
             if (diskFree !== null && diskCapacity !== null) {
-               diskFree = Number(diskFree) / 1000000000;
-               diskCapacity = Number(diskCapacity) / 1000000000;
+               diskFree = Number(diskFree) / 1073741824;
+               diskCapacity = Number(diskCapacity) / 1073741824;
                realPercentage = (diskCapacity - diskFree) / diskCapacity;
                // nudge percentage for the progress bar length
                nudgedPercentage = realPercentage * .965;
                descNumber = (diskCapacity - (realPercentage * diskCapacity));
-               descVal = descNumber.toLocaleString(undefined, { maximumFractionDigits: 0 }) + " GB Free";
+               descVal = descNumber.toLocaleString(undefined, { maximumFractionDigits: 0 }) + " GiB Free";
             }
             localState.columnSearchState['storage'][agent.id] = descNumber;
             return <Stack horizontal horizontalAlign={"center"}>
diff --git a/Engine/Source/Programs/Horde/HordeDashboard/src/components/HistoryModal.tsx b/Engine/Source/Programs/Horde/HordeDashboard/src/components/HistoryModal.tsx
--- a/Engine/Source/Programs/Horde/HordeDashboard/src/components/HistoryModal.tsx
+++ b/Engine/Source/Programs/Horde/HordeDashboard/src/components/HistoryModal.tsx
@@ -200,7 +200,7 @@ class HistoryModalState {
                                  subItemData[1] += " GB";
                               }
                               else if (subItemData[0].indexOf("Disk") !== -1) {
-                                 subItemData[1] = (Number(subItemData[1]) / 1000000000).toLocaleString(undefined, { maximumFractionDigits: 0 }) + " GB";
+                                 subItemData[1] = (Number(subItemData[1]) / 1073741824).toLocaleString(undefined, { maximumFractionDigits: 0 }) + " GiB";
                               }
                               subItems.push({ name: subItemData[0], value: subItemData[1] });
                            }
diff --git a/Engine/Source/Programs/Horde/HordeDashboard/src/components/jobDetailsV2/JobDetailViewSteps.tsx b/Engine/Source/Programs/Horde/HordeDashboard/src/components/jobDetailsV2/JobDetailViewSteps.tsx
--- a/Engine/Source/Programs/Horde/HordeDashboard/src/components/jobDetailsV2/JobDetailViewSteps.tsx
+++ b/Engine/Source/Programs/Horde/HordeDashboard/src/components/jobDetailsV2/JobDetailViewSteps.tsx
@@ -361,7 +361,22 @@ export const StepsPanelInner: React.FC<{ jobDetails: JobDetailsV2, depStepId?: s
 
    const stateFilter = jobFilter.filterStates;
    jobDetails.getSteps().forEach(step => {
-      if (stateFilter.indexOf("All") === -1 && stateFilter.indexOf(step.outcome as StateFilter) === -1 && stateFilter.indexOf(step.state) === -1) {
+
+      if (stateFilter.indexOf("All") !== -1) {
+         return;
+      }
+
+      if (stateFilter.indexOf("Skipped") === -1 && step.state === JobStepState.Skipped) {
+         jobStateFilter.add(step.id);
+         return;
+      }
+
+      if (stateFilter.indexOf("Aborted") === -1 && step.state === JobStepState.Aborted) {
+         jobStateFilter.add(step.id);
+         return;
+      }
+
+      if (stateFilter.indexOf(step.outcome as StateFilter) === -1 && stateFilter.indexOf(step.state) === -1) {
          jobStateFilter.add(step.id);
       }
    })
diff --git a/Engine/Source/Programs/Horde/HordeDashboard/src/components/jobDetailsV2/JobDetailsViewCommon.tsx b/Engine/Source/Programs/Horde/HordeDashboard/src/components/jobDetailsV2/JobDetailsViewCommon.tsx
--- a/Engine/Source/Programs/Horde/HordeDashboard/src/components/jobDetailsV2/JobDetailsViewCommon.tsx
+++ b/Engine/Source/Programs/Horde/HordeDashboard/src/components/jobDetailsV2/JobDetailsViewCommon.tsx
@@ -1109,11 +1109,21 @@ export const JobFilterBar: React.FC<{ jobDetails: JobDetailsV2 }> = observer(({
    const stateItems: StateItem[] = ["All", "Failure", "Warnings", "Skipped", "Running", "Completed", "Aborted", "Waiting", "Ready"].map(state => {
       return {
          key: state,
-         text: state === "Aborted" ? "Canceled" : state,
+         text: state,
          state: state as StateFilter
       };
    });
 
+   stateItems.forEach(item => {
+      if (item.state === "Aborted") {
+         item.text = "Canceled";
+      }
+      else if (item.state === "Failure") {
+         item.text = "Errors";
+      }
+
+   });
+
    // steps
    /*
    const step = selected?.step;
