Horde: Send telemetry events for CPU and memory usage from agent to server

diff --git a/Engine/Source/Programs/Horde/Horde.Agent/AgentSettings.cs b/Engine/Source/Programs/Horde/Horde.Agent/AgentSettings.cs
--- a/Engine/Source/Programs/Horde/Horde.Agent/AgentSettings.cs
+++ b/Engine/Source/Programs/Horde/Horde.Agent/AgentSettings.cs
@@ -190,6 +190,11 @@ namespace Horde.Agent
 		/// Incoming port for listening for compute work. Needs to be tied with a lease.
 		/// </summary>
 		public int ComputePort { get; set; } = 7000;
+		
+		/// <summary>
+		/// How often to report telemetry events to server in milliseconds
+		/// </summary>
+		public int TelemetryReportInterval { get; set; } = 10 * 1000;
 
 		/// <summary>
 		/// Key/value properties in addition to those set internally by the agent
diff --git a/Engine/Source/Programs/Horde/Horde.Agent/Horde.Agent.csproj b/Engine/Source/Programs/Horde/Horde.Agent/Horde.Agent.csproj
--- a/Engine/Source/Programs/Horde/Horde.Agent/Horde.Agent.csproj
+++ b/Engine/Source/Programs/Horde/Horde.Agent/Horde.Agent.csproj
@@ -56,6 +56,7 @@
     <PackageReference Include="Serilog.Settings.Configuration" Version="3.4.0" />
     <PackageReference Include="Serilog.Sinks.Console" Version="4.1.0" />
     <PackageReference Include="Serilog.Sinks.File" Version="5.0.0" />
+    <PackageReference Include="System.Diagnostics.PerformanceCounter" Version="6.0.0" />
     <PackageReference Include="System.ServiceProcess.ServiceController" Version="6.0.0" />
   </ItemGroup>
 
@@ -153,6 +154,7 @@
     <Protobuf Include="..\Horde.Server\Protos\horde\messages\agent_workspace.proto" ProtoRoot="..\Horde.Server\Protos" Access="Public" />
     <Protobuf Include="..\Horde.Server\Protos\horde\messages\lease.proto" ProtoRoot="..\Horde.Server\Protos" Access="Public" />
     <Protobuf Include="..\Horde.Server\Protos\horde\messages\aws_lambda.proto" ProtoRoot="..\Horde.Server\Protos" Access="Public" />
+    <Protobuf Include="..\Horde.Server\Protos\horde\messages\telemetry.proto" ProtoRoot="..\Horde.Server\Protos" Access="Public" />
     <Protobuf Include="..\Horde.Server\Protos\horde\tasks\compute_task.proto" ProtoRoot="..\Horde.Server\Protos" Access="Public" />
     <Protobuf Include="..\Horde.Server\Protos\horde\tasks\conform_task.proto" ProtoRoot="..\Horde.Server\Protos" Access="Public" />
     <Protobuf Include="..\Horde.Server\Protos\horde\tasks\job_task.proto" ProtoRoot="..\Horde.Server\Protos" Access="Public" />
diff --git a/Engine/Source/Programs/Horde/Horde.Agent/Services/TelemetryService.cs b/Engine/Source/Programs/Horde/Horde.Agent/Services/TelemetryService.cs
new file mode 100644
--- /dev/null
+++ b/Engine/Source/Programs/Horde/Horde.Agent/Services/TelemetryService.cs
@@ -0,0 +1,285 @@
+// Copyright Epic Games, Inc. All Rights Reserved.
+
+using System;
+using System.Threading;
+using System.Threading.Tasks;
+using Microsoft.Extensions.Hosting;
+using System.Diagnostics;
+using System.Runtime.InteropServices;
+using Grpc.Core;
+using Grpc.Net.Client;
+using HordeCommon.Rpc;
+using HordeCommon.Rpc.Messages.Telemetry;
+using Microsoft.Extensions.Logging;
+using Microsoft.Extensions.Options;
+using Microsoft.Management.Infrastructure;
+
+namespace Horde.Agent.Services;
+
+/// <summary>
+/// Metrics for CPU usage
+/// </summary>
+public class CpuMetrics
+{
+	/// <summary>
+	/// Percentage of time the CPU was busy executing code in user space
+	/// </summary>
+	public float User { get; set; }
+	
+	/// <summary>
+	/// Percentage of time the CPU was busy executing code in kernel space
+	/// </summary>
+	public float System { get; set; }
+	
+	/// <summary>
+	/// Percentage of time the CPU was idling
+	/// </summary>
+	public float Idle { get; set; }
+
+	/// <inheritdoc />
+	public override string ToString()
+	{
+		return $"User={User,5:F1}% System={System,5:F1}% Idle={Idle,5:F1}%";
+	}
+
+	/// <summary>
+	/// Convert to Protobuf-based event
+	/// </summary>
+	/// <returns></returns>
+	public CpuMetricsEvent ToEvent()
+	{
+		return new CpuMetricsEvent { User = User, System = System, Idle = Idle };
+	}
+}
+
+/// <summary>
+/// Metrics for memory usage
+/// </summary>
+public class MemoryMetrics
+{
+	/// <summary>
+	/// Total memory installed (kibibytes)
+	/// </summary>
+	public uint Total { get; set; }
+	
+	/// <summary>
+	/// Available memory (kibibytes)
+	/// </summary>
+	public uint Available { get; set; }
+	
+	/// <summary>
+	/// Used memory (kibibytes)
+	/// </summary>
+	public uint Used { get; set; }
+	
+	/// <summary>
+	/// Used memory (percentage)
+	/// </summary>
+	public float UsedPercentage { get; set; }
+
+	/// <inheritdoc />
+	public override string ToString()
+	{
+		return $"Total={Total} kB, Available={Available} kB, Used={Used} kB, Used={UsedPercentage * 100.0:F1} %";
+	}
+	
+	/// <summary>
+	/// Convert to Protobuf-based event
+	/// </summary>
+	/// <returns></returns>
+	public MemoryMetricsEvent ToEvent()
+	{
+		return new MemoryMetricsEvent { Total = Total, Free = Available, Used = Used, UsedPercentage = UsedPercentage };
+	}
+}
+
+/// <summary>
+/// OS agnostic interface for retrieving system metrics (CPU, memory etc)
+/// </summary>
+public interface ISystemMetrics
+{
+	/// <summary>
+	/// Get CPU usage metrics
+	/// </summary>
+	/// <returns>An object with CPU usage metrics</returns>
+	CpuMetrics GetCpu();
+	
+	/// <summary>
+	/// Get memory usage metrics
+	/// </summary>
+	/// <returns>An object with memory usage metrics</returns>
+	MemoryMetrics GetMemory();
+}
+
+// Suppress call sites not available on all platforms
+#pragma warning disable CA1416
+
+/// <summary>
+/// Windows specific implementation for gathering system metrics
+/// </summary>
+public class WindowsSystemMetrics : ISystemMetrics
+{
+	private const string ProcessorInfo = "Processor Information"; // Prefer this over "Processor" as it's more modern
+	private const string Memory = "Memory";
+	private const string Total = "_Total";
+
+	private readonly PerformanceCounter _procIdleTime = new (ProcessorInfo, "% Idle Time", Total);
+	private readonly PerformanceCounter _procUserTime = new (ProcessorInfo, "% User Time", Total);
+	private readonly PerformanceCounter _procPrivilegedTime = new (ProcessorInfo, "% Privileged Time", Total);
+	
+	private readonly uint _totalPhysicalMemory = GetPhysicalMemory();
+	private readonly PerformanceCounter _memAvailableBytes = new (Memory, "Available Bytes");
+
+	/// <summary>
+	/// Constructor
+	/// </summary>
+	public WindowsSystemMetrics()
+	{
+		GetCpu(); // Trigger this to ensure performance counter has a fetched value. Avoids an initial zero result when called later.
+	}
+
+	/// <inheritdoc />
+	public CpuMetrics GetCpu()
+	{
+		return new CpuMetrics
+		{
+			User = _procUserTime.NextValue(),
+			System = _procPrivilegedTime.NextValue(),
+			Idle = _procIdleTime.NextValue()
+		};
+	}
+
+	/// <inheritdoc />
+	public MemoryMetrics GetMemory()
+	{
+		uint available = (uint)(_memAvailableBytes.NextValue() / 1024);
+		uint used = _totalPhysicalMemory - available;
+		return new MemoryMetrics
+		{
+			Total = _totalPhysicalMemory,
+			Available = available,
+			Used = used,
+			UsedPercentage = used / (float)_totalPhysicalMemory,
+		};
+	}
+
+	private static uint GetPhysicalMemory()
+	{
+		using CimSession session = CimSession.Create(null);
+		const string QueryNamespace = @"root\cimv2";
+		const string QueryDialect = "WQL";
+		ulong totalCapacity = 0;
+		
+		foreach (CimInstance instance in session.QueryInstances(QueryNamespace, QueryDialect, "select Capacity from Win32_PhysicalMemory"))
+		{
+			foreach (CimProperty property in instance.CimInstanceProperties)
+			{
+				if (property.Name.Equals("Capacity", StringComparison.OrdinalIgnoreCase) && property.Value is ulong capacity)
+				{
+					totalCapacity += capacity;
+				}
+			}
+		}
+
+		return (uint)(totalCapacity / 1024); // as kibibytes
+	}
+}
+
+#pragma warning restore CA1416
+
+/// <summary>
+/// Send telemetry events back to server at regular intervals
+/// </summary>
+class TelemetryService : BackgroundService
+{
+	private readonly GrpcService _grpcService;
+	private readonly AgentSettings _agentSettings;
+	private readonly ILogger<TelemetryService> _logger;
+	private readonly TimeSpan _reportInterval;
+	private readonly AgentMetadataEvent _agentMetadataEvent;
+	private readonly TimeSpan _agentMetadataReportInterval = TimeSpan.FromMinutes(10);
+	private ISystemMetrics? _systemMetrics;
+	private DateTime _lastTimeAgentMetadataSent = DateTime.UnixEpoch;
+
+	/// <summary>
+	/// Constructor
+	/// </summary>
+	TelemetryService(GrpcService grpcService, IOptions<AgentSettings> settings, ILogger<TelemetryService> logger)
+	{
+		_grpcService = grpcService;
+		_agentSettings = settings.Value;
+		_logger = logger;
+		_reportInterval = TimeSpan.FromMilliseconds(_agentSettings.TelemetryReportInterval);
+		
+		// Calculate this once at startup as it should not change during lifetime of process
+		_agentMetadataEvent = GetAgentMetadataEvent();
+	}
+ 
+	/// <inheritdoc />
+	public override Task StartAsync(CancellationToken cancellationToken)
+	{
+		if (RuntimeInformation.IsOSPlatform(OSPlatform.Windows))
+		{
+			_systemMetrics = new WindowsSystemMetrics();
+		}
+		else
+		{
+			_logger.LogInformation("System metric collection only implemented on Windows");
+		}
+		
+		return base.StartAsync(cancellationToken);
+	}
+
+	/// <inheritdoc />
+	protected override async Task ExecuteAsync(CancellationToken stoppingToken)
+	{
+		using GrpcChannel channel = _grpcService.CreateGrpcChannel();
+		HordeRpc.HordeRpcClient client = new (channel);
+		
+		while (!stoppingToken.IsCancellationRequested && _systemMetrics != null)
+		{
+			_logger.LogDebug("Sending telemetry events to server...");
+			CpuMetricsEvent cpuMetricsEvent = _systemMetrics.GetCpu().ToEvent();
+			MemoryMetricsEvent memMetricsEvent = _systemMetrics.GetMemory().ToEvent();
+
+			SendTelemetryEventsRequest request = new();
+			request.Events.Add(new WrappedTelemetryEvent { Cpu = cpuMetricsEvent });
+			request.Events.Add(new WrappedTelemetryEvent { Mem = memMetricsEvent });
+
+			if (DateTime.UtcNow > _lastTimeAgentMetadataSent + _agentMetadataReportInterval)
+			{
+				// Report agent metadata every now and then as events are not guaranteed to be delivered.
+				// Re-sending ensures the metadata will eventually make it to the server.
+				request.Events.Add(new WrappedTelemetryEvent { AgentMetadata = _agentMetadataEvent });
+				_lastTimeAgentMetadataSent = DateTime.UtcNow;
+			}
+			
+			await client.SendTelemetryEventsAsync(request, new CallOptions(cancellationToken: stoppingToken));
+			await Task.Delay(_reportInterval, stoppingToken);
+		}
+	}
+
+	AgentMetadataEvent GetAgentMetadataEvent()
+	{
+		return new AgentMetadataEvent
+		{
+			Ip = null,
+			Hostname = _agentSettings.GetAgentName(),
+			Region = null,
+			AvailabilityZone = null,
+			Environment = null,
+			AgentVersion = null,
+			Os = GetOs(),
+			OsVersion = Environment.OSVersion.Version.ToString(),
+			Architecture = RuntimeInformation.OSArchitecture.ToString()
+		};
+	}
+
+	private static string GetOs()
+	{
+		if (RuntimeInformation.IsOSPlatform(OSPlatform.Windows)) { return "Windows"; }
+		if (RuntimeInformation.IsOSPlatform(OSPlatform.Linux)) { return "Linux"; }
+		if (RuntimeInformation.IsOSPlatform(OSPlatform.OSX)) { return "macOS"; }
+		return "Unknown";
+	}
+}
\ No newline at end of file
diff --git a/Engine/Source/Programs/Horde/Horde.Server/Horde.Server.csproj b/Engine/Source/Programs/Horde/Horde.Server/Horde.Server.csproj
--- a/Engine/Source/Programs/Horde/Horde.Server/Horde.Server.csproj
+++ b/Engine/Source/Programs/Horde/Horde.Server/Horde.Server.csproj
@@ -93,6 +93,7 @@
     <Protobuf Include="Protos\horde\messages\agent_requirements.proto" ProtoRoot="Protos" />
     <Protobuf Include="Protos\horde\messages\agent_workspace.proto" ProtoRoot="Protos" />
     <Protobuf Include="Protos\horde\messages\lease.proto" ProtoRoot="Protos" />
+    <Protobuf Include="Protos\horde\messages\telemetry.proto" ProtoRoot="Protos" />
     <Protobuf Include="Protos\horde\tasks\compute_task.proto" ProtoRoot="Protos" />
     <Protobuf Include="Protos\horde\tasks\conform_task.proto" ProtoRoot="Protos" />
     <Protobuf Include="Protos\horde\tasks\job_task.proto" ProtoRoot="Protos" />
