Horde: Trigger a new OIDC login if the current refresh token has expired.

diff --git a/Engine/Source/Programs/Horde/Samples/RemoteClient/RemoteClientApp.cs b/Engine/Source/Programs/Horde/Samples/RemoteClient/RemoteClientApp.cs
--- a/Engine/Source/Programs/Horde/Samples/RemoteClient/RemoteClientApp.cs
+++ b/Engine/Source/Programs/Horde/Samples/RemoteClient/RemoteClientApp.cs
@@ -170,7 +170,16 @@ namespace RemoteClient
 			using ITokenStore tokenStore = TokenStoreFactory.CreateTokenStore();
 			IConfiguration providerConfiguration = ProviderConfigurationFactory.ReadConfiguration(engineDir.ToDirectoryInfo(), projectDir?.ToDirectoryInfo());
 			OidcTokenManager oidcTokenManager = OidcTokenManager.CreateTokenManager(providerConfiguration, tokenStore, new List<string>() { oidcProvider });
-			OidcTokenInfo result = await oidcTokenManager.GetAccessToken(oidcProvider);
+
+			OidcTokenInfo result;
+			try
+			{
+				result = await oidcTokenManager.GetAccessToken(oidcProvider);
+			}
+			catch (NotLoggedInException)
+			{
+				result = await oidcTokenManager.Login(oidcProvider);
+			}
 
 			if (result.AccessToken == null)
 			{
