Horde: Expand small subdirectories in artifacts.

diff --git a/Engine/Source/Programs/Horde/Horde.Server/Artifacts/ArtifactsController.cs b/Engine/Source/Programs/Horde/Horde.Server/Artifacts/ArtifactsController.cs
--- a/Engine/Source/Programs/Horde/Horde.Server/Artifacts/ArtifactsController.cs
+++ b/Engine/Source/Programs/Horde/Horde.Server/Artifacts/ArtifactsController.cs
@@ -87,6 +87,11 @@ namespace Horde.Server.Artifacts
 		/// <inheritdoc cref="FileEntry.Hash"/>
 		public IoHash Hash => _entry.Hash;
 
+		/// <summary>
+		/// Contents of this directory. Included optionally to reduce requests when browsing the tree.
+		/// </summary>
+		public GetArtifactDirectoryResponse? Contents { get; set; }
+
 		internal GetArtifactDirectoryEntryResponse(DirectoryEntry entry) => _entry = entry;
 	}
 
@@ -104,12 +109,6 @@ namespace Horde.Server.Artifacts
 		/// Files within the directory
 		/// </summary>
 		public List<GetArtifactFileEntryResponse> Files { get; } = new List<GetArtifactFileEntryResponse>();
-
-		internal GetArtifactDirectoryResponse(DirectoryNode directory)
-		{
-			Directories.AddRange(directory.Directories.Select(x => new GetArtifactDirectoryEntryResponse(x)));
-			Files.AddRange(directory.Files.Select(x => new GetArtifactFileEntryResponse(x)));
-		}
 	}
 
 	/// <summary>
@@ -260,8 +259,46 @@ namespace Horde.Server.Artifacts
 				}
 			}
 
-			DirectoryNode directory = await directoryRef.ExpandAsync(reader, cancellationToken);
-			return PropertyFilter.Apply(new GetArtifactDirectoryResponse(directory), filter);
+			DirectoryNode directoryNode = await directoryRef.ExpandAsync(reader, cancellationToken);
+			GetArtifactDirectoryResponse response = await ExpandDirectoriesAsync(reader, directoryNode, 0, cancellationToken);
+			return PropertyFilter.Apply(response, filter);
+		}
+
+		async Task<GetArtifactDirectoryResponse> ExpandDirectoriesAsync(TreeReader reader, DirectoryNode directoryNode, int depth, CancellationToken cancellationToken)
+		{
+			GetArtifactDirectoryResponse response = new GetArtifactDirectoryResponse();
+
+			foreach (DirectoryEntry subDirectoryEntry in directoryNode.Directories)
+			{
+				DirectoryNode subDirectoryNode = await subDirectoryEntry.ExpandAsync(reader, cancellationToken);
+
+				GetArtifactDirectoryEntryResponse subDirectoryEntryResponse = new GetArtifactDirectoryEntryResponse(subDirectoryEntry);
+				if (depth == 0)
+				{
+					if (subDirectoryNode.Directories.Count + subDirectoryNode.Files.Count < 15)
+					{
+						subDirectoryEntryResponse.Contents = await ExpandDirectoriesAsync(reader, subDirectoryNode, depth + 1, cancellationToken);
+					}
+				}
+				else if (depth == 1)
+				{
+					if (subDirectoryNode.Directories.Count + subDirectoryNode.Files.Count < 5)
+					{
+						subDirectoryEntryResponse.Contents = await ExpandDirectoriesAsync(reader, subDirectoryNode, depth + 1, cancellationToken);
+					}
+				}
+				else if (depth < 5)
+				{
+					if (subDirectoryNode.Directories.Count == 1 && subDirectoryNode.Files.Count == 0)
+					{
+						subDirectoryEntryResponse.Contents = await ExpandDirectoriesAsync(reader, subDirectoryNode, depth + 1, cancellationToken);
+					}
+				}
+				response.Directories.Add(subDirectoryEntryResponse);
+			}
+
+			response.Files.AddRange(directoryNode.Files.Select(x => new GetArtifactFileEntryResponse(x)));
+			return response;
 		}
 
 		/// <summary>
