Horde: Flatten the response structure for subdirectories and files.

diff --git a/Engine/Source/Programs/Horde/Horde.Server/Artifacts/ArtifactsController.cs b/Engine/Source/Programs/Horde/Horde.Server/Artifacts/ArtifactsController.cs
--- a/Engine/Source/Programs/Horde/Horde.Server/Artifacts/ArtifactsController.cs
+++ b/Engine/Source/Programs/Horde/Horde.Server/Artifacts/ArtifactsController.cs
@@ -74,7 +74,7 @@ namespace Horde.Server.Artifacts
 	/// <summary>
 	/// Describes a file within an artifact
 	/// </summary>
-	public class GetArtifactDirectoryEntryResponse
+	public class GetArtifactDirectoryEntryResponse : GetArtifactDirectoryResponse
 	{
 		readonly DirectoryEntry _entry;
 
@@ -87,11 +87,6 @@ namespace Horde.Server.Artifacts
 		/// <inheritdoc cref="FileEntry.Hash"/>
 		public IoHash Hash => _entry.Hash;
 
-		/// <summary>
-		/// Contents of this directory. Included optionally to reduce requests when browsing the tree.
-		/// </summary>
-		public GetArtifactDirectoryResponse? Contents { get; set; }
-
 		internal GetArtifactDirectoryEntryResponse(DirectoryEntry entry) => _entry = entry;
 	}
 
@@ -103,12 +98,12 @@ namespace Horde.Server.Artifacts
 		/// <summary>
 		/// Names of sub-directories
 		/// </summary>
-		public List<GetArtifactDirectoryEntryResponse> Directories { get; } = new List<GetArtifactDirectoryEntryResponse>();
-
+		public List<GetArtifactDirectoryEntryResponse>? Directories { get; internal set; }
+			
 		/// <summary>
 		/// Files within the directory
 		/// </summary>
-		public List<GetArtifactFileEntryResponse> Files { get; } = new List<GetArtifactFileEntryResponse>();
+		public List<GetArtifactFileEntryResponse>? Files { get; internal set; }
 	}
 
 	/// <summary>
@@ -260,45 +255,51 @@ namespace Horde.Server.Artifacts
 			}
 
 			DirectoryNode directoryNode = await directoryRef.ExpandAsync(reader, cancellationToken);
-			GetArtifactDirectoryResponse response = await ExpandDirectoriesAsync(reader, directoryNode, 0, cancellationToken);
+
+			GetArtifactDirectoryResponse response = new GetArtifactDirectoryResponse();
+			await ExpandDirectoriesAsync(reader, directoryNode, 0, response, cancellationToken);
 			return PropertyFilter.Apply(response, filter);
 		}
 
-		async Task<GetArtifactDirectoryResponse> ExpandDirectoriesAsync(TreeReader reader, DirectoryNode directoryNode, int depth, CancellationToken cancellationToken)
+		async Task ExpandDirectoriesAsync(TreeReader reader, DirectoryNode directoryNode, int depth, GetArtifactDirectoryResponse response, CancellationToken cancellationToken)
 		{
-			GetArtifactDirectoryResponse response = new GetArtifactDirectoryResponse();
-
-			foreach (DirectoryEntry subDirectoryEntry in directoryNode.Directories)
+			if (directoryNode.Directories.Count > 0)
 			{
-				DirectoryNode subDirectoryNode = await subDirectoryEntry.ExpandAsync(reader, cancellationToken);
-
-				GetArtifactDirectoryEntryResponse subDirectoryEntryResponse = new GetArtifactDirectoryEntryResponse(subDirectoryEntry);
-				if (depth == 0)
+				response.Directories = new List<GetArtifactDirectoryEntryResponse>();
+				foreach (DirectoryEntry subDirectoryEntry in directoryNode.Directories)
 				{
-					if (subDirectoryNode.Directories.Count + subDirectoryNode.Files.Count < 15)
+					DirectoryNode subDirectoryNode = await subDirectoryEntry.ExpandAsync(reader, cancellationToken);
+
+					GetArtifactDirectoryEntryResponse subDirectoryEntryResponse = new GetArtifactDirectoryEntryResponse(subDirectoryEntry);
+					if (depth == 0)
 					{
-						subDirectoryEntryResponse.Contents = await ExpandDirectoriesAsync(reader, subDirectoryNode, depth + 1, cancellationToken);
+						if (subDirectoryNode.Directories.Count + subDirectoryNode.Files.Count < 16)
+						{
+							await ExpandDirectoriesAsync(reader, subDirectoryNode, depth + 1, subDirectoryEntryResponse, cancellationToken);
+						}
 					}
-				}
-				else if (depth == 1)
-				{
-					if (subDirectoryNode.Directories.Count + subDirectoryNode.Files.Count < 5)
+					else if (depth == 1)
 					{
-						subDirectoryEntryResponse.Contents = await ExpandDirectoriesAsync(reader, subDirectoryNode, depth + 1, cancellationToken);
+						if (subDirectoryNode.Directories.Count + subDirectoryNode.Files.Count < 8)
+						{
+							await ExpandDirectoriesAsync(reader, subDirectoryNode, depth + 1, subDirectoryEntryResponse, cancellationToken);
+						}
 					}
-				}
-				else if (depth < 5)
-				{
-					if (subDirectoryNode.Directories.Count == 1 && subDirectoryNode.Files.Count == 0)
+					else if (depth < 5)
 					{
-						subDirectoryEntryResponse.Contents = await ExpandDirectoriesAsync(reader, subDirectoryNode, depth + 1, cancellationToken);
+						if (subDirectoryNode.Directories.Count == 1 && subDirectoryNode.Files.Count == 0)
+						{
+							await ExpandDirectoriesAsync(reader, subDirectoryNode, depth, subDirectoryEntryResponse, cancellationToken);
+						}
 					}
+					response.Directories.Add(subDirectoryEntryResponse);
 				}
-				response.Directories.Add(subDirectoryEntryResponse);
 			}
 
-			response.Files.AddRange(directoryNode.Files.Select(x => new GetArtifactFileEntryResponse(x)));
-			return response;
+			if (directoryNode.Files.Count > 0)
+			{
+				response.Files = directoryNode.Files.Select(x => new GetArtifactFileEntryResponse(x)).ToList();
+			}
 		}
 
 		/// <summary>
