Horde: Limit number of expanded subdirectories to 10 levels.

diff --git a/Engine/Source/Programs/Horde/Horde.Server/Artifacts/ArtifactsController.cs b/Engine/Source/Programs/Horde/Horde.Server/Artifacts/ArtifactsController.cs
--- a/Engine/Source/Programs/Horde/Horde.Server/Artifacts/ArtifactsController.cs
+++ b/Engine/Source/Programs/Horde/Horde.Server/Artifacts/ArtifactsController.cs
@@ -285,11 +285,11 @@ namespace Horde.Server.Artifacts
 							await ExpandDirectoriesAsync(reader, subDirectoryNode, depth + 1, subDirectoryEntryResponse, cancellationToken);
 						}
 					}
-					else if (depth < 5)
+					else if (depth < 10)
 					{
 						if (subDirectoryNode.Directories.Count == 1 && subDirectoryNode.Files.Count == 0)
 						{
-							await ExpandDirectoriesAsync(reader, subDirectoryNode, depth, subDirectoryEntryResponse, cancellationToken);
+							await ExpandDirectoriesAsync(reader, subDirectoryNode, depth + 1, subDirectoryEntryResponse, cancellationToken);
 						}
 					}
 					response.Directories.Add(subDirectoryEntryResponse);
