Horde: Account for inline directories
#skipci
#rnx

diff --git a/Engine/Source/Programs/Horde/HordeDashboard/src/components/jobDetailsV2/JobDetailArtifactsV3.tsx b/Engine/Source/Programs/Horde/HordeDashboard/src/components/jobDetailsV2/JobDetailArtifactsV3.tsx
--- a/Engine/Source/Programs/Horde/HordeDashboard/src/components/jobDetailsV2/JobDetailArtifactsV3.tsx
+++ b/Engine/Source/Programs/Horde/HordeDashboard/src/components/jobDetailsV2/JobDetailArtifactsV3.tsx
@@ -146,12 +146,12 @@ class ArtifactsHandler {
       if (!selection?.length) {
          selection = [];
 
-         browse.directories.forEach(d => {
-            selection!.push({ key: d.ioHash, text: d.name, icon: "Folder", type: BrowserType.Directory, size: d.length });
+         browse.directories?.forEach(d => {
+            selection!.push({ key: d.hash, text: d.name, icon: "Folder", type: BrowserType.Directory, size: d.length });
          });
 
-         browse.files.forEach(d => {
-            selection!.push({ key: d.ioHash, text: d.name, icon: "Document", type: BrowserType.File, size: d.length });
+         browse.files?.forEach(d => {
+            selection!.push({ key: d.hash, text: d.name, icon: "Document", type: BrowserType.File, size: d.length });
          });
       }
 
@@ -409,12 +409,26 @@ const JobDetailArtifactsInner: React.FC<{ handler: ArtifactsHandler }> = observe
       items.push({ key: "navigate up", text: "..", type: BrowserType.NavigateUp });
    }
 
-   browse.directories.forEach(d => {
-      items.push({ key: d.ioHash, text: d.name, icon: "Folder", type: BrowserType.Directory, size: d.length });
+   browse.directories?.forEach(d => {
+
+      function recurseDirectories(dir: GetArtifactDirectoryEntryResponse, flattened: GetArtifactDirectoryEntryResponse[]) {
+         if (!dir.directories?.length) {
+            const name = flattened.length ? flattened.map(d => d.name).join("/") + "/" + dir.name : dir.name;
+            console.log("HI", name);
+            items.push({ key: d.hash, text: name, icon: "Folder", type: BrowserType.Directory, size: d.length });
+         } else {
+            flattened.push(dir);
+            dir.directories.forEach(d => recurseDirectories(d, flattened));
+         }
+      }
+
+      recurseDirectories(d, []);
+
+
    });
 
-   browse.files.forEach(d => {
-      items.push({ key: d.ioHash, text: d.name, icon: "Document", type: BrowserType.File, size: d.length });
+   browse.files?.forEach(d => {
+      items.push({ key: d.hash, text: d.name, icon: "Document", type: BrowserType.File, size: d.length });
    });
 
    const columns: IColumn[] = [
@@ -443,7 +457,7 @@ const JobDetailArtifactsInner: React.FC<{ handler: ArtifactsHandler }> = observe
       if (column.name === "View") {
          if (item.type !== BrowserType.File) {
             return null;
-         }         
+         }
          const path = encodeURI(handler.path + "/" + item.text);
          const server = backend.serverUrl;
          const href = `${server}/api/v2/artifacts/${handler.artifact!.id}/file?path=${path}&inline=true`;
