Horde: Fix dashboard build due to missing file
#rnx
#skipci

diff --git a/Engine/Source/Programs/Horde/HordeDashboard/src/backend/Api.ts b/Engine/Source/Programs/Horde/HordeDashboard/src/backend/Api.ts
--- a/Engine/Source/Programs/Horde/HordeDashboard/src/backend/Api.ts
+++ b/Engine/Source/Programs/Horde/HordeDashboard/src/backend/Api.ts
@@ -1454,24 +1454,25 @@ export type GetArtifactFileEntryResponse = {
 
 	name: string;
 	length: number;
-	ioHash: string;
-}
-
-/** Describes a directory within an artifact */
-export type GetArtifactDirectoryEntryResponse = {
-	name: string;
-	length: number;
-	ioHash: string;
+	hash: string;
 }
 
 /** Describes a directory within an artifact */
 export type GetArtifactDirectoryResponse = {
 
 	/** Names of sub-directories */
-	directories: GetArtifactDirectoryEntryResponse[];
+	directories?: GetArtifactDirectoryEntryResponse[];
 
 	/** Files within the directory */
-	files: GetArtifactFileEntryResponse[];
+	files?: GetArtifactFileEntryResponse[];
+}
+
+
+/** Describes a directory within an artifact */
+export type GetArtifactDirectoryEntryResponse = GetArtifactDirectoryResponse & {
+	name: string;
+	length: number;
+	hash: string;
 }
 
 /**Parameters required to update a log file */
