Horde: Artifact browser improvements
#skipci
#rnx

diff --git a/Engine/Source/Programs/Horde/HordeDashboard/src/components/jobDetailsV2/JobDetailArtifactsV3.tsx b/Engine/Source/Programs/Horde/HordeDashboard/src/components/jobDetailsV2/JobDetailArtifactsV3.tsx
--- a/Engine/Source/Programs/Horde/HordeDashboard/src/components/jobDetailsV2/JobDetailArtifactsV3.tsx
+++ b/Engine/Source/Programs/Horde/HordeDashboard/src/components/jobDetailsV2/JobDetailArtifactsV3.tsx
@@ -380,6 +380,8 @@ const DownloadButton: React.FC<{ handler: ArtifactsHandler }> = observer(({ hand
 
 });
 
+let idcounter = 0;
+
 const JobDetailArtifactsInner: React.FC<{ handler: ArtifactsHandler }> = observer(({ handler }) => {
 
    // subscribe
@@ -413,12 +415,11 @@ const JobDetailArtifactsInner: React.FC<{ handler: ArtifactsHandler }> = observe
 
       function recurseDirectories(dir: GetArtifactDirectoryEntryResponse, flattened: GetArtifactDirectoryEntryResponse[]) {
          if (!dir.directories?.length) {
-            const name = flattened.length ? flattened.map(d => d.name).join("/") + "/" + dir.name : dir.name;
-            console.log("HI", name);
+            const name = flattened.length ? flattened.map(d => d.name).join("/") + "/" + dir.name : dir.name;            
             items.push({ key: d.hash, text: name, icon: "Folder", type: BrowserType.Directory, size: d.length });
          } else {
             flattened.push(dir);
-            dir.directories.forEach(d => recurseDirectories(d, flattened));
+            dir.directories.forEach(d => recurseDirectories(d, [...flattened]));
          }
       }
 
@@ -526,7 +527,7 @@ const JobDetailArtifactsInner: React.FC<{ handler: ArtifactsHandler }> = observe
 
    }
 
-   return <Stack tokens={{ childrenGap: 12 }}>
+   return <Stack key={ `jobdetailartifacts_${idcounter++}`} tokens={{ childrenGap: 12 }}>
       <Stack horizontal verticalAlign="center" tokens={{ childrenGap: 18 }} style={{ paddingBottom: 12 }}>
          <Stack>
             <Dropdown className={hordeClasses.modal} style={{ width: 128 }} options={coptions} selectedKey={handler.context} onChange={(ev, option) => {
