Horde: Downgrade telemetry warnings to information.

diff --git a/Engine/Source/Programs/Horde/Horde.Agent/Execution/JobExecutor.cs b/Engine/Source/Programs/Horde/Horde.Agent/Execution/JobExecutor.cs
--- a/Engine/Source/Programs/Horde/Horde.Agent/Execution/JobExecutor.cs
+++ b/Engine/Source/Programs/Horde/Horde.Agent/Execution/JobExecutor.cs
@@ -1351,7 +1351,7 @@ namespace Horde.Agent.Execution
 				}
 				catch (Exception ex)
 				{
-					jobLogger.LogWarning(ex, "Unable to parse/upload telemetry files: {Message}", ex.Message);
+					jobLogger.LogInformation(ex, "Unable to parse/upload telemetry files: {Message}", ex.Message);
 				}
 
 				await CreateArtifactsAsync(step.StepId, JobArtifactType.Trace, workspaceDir, telemetryFiles, jobLogger, cancellationToken);
