Horde: Fix infinite loop when purging excess deployments from a tool.

diff --git a/Engine/Source/Programs/Horde/Horde.Server/Tools/ToolCollection.cs b/Engine/Source/Programs/Horde/Horde.Server/Tools/ToolCollection.cs
--- a/Engine/Source/Programs/Horde/Horde.Server/Tools/ToolCollection.cs
+++ b/Engine/Source/Programs/Horde/Horde.Server/Tools/ToolCollection.cs
@@ -304,7 +304,7 @@ namespace Horde.Server.Tools
 
 			// If there are already a maximum number of deployments, remove the oldest one
 			const int MaxDeploymentCount = 5;
-			while (tool.Deployments.Count >= MaxDeploymentCount)
+			while (newTool.Deployments.Count >= MaxDeploymentCount)
 			{
 				newTool = await _tools.UpdateAsync(newTool, Builders<Tool>.Update.PopFirst(x => x.Deployments));
 				if (newTool == null)
