Horde: Support parsing traces with explicit parent span ids.

diff --git a/Engine/Source/Programs/Horde/Horde.Agent/Execution/JobExecutor.cs b/Engine/Source/Programs/Horde/Horde.Agent/Execution/JobExecutor.cs
--- a/Engine/Source/Programs/Horde/Horde.Agent/Execution/JobExecutor.cs
+++ b/Engine/Source/Programs/Horde/Horde.Agent/Execution/JobExecutor.cs
@@ -125,6 +125,8 @@ namespace Horde.Agent.Execution
 		class TraceEvent
 		{
 			public string Name { get; set; } = "Unknown";
+			public string? SpanId { get; set; }
+			public string? ParentId { get; set; }
 			public string? Service { get; set; }
 			public string? Resource { get; set; }
 			public DateTimeOffset StartTime { get; set; }
@@ -143,12 +145,20 @@ namespace Horde.Agent.Execution
 		class TraceSpan
 		{
 			public string? Name { get; set; }
+			public string? SpanId { get; set; }
+			public string? ParentId { get; set; }
 			public string? Service { get; set; }
 			public string? Resource { get; set; }
 			public long Start { get; set; }
 			public long Finish { get; set; }
 			public Dictionary<string, string>? Properties { get; set; }
 			public List<TraceSpan>? Children { get; set; }
+
+			public void AddChild(TraceSpan child)
+			{
+				Children ??= new List<TraceSpan>();
+				Children.Add(child);
+			}
 		}
 
 		class TestDataItem
@@ -1460,9 +1470,9 @@ namespace Horde.Agent.Execution
 			}
 
 			List<TraceEvent> telemetrySpans = new List<TraceEvent>();
-			foreach (TraceEventList telemetry in telemetryList.OrderBy(x => x.Spans.First().StartTime).ThenBy(x => x.Spans.Last().FinishTime))
+			foreach (TraceEventList telemetryEventList in telemetryList.OrderBy(x => x.Spans.First().StartTime).ThenBy(x => x.Spans.Last().FinishTime))
 			{
-				foreach (TraceEvent span in telemetry.Spans)
+				foreach (TraceEvent span in telemetryEventList.Spans)
 				{
 					if (span.FinishTime - span.StartTime > TimeSpan.FromMilliseconds(1.0))
 					{
@@ -1481,10 +1491,15 @@ namespace Horde.Agent.Execution
 				Stack<TraceSpan> stack = new Stack<TraceSpan>();
 				stack.Push(rootSpan);
 
+				List<TraceSpan> spansWithExplicitParent = new List<TraceSpan>();
+				Dictionary<string, TraceSpan> spans = new Dictionary<string, TraceSpan>(StringComparer.OrdinalIgnoreCase);
+
 				foreach (TraceEvent traceEvent in telemetrySpans.OrderBy(x => x.StartTime).ThenByDescending(x => x.FinishTime).ThenBy(x => x.Index))
 				{
 					TraceSpan newSpan = new TraceSpan();
 					newSpan.Name = traceEvent.Name;
+					newSpan.SpanId = traceEvent.SpanId;
+					newSpan.ParentId = traceEvent.ParentId;
 					newSpan.Service = traceEvent.Service;
 					newSpan.Resource = traceEvent.Resource;
 					newSpan.Start = traceEvent.StartTime.UtcTicks;
@@ -1494,6 +1509,11 @@ namespace Horde.Agent.Execution
 						newSpan.Properties = traceEvent.Metadata;
 					}
 
+					if (!String.IsNullOrEmpty(newSpan.SpanId))
+					{
+						spans.Add(newSpan.SpanId, newSpan);
+					}
+
 					TraceSpan stackTop = stack.Peek();
 					while (stack.Count > 1 && newSpan.Start >= stackTop.Finish)
 					{
@@ -1501,21 +1521,36 @@ namespace Horde.Agent.Execution
 						stackTop = stack.Peek();
 					}
 
-					if (stack.Count > 1 && newSpan.Finish > stackTop.Finish)
+					if (String.IsNullOrEmpty(newSpan.ParentId))
 					{
-						jobLogger.LogInformation("Trace event name='{Name}', service'{Service}', resource='{Resource}' has invalid finish time ({SpanFinish} < {StackFinish})", newSpan.Name, newSpan.Service, newSpan.Resource, newSpan.Finish, stackTop.Finish);
-						newSpan.Finish = stackTop.Finish;
-					}
+						if (stack.Count > 1 && newSpan.Finish > stackTop.Finish)
+						{
+							jobLogger.LogInformation("Trace event name='{Name}', service'{Service}', resource='{Resource}' has invalid finish time ({SpanFinish} < {StackFinish})", newSpan.Name, newSpan.Service, newSpan.Resource, newSpan.Finish, stackTop.Finish);
+							newSpan.Finish = stackTop.Finish;
+						}
 
-					if (stackTop.Children == null)
+						stackTop.AddChild(newSpan);
+					}
+					else
 					{
-						stackTop.Children = new List<TraceSpan>();
+						spansWithExplicitParent.Add(newSpan);
 					}
 
-					stackTop.Children.Add(newSpan);
 					stack.Push(newSpan);
 				}
 
+				foreach (TraceSpan span in spansWithExplicitParent)
+				{
+					if (span.ParentId != null && spans.TryGetValue(span.ParentId, out TraceSpan? parentSpan))
+					{
+						parentSpan.AddChild(span);
+					}
+					else
+					{
+						_logger.LogInformation("Parent {ParentId} of span {SpanId} was not found.", span.ParentId, span.SpanId);
+					}
+				}
+
 				rootSpan.Start = rootSpan.Children!.First().Start;
 				rootSpan.Finish = rootSpan.Children!.Last().Finish;
 
