Horde: Tweak trendsview to be able to display full range and also start zoomed out
#skipci
#rnx

diff --git a/Engine/Source/Programs/Horde/HordeDashboard/src/components/jobDetailsV2/JobDetailStepTrends.tsx b/Engine/Source/Programs/Horde/HordeDashboard/src/components/jobDetailsV2/JobDetailStepTrends.tsx
--- a/Engine/Source/Programs/Horde/HordeDashboard/src/components/jobDetailsV2/JobDetailStepTrends.tsx
+++ b/Engine/Source/Programs/Horde/HordeDashboard/src/components/jobDetailsV2/JobDetailStepTrends.tsx
@@ -430,7 +430,7 @@ const StepTrendsPanelInner: React.FC<{ jobDetails: JobDetailsV2; dataView: Trend
    const [includeWaitTimes, setIncludeWaitTimes] = useState<boolean>(false);
    const [includeStepDependencies, setIncludeStepDependencies] = useState<boolean>(false);
 
-   const [pageInfo, setPageInfo] = useState<PageInfo>({ pageNumber: 1, zoomLevel: 1, pageSize: 10 });
+   const [pageInfo, setPageInfo] = useState<PageInfo>({ pageNumber: 1, zoomLevel: 10, pageSize: 10 });
 
    dataView.subscribe();
 
@@ -452,15 +452,28 @@ const StepTrendsPanelInner: React.FC<{ jobDetails: JobDetailsV2; dataView: Trend
       findTemplate();
       return null;
    }
+
+   // @todo: this view needs a reboot
+   let count = selectedTemplateRef.id?.indexOf("incremental") === -1 ? 150 : 100;
+   if (step) {
+      count *= 2;
+   }
+
    if (!dataView.cacheData) {
       const query: JobTimingsQuery = {
          streamId: jobDetails!.stream!.id,
          template: [selectedTemplateRef.id],
-         count: selectedTemplateRef.id?.indexOf("incremental") === -1 ? 50 : 100
+         count: count
       };
       dataView.populateChartDataCache(query);
    }
    
+   if (dataView.cacheData?.chartData) {
+      let totalValues = Object.values(dataView.cacheData.chartData!.labelData).map(c => c.length).reduce((a, b) => Math.max(a, b));   
+      pageInfo.pageSize = totalValues / 10;
+   }
+   
+
 
    const filterChartData = (chartDataCache?: ChartDataCache) => {
       if (!chartDataCache) {
