Horde: Allow querying the deployed server version without auth.

diff --git a/Engine/Source/Programs/Horde/Horde.Server/Server/ServerController.cs b/Engine/Source/Programs/Horde/Horde.Server/Server/ServerController.cs
--- a/Engine/Source/Programs/Horde/Horde.Server/Server/ServerController.cs
+++ b/Engine/Source/Programs/Horde/Horde.Server/Server/ServerController.cs
@@ -38,6 +38,7 @@ namespace Horde.Server.Server
 		/// Get server version
 		/// </summary>
 		[HttpGet]
+		[AllowAnonymous]
 		[Route("/api/v1/server/version")]
 		public ActionResult GetVersionAsync()
 		{
@@ -49,6 +50,7 @@ namespace Horde.Server.Server
 		/// Get server information
 		/// </summary>
 		[HttpGet]
+		[AllowAnonymous]
 		[Route("/api/v1/server/info")]
 		[ProducesResponseType(typeof(GetServerInfoResponse), 200)]
 		public async Task<ActionResult<GetServerInfoResponse>> GetServerInfo()
