Horde: Pump up the default debug token expiry to 8 hours
#rnx

diff --git a/Engine/Source/Programs/Horde/Horde.Server/ServerSettings.cs b/Engine/Source/Programs/Horde/Horde.Server/ServerSettings.cs
--- a/Engine/Source/Programs/Horde/Horde.Server/ServerSettings.cs
+++ b/Engine/Source/Programs/Horde/Horde.Server/ServerSettings.cs
@@ -426,7 +426,7 @@ namespace Horde.Server
 		/// <summary>
 		/// Length of time before JWT tokens expire, in hours
 		/// </summary>
-		public int JwtExpiryTimeHours { get; set; } = 4;
+		public int JwtExpiryTimeHours { get; set; } = 8;
 
 		/// <summary>
 		/// Whether to enable Cors, generally for development purposes
