Horde: Add success counts to test suite refs
#rnx

diff --git a/Engine/Source/Programs/Horde/Horde.Server/Jobs/TestData/ITestData.cs b/Engine/Source/Programs/Horde/Horde.Server/Jobs/TestData/ITestData.cs
--- a/Engine/Source/Programs/Horde/Horde.Server/Jobs/TestData/ITestData.cs
+++ b/Engine/Source/Programs/Horde/Horde.Server/Jobs/TestData/ITestData.cs
@@ -214,8 +214,13 @@ namespace Horde.Server.Jobs.TestData
 		/// </summary>
 		int? SuiteErrorCount { get; }
 
+		/// <summary>
+		/// Suite test successes
+		/// </summary>
+		int? SuiteSuccessCount { get; }
+
 	}
-	
+
 	/// <summary>
 	/// Test data details
 	/// </summary>
diff --git a/Engine/Source/Programs/Horde/Horde.Server/Jobs/TestData/TestDataCollection.cs b/Engine/Source/Programs/Horde/Horde.Server/Jobs/TestData/TestDataCollection.cs
--- a/Engine/Source/Programs/Horde/Horde.Server/Jobs/TestData/TestDataCollection.cs
+++ b/Engine/Source/Programs/Horde/Horde.Server/Jobs/TestData/TestDataCollection.cs
@@ -266,6 +266,9 @@ namespace Horde.Server.Jobs.TestData
 			[BsonIgnoreIfNull, BsonElement("sec")]
 			public int? SuiteErrorCount { get; set; }
 
+			[BsonIgnoreIfNull, BsonElement("spc")]
+			public int? SuiteSuccessCount { get; set; }
+
 			private TestDataRefDocument()
 			{
 
@@ -1656,6 +1659,7 @@ namespace Horde.Server.Jobs.TestData
 						int skipCount = 0;
 						int warningCount = 0;
 						int errorCount = 0;
+						int successCount = 0;
 
 						// populate the suite tests
 						List<SuiteTestData> suiteTestData = new List<SuiteTestData>();
@@ -1678,6 +1682,7 @@ namespace Horde.Server.Jobs.TestData
 							{
 								case "Success":
 									outcome = TestOutcome.Success;
+									successCount++;
 									break;
 								case "Skipped":
 									outcome = TestOutcome.Skipped;
@@ -1693,7 +1698,8 @@ namespace Horde.Server.Jobs.TestData
 
 						testRef.SuiteSkipCount = skipCount;
 						testRef.SuiteErrorCount = errorCount;
-						testRef.SuiteWarningCount =warningCount;
+						testRef.SuiteWarningCount = warningCount;
+						testRef.SuiteSuccessCount = successCount;
 
 						// and add to collections
 						await AddTestRef(testRef);
diff --git a/Engine/Source/Programs/Horde/Horde.Server/Jobs/TestData/TestDataMessages.cs b/Engine/Source/Programs/Horde/Horde.Server/Jobs/TestData/TestDataMessages.cs
--- a/Engine/Source/Programs/Horde/Horde.Server/Jobs/TestData/TestDataMessages.cs
+++ b/Engine/Source/Programs/Horde/Horde.Server/Jobs/TestData/TestDataMessages.cs
@@ -425,6 +425,11 @@ namespace Horde.Server.Jobs.TestData
 		/// </summary>
 		public int? SuiteErrorCount { get; set; }
 
+		/// <summary>
+		/// Suite test successes
+		/// </summary>
+		public int? SuiteSuccessCount { get; set; }
+
 		/// <summary>
 		/// Constructor
 		/// </summary>
@@ -445,6 +450,7 @@ namespace Horde.Server.Jobs.TestData
 			SuiteSkipCount = testData.SuiteSkipCount;
 			SuiteWarningCount = testData.SuiteWarningCount;
 			SuiteErrorCount = testData.SuiteErrorCount;
+			SuiteSuccessCount = testData.SuiteSuccessCount;
 		}
 	}
 }
