Horde: Adding some more test data to responses
#rnx

diff --git a/Engine/Source/Programs/Horde/Horde.Server/Jobs/TestData/ITestData.cs b/Engine/Source/Programs/Horde/Horde.Server/Jobs/TestData/ITestData.cs
--- a/Engine/Source/Programs/Horde/Horde.Server/Jobs/TestData/ITestData.cs
+++ b/Engine/Source/Programs/Horde/Horde.Server/Jobs/TestData/ITestData.cs
@@ -154,6 +154,17 @@ namespace Horde.Server.Jobs.TestData
 		/// Test UID for looking up in test details
 		/// </summary>
 		string UID { get; }
+
+		/// <summary>
+		/// The number of warnings
+		/// </summary>
+		int? WarningCount { get; }
+
+		/// <summary>
+		/// The number of errors
+		/// </summary>
+		int? ErrorCount { get; }
+
 	}
 
 	/// <summary>
@@ -169,6 +180,16 @@ namespace Horde.Server.Jobs.TestData
 		/// </summary>
 		StreamId StreamId { get; }
 
+		/// <summary>
+		/// The associated job
+		/// </summary>
+		JobId? JobId { get; }
+
+		/// <summary>
+		/// The associated job step
+		/// </summary>
+		SubResourceId? StepId { get; }
+
 		/// <summary>
 		/// How long the test ran
 		/// </summary>
diff --git a/Engine/Source/Programs/Horde/Horde.Server/Jobs/TestData/TestDataCollection.cs b/Engine/Source/Programs/Horde/Horde.Server/Jobs/TestData/TestDataCollection.cs
--- a/Engine/Source/Programs/Horde/Horde.Server/Jobs/TestData/TestDataCollection.cs
+++ b/Engine/Source/Programs/Horde/Horde.Server/Jobs/TestData/TestDataCollection.cs
@@ -220,12 +220,20 @@ namespace Horde.Server.Jobs.TestData
 			[BsonRequired, BsonElement("uid")]
 			public string UID { get; set; }
 
-			public SuiteTestData(TestId id, TestOutcome outcome, TimeSpan duration, string uid)
+			[BsonIgnoreIfNull, BsonElement("wc")]
+			public int? WarningCount { get; set; }
+
+			[BsonIgnoreIfNull, BsonElement("ec")]
+			public int? ErrorCount { get; set; }
+
+			public SuiteTestData(TestId id, TestOutcome outcome, TimeSpan duration, string uid, int? warningCount, int? errorCount)
 			{
 				TestId = id;
 				Outcome = outcome;
 				Duration = duration;
 				UID = uid;
+				WarningCount = warningCount;
+				ErrorCount = errorCount;
 			}
 		}
 
@@ -269,6 +277,12 @@ namespace Horde.Server.Jobs.TestData
 			[BsonIgnoreIfNull, BsonElement("spc")]
 			public int? SuiteSuccessCount { get; set; }
 
+			[BsonIgnoreIfNull, BsonElement("jid")]
+			public JobId? JobId { get; set; }
+
+			[BsonIgnoreIfNull, BsonElement("sid")]
+			public SubResourceId? StepId { get; set; }
+
 			private TestDataRefDocument()
 			{
 
@@ -1516,6 +1530,8 @@ namespace Horde.Server.Jobs.TestData
 
 				TestDataRefDocument testRef = new TestDataRefDocument(job.StreamId);
 				testRef.StreamId = job.StreamId;
+				testRef.JobId = job.Id;
+				testRef.StepId = step.Id;
 				testRef.Metadata = metaData.Id;
 				testRef.BuildChangeList = testData.BuildChangeList <= 0 ? job.Change : testData.BuildChangeList;
 				testRef.Duration = TimeSpan.FromSeconds(testData.TotalDurationSeconds);
@@ -1650,6 +1666,8 @@ namespace Horde.Server.Jobs.TestData
 
 						TestDataRefDocument testRef = new TestDataRefDocument(job.StreamId);
 						testRef.StreamId = job.StreamId;
+						testRef.JobId = job.Id;
+						testRef.StepId = step.Id;
 						testRef.Metadata = metaData.Id;
 						testRef.BuildChangeList = job.Change;
 						testRef.Duration = TimeSpan.FromSeconds(session.TestSessionInfo!.TimeElapseSec);
@@ -1665,13 +1683,16 @@ namespace Horde.Server.Jobs.TestData
 						List<SuiteTestData> suiteTestData = new List<SuiteTestData>();
 						foreach (AutomatedTestSessionData.SessionTest test in suiteTests)
 						{
+							bool counted = false;
 							if (test.WarningCount > 0)
 							{
+								counted = true;
 								warningCount++;
 							}
 
 							if (test.ErrorCount > 0)
 							{
+								counted = true;
 								errorCount++;
 							}
 
@@ -1682,18 +1703,28 @@ namespace Horde.Server.Jobs.TestData
 							{
 								case "Success":
 									outcome = TestOutcome.Success;
-									successCount++;
+									if (!counted)
+									{
+										successCount++;
+									}
 									break;
 								case "Skipped":
 									outcome = TestOutcome.Skipped;
-									skipCount++;
+									if (!counted)
+									{
+										skipCount++;
+									}										
 									break;
 								case "Fail":
 									outcome = TestOutcome.Failure;
+									if (!counted)
+									{
+										errorCount++;
+									}
 									break;
 							}
 
-							suiteTestData.Add(new SuiteTestData(document.Id, outcome, TimeSpan.FromSeconds(test.TimeElapseSec), test.TestUID));
+							suiteTestData.Add(new SuiteTestData(document.Id, outcome, TimeSpan.FromSeconds(test.TimeElapseSec), test.TestUID, test.WarningCount != 0 ? test.WarningCount : null , test.ErrorCount != 0 ? test.ErrorCount : null));
 						}
 
 						testRef.SuiteSkipCount = skipCount;
diff --git a/Engine/Source/Programs/Horde/Horde.Server/Jobs/TestData/TestDataMessages.cs b/Engine/Source/Programs/Horde/Horde.Server/Jobs/TestData/TestDataMessages.cs
--- a/Engine/Source/Programs/Horde/Horde.Server/Jobs/TestData/TestDataMessages.cs
+++ b/Engine/Source/Programs/Horde/Horde.Server/Jobs/TestData/TestDataMessages.cs
@@ -320,6 +320,16 @@ namespace Horde.Server.Jobs.TestData
 		/// </summary>
 		public string UID { get; set; }
 
+		/// <summary>
+		/// The number of test warnings generated
+		/// </summary>
+		public int? WarningCount { get; set; }
+
+		/// <summary>
+		/// The number of test errors generated
+		/// </summary>
+		public int? ErrorCount { get; set; }
+
 		/// <summary>
 		/// Constructor
 		/// </summary>
@@ -330,6 +340,8 @@ namespace Horde.Server.Jobs.TestData
 			Outcome = data.Outcome;
 			Duration = data.Duration;
 			UID = data.UID;
+			WarningCount = data.WarningCount;
+			ErrorCount = data.ErrorCount;
 		}
 	}
 
@@ -380,6 +392,16 @@ namespace Horde.Server.Jobs.TestData
 		/// </summary>
 		public string StreamId { get; set; }
 
+		/// <summary>
+		/// The associated job id
+		/// </summary>
+		public string? JobId { get; set; }
+
+		/// <summary>
+		/// The associated step id
+		/// </summary>
+		public string? StepId { get; set; }
+
 		/// <summary>
 		/// How long the test ran
 		/// </summary>
@@ -438,6 +460,8 @@ namespace Horde.Server.Jobs.TestData
 		{
 			Id = testData.Id.ToString();
 			StreamId = testData.StreamId.ToString();
+			JobId = testData.JobId?.ToString();
+			StepId = testData.StepId?.ToString();
 			Duration = testData.Duration;
 			BuildChangeList = testData.BuildChangeList;
 			MetaId = testData.Metadata.ToString();
