Horde: Fix BSON element name clash

diff --git a/Engine/Source/Programs/Horde/Horde.Server/Jobs/TestData/TestDataCollection.cs b/Engine/Source/Programs/Horde/Horde.Server/Jobs/TestData/TestDataCollection.cs
--- a/Engine/Source/Programs/Horde/Horde.Server/Jobs/TestData/TestDataCollection.cs
+++ b/Engine/Source/Programs/Horde/Horde.Server/Jobs/TestData/TestDataCollection.cs
@@ -280,7 +280,7 @@ namespace Horde.Server.Jobs.TestData
 			[BsonIgnoreIfNull, BsonElement("jid")]
 			public JobId? JobId { get; set; }
 
-			[BsonIgnoreIfNull, BsonElement("sid")]
+			[BsonIgnoreIfNull, BsonElement("stepid")]
 			public SubResourceId? StepId { get; set; }
 
 			private TestDataRefDocument()
