Horde: Add support for multiple telemetry sinks

Including support for ClickHouse.

diff --git a/Engine/Source/Programs/Horde/Horde.Agent/Services/TelemetryService.cs b/Engine/Source/Programs/Horde/Horde.Agent/Services/TelemetryService.cs
--- a/Engine/Source/Programs/Horde/Horde.Agent/Services/TelemetryService.cs
+++ b/Engine/Source/Programs/Horde/Horde.Agent/Services/TelemetryService.cs
@@ -46,9 +46,9 @@ public class CpuMetrics
 	/// Convert to Protobuf-based event
 	/// </summary>
 	/// <returns></returns>
-	public CpuMetricsEvent ToEvent()
+	public AgentCpuMetricsEvent ToEvent()
 	{
-		return new CpuMetricsEvent { User = User, System = System, Idle = Idle };
+		return new AgentCpuMetricsEvent { User = User, System = System, Idle = Idle };
 	}
 }
 
@@ -87,9 +87,9 @@ public class MemoryMetrics
 	/// Convert to Protobuf-based event
 	/// </summary>
 	/// <returns></returns>
-	public MemoryMetricsEvent ToEvent()
+	public AgentMemoryMetricsEvent ToEvent()
 	{
-		return new MemoryMetricsEvent { Total = Total, Free = Available, Used = Used, UsedPercentage = UsedPercentage };
+		return new AgentMemoryMetricsEvent { Total = Total, Free = Available, Used = Used, UsedPercentage = UsedPercentage };
 	}
 }
 
@@ -239,8 +239,8 @@ class TelemetryService : BackgroundService
 		while (!stoppingToken.IsCancellationRequested && _systemMetrics != null)
 		{
 			_logger.LogDebug("Sending telemetry events to server...");
-			CpuMetricsEvent cpuMetricsEvent = _systemMetrics.GetCpu().ToEvent();
-			MemoryMetricsEvent memMetricsEvent = _systemMetrics.GetMemory().ToEvent();
+			AgentCpuMetricsEvent cpuMetricsEvent = _systemMetrics.GetCpu().ToEvent();
+			AgentMemoryMetricsEvent memMetricsEvent = _systemMetrics.GetMemory().ToEvent();
 
 			SendTelemetryEventsRequest request = new();
 			request.Events.Add(new WrappedTelemetryEvent { Cpu = cpuMetricsEvent });
diff --git a/Engine/Source/Programs/Horde/Horde.Server/Server/RpcService.cs b/Engine/Source/Programs/Horde/Horde.Server/Server/RpcService.cs
--- a/Engine/Source/Programs/Horde/Horde.Server/Server/RpcService.cs
+++ b/Engine/Source/Programs/Horde/Horde.Server/Server/RpcService.cs
@@ -19,9 +19,11 @@ using Horde.Server.Agents.Sessions;
 using Horde.Server.Agents.Software;
 using Horde.Server.Jobs;
 using Horde.Server.Tasks;
+using Horde.Server.Telemetry;
 using Horde.Server.Tools;
 using Horde.Server.Utilities;
 using HordeCommon.Rpc;
+using HordeCommon.Rpc.Messages.Telemetry;
 using Microsoft.AspNetCore.Authorization;
 using Microsoft.Extensions.Logging;
 using Microsoft.Extensions.Options;
@@ -50,6 +52,7 @@ namespace Horde.Server.Server
 		readonly AgentService _agentService;
 		readonly PoolService _poolService;
 		readonly LifetimeService _lifetimeService;
+		readonly ITelemetrySink _telemetrySink;
 		readonly ConformTaskSource _conformTaskSource;
 		readonly JobRpcCommon _jobRpcCommon;
 		readonly IToolCollection _toolCollection;
@@ -59,11 +62,12 @@ namespace Horde.Server.Server
 		/// <summary>
 		/// Constructor
 		/// </summary>
-		public RpcService(AgentService agentService, PoolService poolService, LifetimeService lifetimeService, ConformTaskSource conformTaskSource, JobRpcCommon jobRpcCommon, IToolCollection toolCollection, IOptionsSnapshot<GlobalConfig> globalConfig, ILogger<RpcService> logger)
+		public RpcService(AgentService agentService, PoolService poolService, LifetimeService lifetimeService, ITelemetrySink telemetrySink, ConformTaskSource conformTaskSource, JobRpcCommon jobRpcCommon, IToolCollection toolCollection, IOptionsSnapshot<GlobalConfig> globalConfig, ILogger<RpcService> logger)
 		{
 			_agentService = agentService;
 			_poolService = poolService;
 			_lifetimeService = lifetimeService;
+			_telemetrySink = telemetrySink;
 			_conformTaskSource = conformTaskSource;
 			_jobRpcCommon = jobRpcCommon;
 			_toolCollection = toolCollection;
@@ -461,6 +465,27 @@ namespace Horde.Server.Server
 			}
 		}
 
+		/// <summary>
+		/// Receives telemetry events from agents
+		/// </summary>
+		/// <param name="request">Request arguments</param>
+		/// <param name="context">Context for the RPC call</param>
+		/// <returns>An empty response</returns>
+		public override Task<Empty> SendTelemetryEvents(SendTelemetryEventsRequest request, ServerCallContext context)
+		{
+			foreach (WrappedTelemetryEvent e in request.Events)
+			{
+				switch (e.EventCase)
+				{
+					case WrappedTelemetryEvent.EventOneofCase.AgentMetadata: _telemetrySink.SendEvent("Agent.Metadata", e.AgentMetadata); break;
+					case WrappedTelemetryEvent.EventOneofCase.Cpu: _telemetrySink.SendEvent("Agent.Cpu", e.Cpu); break;
+					case WrappedTelemetryEvent.EventOneofCase.Mem: _telemetrySink.SendEvent("Agent.Memory", e.Mem); break;
+				}
+			}
+
+			return Task.FromResult(new Empty());
+		}
+
 		/// <inheritdoc/>
 		public override Task<Empty> WriteOutput(WriteOutputRequest request, ServerCallContext context) => _jobRpcCommon.WriteOutput(request, context);
 
diff --git a/Engine/Source/Programs/Horde/Horde.Server/ServerSettings.cs b/Engine/Source/Programs/Horde/Horde.Server/ServerSettings.cs
--- a/Engine/Source/Programs/Horde/Horde.Server/ServerSettings.cs
+++ b/Engine/Source/Programs/Horde/Horde.Server/ServerSettings.cs
@@ -182,18 +182,34 @@ namespace Horde.Server
 		/// Use the Epic telemetry sink
 		/// </summary>
 		Epic,
+		
+		/// <summary>
+		/// Use the ClickHouse telemetry sink
+		/// </summary>
+		ClickHouse,
 	}
 
 	/// <summary>
 	/// Configuration for the telemetry sink
 	/// </summary>
-	public class TelemetryConfig : IEpicTelemetrySinkConfig
+	public class BaseTelemetryConfig
 	{
+		/// <summary>
+		/// Unique ID for this sink config (any arbitrary string)
+		/// </summary>
+		public string? Id { get; set; }
+		
 		/// <summary>
 		/// Type of telemetry sink
 		/// </summary>
 		public TelemetrySinkType Type { get; set; } = TelemetrySinkType.None;
+	}
 
+	/// <summary>
+	/// Configuration for the telemetry sink
+	/// </summary>
+	public class EpicTelemetryConfig : BaseTelemetryConfig, IEpicTelemetrySinkConfig
+	{
 		/// <summary>
 		/// Base URL for the telemetry server
 		/// </summary>
@@ -203,6 +219,29 @@ namespace Horde.Server
 		/// Application name to send in the event messages
 		/// </summary>
 		public string AppId { get; set; } = "Horde";
+
+		/// <inheritdoc />
+		public override string ToString()
+		{
+			return $"{nameof(Url)}={Url} {nameof(AppId)}={AppId}";
+		}
+	}
+	
+	/// <summary>
+	/// Configuration for the telemetry sink
+	/// </summary>
+	public class ClickHouseTelemetryConfig : BaseTelemetryConfig
+	{
+		/// <summary>
+		/// Base URL for ClickHouse server
+		/// </summary>
+		public Uri? Url { get; set; }
+		
+		/// <inheritdoc />
+		public override string ToString()
+		{
+			return $"{nameof(Url)}={Url}";
+		}
 	}
 	
 	/// <summary>
@@ -698,9 +737,9 @@ namespace Horde.Server
 		public CommitSettings Commits { get; set; } = new CommitSettings();
 
 		/// <summary>
-		/// Settings for sending telemetry events
+		/// Settings for sending telemetry events to external services (for example Snowflake, ClickHouse etc)
 		/// </summary>
-		public TelemetryConfig Telemetry { get; set; } = new TelemetryConfig();
+		public List<BaseTelemetryConfig> Telemetry { get; set; } = new ();
 
 		/// <summary>
 		/// Tools bundled along with the server. Data for each tool can be produced using the 'bundle create' command, and should be stored in the /tools/{id} directory.
diff --git a/Engine/Source/Programs/Horde/Horde.Server/Startup.cs b/Engine/Source/Programs/Horde/Horde.Server/Startup.cs
--- a/Engine/Source/Programs/Horde/Horde.Server/Startup.cs
+++ b/Engine/Source/Programs/Horde/Horde.Server/Startup.cs
@@ -284,6 +284,47 @@ namespace Horde.Server
 
 		public IConfiguration Configuration { get; }
 
+
+		/// <summary>
+		/// Bind config to concrete subclasses of BaseTelemetryConfig
+		/// </summary>
+		/// <param name="settings">Settings object being updated</param>
+		/// <param name="telemetrySection">Telemetry config section</param>
+		/// <exception cref="ConfigurationException"></exception>
+		private static void BindTelemetrySettings(ServerSettings settings, IConfiguration telemetrySection)
+		{
+			List<BaseTelemetryConfig> telemetryConfigs = new();
+			foreach (IConfigurationSection child in telemetrySection.GetChildren())
+			{
+				string typeStr = child.GetValue<string>("Type");
+				if (!Enum.TryParse(typeStr, true, out TelemetrySinkType sinkType))
+				{
+					throw new ConfigurationException($"Unable to parse sink type '{typeStr}'");
+				}
+				
+				switch (sinkType)
+				{
+					case TelemetrySinkType.Epic:
+						EpicTelemetryConfig epic = new ();
+						child.Bind(epic);
+						telemetryConfigs.Add(epic);
+						break;
+						
+					case TelemetrySinkType.ClickHouse:
+						ClickHouseTelemetryConfig clickHouse = new ();
+						child.Bind(clickHouse);
+						telemetryConfigs.Add(clickHouse);
+						break;
+					
+					case TelemetrySinkType.None:
+					default:
+						break;
+				}
+			}
+
+			settings.Telemetry = telemetryConfigs;
+		}
+		
 		// This method gets called *multiple times* by the runtime. Use this method to add services to the container.
 		public void ConfigureServices(IServiceCollection services)
 		{
@@ -294,6 +335,7 @@ namespace Horde.Server
 			IConfigurationSection configSection = Configuration.GetSection("Horde");
 			ServerSettings settings = new ServerSettings();
 			configSection.Bind(settings);
+			BindTelemetrySettings(settings, configSection.GetSection("Telemetry"));
 			settings.Validate();
 
 			if (settings.GlobalThreadPoolMinSize != null)
@@ -613,19 +655,11 @@ namespace Horde.Server
 					throw new ArgumentException($"Invalid auth method {settings.AuthMethod}");
 			}
 
-			TelemetryConfig telemetryConfig = settings.Telemetry;
-			switch (telemetryConfig.Type)
-			{
-				case TelemetrySinkType.None:
-					services.AddSingleton<ITelemetrySink, NullTelemetrySink>();
-					break;
-				case TelemetrySinkType.Epic:
-					services.AddHttpClient(EpicTelemetrySink.HttpClientName, client => { });
-					services.AddSingleton<EpicTelemetrySink>();
-					services.AddSingleton<ITelemetrySink>(sp => sp.GetRequiredService<EpicTelemetrySink>());
-					services.AddHostedService(sp => sp.GetRequiredService<EpicTelemetrySink>());
-					break;
-			}
+			services.AddSingleton<TelemetryManager>();
+			services.AddSingleton<ITelemetrySink>(sp => sp.GetRequiredService<TelemetryManager>());
+			services.AddHostedService(sp => sp.GetRequiredService<TelemetryManager>());
+			services.AddHttpClient(EpicTelemetrySink.HttpClientName, client => { });
+			services.AddHttpClient(ClickHouseTelemetrySink.HttpClientName, client => { });
 
 			authBuilder.AddScheme<JwtBearerOptions, HordeServerJwtBearerHandler>(HordeServerJwtBearerHandler.AuthenticationScheme, options => { });
 			schemes.Add(HordeServerJwtBearerHandler.AuthenticationScheme);
diff --git a/Engine/Source/Programs/Horde/Horde.Server/Telemetry/ClickHouseTelemetrySink.cs b/Engine/Source/Programs/Horde/Horde.Server/Telemetry/ClickHouseTelemetrySink.cs
new file mode 100644
--- /dev/null
+++ b/Engine/Source/Programs/Horde/Horde.Server/Telemetry/ClickHouseTelemetrySink.cs
@@ -0,0 +1,189 @@
+// Copyright Epic Games, Inc. All Rights Reserved.
+
+using System;
+using System.Collections.Concurrent;
+using System.Collections.Generic;
+using System.Net.Http;
+using System.Text;
+using System.Threading;
+using System.Threading.Tasks;
+using Horde.Server.Jobs;
+using HordeCommon;
+using HordeCommon.Rpc.Messages.Telemetry;
+using Microsoft.Extensions.Logging;
+
+namespace Horde.Server.Telemetry
+{
+	/// <summary>
+	/// ClickHouse telemetry sink
+	/// </summary>
+	public sealed class ClickHouseTelemetrySink : ITelemetrySinkInternal
+	{
+		/// <summary>
+		/// Name of the HTTP client for writing telemetry data
+		/// </summary>
+		public const string HttpClientName = "ClickHouseTelemetrySink";
+
+		private readonly IHttpClientFactory _httpClientFactory;
+		private readonly Uri? _uri;
+
+		private readonly ITicker _ticker;
+		private readonly ILogger _logger;
+		private readonly ConcurrentQueue<(string eventName, object evt)> _queuedEvents = new();
+
+		/// <inheritdoc/>
+		public bool Enabled => _uri != null;
+
+		/// <summary>
+		/// Constructor
+		/// </summary>
+		public ClickHouseTelemetrySink(ClickHouseTelemetryConfig config, IHttpClientFactory httpClientFactory, IClock clock, ILogger<ClickHouseTelemetrySink> logger)
+		{
+			_httpClientFactory = httpClientFactory;
+			_ticker = clock.AddTicker<ClickHouseTelemetrySink>(TimeSpan.FromSeconds(10.0), FlushAsync, logger);
+			_logger = logger;
+			_uri = config.Url;
+		}
+
+		/// <inheritdoc/>
+		public ValueTask DisposeAsync()
+		{
+			_ticker.Dispose();
+			return ValueTask.CompletedTask;
+		}
+		
+		/// <inheritdoc/>
+		public void SendEvent(string eventName, object attributes)
+		{
+			_queuedEvents.Enqueue((eventName, attributes));
+		}
+
+		/// <inheritdoc />
+		public async ValueTask FlushAsync(CancellationToken cancellationToken)
+		{
+			List<AgentMetadataEvent> agentMetadataEvents = new();
+			List<AgentCpuMetricsEvent> agentCpuEvents = new();
+			List<AgentMemoryMetricsEvent> agentMemEvents = new();
+
+			int c = 0;
+			while (_queuedEvents.TryDequeue(out (string eventName, object evt) tuple))
+			{
+				switch (tuple.evt)
+				{
+					case AgentMetadataEvent agentMetadata: agentMetadataEvents.Add(agentMetadata); break;
+					case AgentCpuMetricsEvent agentCpu: agentCpuEvents.Add(agentCpu); break;
+					case AgentMemoryMetricsEvent agentMem: agentMemEvents.Add(agentMem); break;
+				}
+				c++;
+			}
+
+			StringBuilder sb = new(c * 200); // Rough size of about 200 bytes per event
+			WriteAgentMetadataSql(agentMetadataEvents, sb);
+			WriteAgentCpuSql(agentCpuEvents, sb);
+
+			await SendClickHouseQueryAsync(Encoding.UTF8.GetBytes(sb.ToString()), cancellationToken);
+		}
+
+		private async Task SendClickHouseQueryAsync(byte[] query, CancellationToken cancellationToken)
+		{
+			HttpClient httpClient = _httpClientFactory.CreateClient(HttpClientName);
+			
+			using HttpRequestMessage request = new();
+			request.RequestUri = _uri;
+			request.Method = HttpMethod.Post;
+			request.Content = new ByteArrayContent(query);
+
+			using HttpResponseMessage response = await httpClient.SendAsync(request, cancellationToken);
+			if (response.IsSuccessStatusCode)
+			{
+				_logger.LogDebug("Sending {Size} bytes of telemetry data to {Url}", query.Length, _uri);
+			}
+			else
+			{
+				string content = await response.Content.ReadAsStringAsync(cancellationToken);
+				_logger.LogError("Unable to send telemetry data to server (status {Code}): {Message}", (int)response.StatusCode, content);
+			}
+		}
+
+		/// <summary>
+		/// Generate a INSERT INTO statement for a list of AgentCpuMetricsEvent
+		/// </summary>
+		/// <param name="events"></param>
+		/// <param name="sb"></param>
+		private static void WriteAgentCpuSql(List<AgentCpuMetricsEvent> events, StringBuilder sb)
+		{
+			if (events.Count > 0)
+			{
+				sb.Append("INSERT INTO agentCpu (time, agentId, leaseId, jobId, jobBatchId, user, system, idle) VALUES \n");
+				foreach (AgentCpuMetricsEvent e in events)
+				{
+					ExecutionMetadata em = e.ExecutionMetadata;
+					sb.AppendFormat("({0}, {1}, '{2}', '{3}', {4}, {5:F3}, {6:F3}, {7:F3}),\n",
+						e.Timestamp.Seconds, e.AgentId,
+						WriteEscapedObjectId(em.LeaseId), WriteEscapedObjectId(em.JobId),
+						SubResourceId.Parse(em.JobBatchId).Value, e.User, e.System, e.Idle);
+				}
+				sb.Remove(sb.Length - 2, 2); // Remove trailing comma and newline
+				sb.Append(';');
+			}
+		}
+		
+		/// <summary>
+		/// Generate a INSERT INTO statement for a list of AgentMetadataEvent
+		/// </summary>
+		/// <param name="events"></param>
+		/// <param name="sb"></param>
+		private static void WriteAgentMetadataSql(List<AgentMetadataEvent> events, StringBuilder sb)
+		{
+			if (events.Count > 0)
+			{
+				sb.Append("INSERT INTO agentMetadata (id, ip, hostname, region, az, env, version, os, osVersion, arch, props) VALUES \n");
+				foreach (AgentMetadataEvent e in events)
+				{
+					sb.AppendFormat("({0}, '{1}', '{2}', '{3}', '{4}', '{5}', '{6}', '{7}', '{8}', '{9}', {10}),\n",
+						e.AgentId, e.Ip, e.Hostname, e.Region, e.AvailabilityZone, e.Environment, e.AgentVersion,
+						e.Os, e.OsVersion, e.Architecture, WriteDictionaryAsClickHouseMap(e.Properties));
+				}
+				sb.Remove(sb.Length - 2, 2); // Remove trailing comma and newline
+				sb.Append(';');
+			}
+		}
+
+		private static string WriteEscapedObjectId(string value)
+		{
+			if (value.Length != 24)
+			{
+				throw new ArgumentException($"Expected a MongoDB ObjectId string of 24 chars. Got: '{value}'");
+			}
+
+			StringBuilder sb = new (value.Length * 2);
+			for (int i = 0; i < value.Length; i += 2)
+			{
+				sb.Append("\\x");
+				sb.Append(value[i]);
+				sb.Append(value[i + 1]);
+			}
+
+			return sb.ToString();
+		}
+		
+		private static string WriteDictionaryAsClickHouseMap(IDictionary<string, string> dict)
+		{
+			if (dict.Count == 0)
+			{
+				return "{}";
+			}
+			
+			StringBuilder sb = new (dict.Count * 2 * 20);
+			sb.Append('{');
+			foreach ((string key, string value) in dict)
+			{
+				sb.AppendFormat("'{0}': '{1}', ", key, value);
+			}
+			sb.Remove(sb.Length - 2, 2); // Remove trailing comma and whitespace
+			sb.Append('}');
+
+			return sb.ToString();
+		}
+	}
+}
diff --git a/Engine/Source/Programs/Horde/Horde.Server/Telemetry/EpicTelemetrySink.cs b/Engine/Source/Programs/Horde/Horde.Server/Telemetry/EpicTelemetrySink.cs
--- a/Engine/Source/Programs/Horde/Horde.Server/Telemetry/EpicTelemetrySink.cs
+++ b/Engine/Source/Programs/Horde/Horde.Server/Telemetry/EpicTelemetrySink.cs
@@ -12,9 +12,7 @@ using System.Threading.Tasks;
 using EpicGames.Core;
 using HordeCommon;
 using Microsoft.AspNetCore.WebUtilities;
-using Microsoft.Extensions.Hosting;
 using Microsoft.Extensions.Logging;
-using Microsoft.Extensions.Options;
 
 namespace Horde.Server.Telemetry
 {
@@ -37,7 +35,7 @@ namespace Horde.Server.Telemetry
 	/// <summary>
 	/// Epic internal telemetry sink using the data router
 	/// </summary>
-	public sealed class EpicTelemetrySink : IHostedService, ITelemetrySink, IDisposable
+	public sealed class EpicTelemetrySink : ITelemetrySinkInternal
 	{
 		// Converter for datetime formats that Tableau can ingest
 		class TableauDateTimeConverter : JsonConverter<DateTime>
@@ -84,7 +82,7 @@ namespace Horde.Server.Telemetry
 		/// <summary>
 		/// Constructor
 		/// </summary>
-		public EpicTelemetrySink(IHttpClientFactory httpClientFactory, IClock clock, IOptions<ServerSettings> serverSettings, ILogger<EpicTelemetrySink> logger)
+		public EpicTelemetrySink(IEpicTelemetrySinkConfig config, IHttpClientFactory httpClientFactory, IClock clock, ILogger<EpicTelemetrySink> logger)
 		{
 			_httpClientFactory = httpClientFactory;
 
@@ -102,7 +100,6 @@ namespace Horde.Server.Telemetry
 			_ticker = clock.AddTicker<EpicTelemetrySink>(TimeSpan.FromSeconds(60.0), FlushAsync, logger);
 			_logger = logger;
 
-			IEpicTelemetrySinkConfig config = serverSettings.Value.Telemetry;
 			if (config.Url != null)
 			{
 				Dictionary<string, string?> queryParams = new Dictionary<string, string?>()
@@ -118,25 +115,13 @@ namespace Horde.Server.Telemetry
 		}
 
 		/// <inheritdoc/>
-		public void Dispose()
+		public async ValueTask DisposeAsync()
 		{
-			_packetWriter.Dispose();
-			_eventWriter.Dispose();
+			await _packetWriter.DisposeAsync();
+			await _eventWriter.DisposeAsync();
 			_ticker.Dispose();
 		}
-
-		/// <inheritdoc/>
-		public async Task StartAsync(CancellationToken cancellationToken)
-		{
-			await _ticker.StartAsync();
-		}
-
-		/// <inheritdoc/>
-		public async Task StopAsync(CancellationToken cancellationToken)
-		{
-			await _ticker.StopAsync();
-		}
-
+		
 		/// <inheritdoc/>
 		public async ValueTask FlushAsync(CancellationToken cancellationToken)
 		{
diff --git a/Engine/Source/Programs/Horde/Horde.Server/Telemetry/ITelemetrySink.cs b/Engine/Source/Programs/Horde/Horde.Server/Telemetry/ITelemetrySink.cs
--- a/Engine/Source/Programs/Horde/Horde.Server/Telemetry/ITelemetrySink.cs
+++ b/Engine/Source/Programs/Horde/Horde.Server/Telemetry/ITelemetrySink.cs
@@ -1,5 +1,9 @@
 // Copyright Epic Games, Inc. All Rights Reserved.
 
+using System;
+using System.Threading;
+using System.Threading.Tasks;
+
 namespace Horde.Server.Telemetry
 {
 	/// <summary>
@@ -19,4 +23,17 @@ namespace Horde.Server.Telemetry
 		/// <param name="attributes">Arbitrary object to include in the payload</param>
 		void SendEvent(string eventName, object attributes);
 	}
+
+	/// <summary>
+	/// Internal interface for sinks
+	/// </summary>
+	public interface ITelemetrySinkInternal : ITelemetrySink, IAsyncDisposable
+	{
+		/// <summary>
+		/// Flush any queued metrics to underlying service
+		/// </summary>
+		/// <param name="cancellationToken"></param>
+		/// <returns>Completion task</returns>
+		public ValueTask FlushAsync(CancellationToken cancellationToken);
+	}
 }
diff --git a/Engine/Source/Programs/Horde/Horde.Server/Telemetry/TelemetryManager.cs b/Engine/Source/Programs/Horde/Horde.Server/Telemetry/TelemetryManager.cs
new file mode 100644
--- /dev/null
+++ b/Engine/Source/Programs/Horde/Horde.Server/Telemetry/TelemetryManager.cs
@@ -0,0 +1,96 @@
+// Copyright Epic Games, Inc. All Rights Reserved.
+
+using System;
+using System.Collections.Generic;
+using System.Net.Http;
+using System.Threading;
+using System.Threading.Tasks;
+using HordeCommon;
+using Microsoft.Extensions.Hosting;
+using Microsoft.Extensions.Logging;
+using Microsoft.Extensions.Options;
+
+namespace Horde.Server.Telemetry
+{
+	/// <summary>
+	/// Telemetry sink dispatching incoming events to all registered sinks
+	/// </summary>
+	public sealed class TelemetryManager : ITelemetrySinkInternal, IHostedService
+	{
+		/// <inheritdoc/>
+		public bool Enabled => _telemetrySinks.Count > 0;
+
+		private readonly List<ITelemetrySinkInternal> _telemetrySinks = new();
+		private readonly ITicker _ticker;
+		private readonly ILogger<TelemetryManager> _logger;
+
+		/// <summary>
+		/// Constructor
+		/// </summary>
+		/// <param name="httpClientFactory"></param>
+		/// <param name="clock"></param>
+		/// <param name="serverSettings"></param>
+		/// <param name="loggerFactory"></param>
+		public TelemetryManager(IHttpClientFactory httpClientFactory, IClock clock, IOptions<ServerSettings> serverSettings, ILoggerFactory loggerFactory)
+		{
+			_logger = loggerFactory.CreateLogger<TelemetryManager>();
+			_ticker = clock.AddTicker<EpicTelemetrySink>(TimeSpan.FromSeconds(30.0), FlushAsync, _logger);
+			
+			foreach (BaseTelemetryConfig config in serverSettings.Value.Telemetry)
+			{
+				switch (config)
+				{
+					case EpicTelemetryConfig epicConfig:
+						_telemetrySinks.Add(new EpicTelemetrySink(epicConfig, httpClientFactory, clock, loggerFactory.CreateLogger<EpicTelemetrySink>()));
+						break;
+					case ClickHouseTelemetryConfig chConfig:
+						_telemetrySinks.Add(new ClickHouseTelemetrySink(chConfig, httpClientFactory, clock, loggerFactory.CreateLogger<ClickHouseTelemetrySink>()));
+						break;
+				}
+			}
+		}
+
+		/// <inheritdoc/>
+		public void SendEvent(string eventName, object attributes)
+		{
+			foreach (ITelemetrySinkInternal sink in _telemetrySinks)
+			{
+				if (sink.Enabled)
+				{
+					sink.SendEvent(eventName, attributes);
+				}
+			}
+		}
+
+		/// <inheritdoc />
+		public async Task StartAsync(CancellationToken cancellationToken)
+		{
+			await _ticker.StartAsync();
+		}
+
+		/// <inheritdoc />
+		public async Task StopAsync(CancellationToken cancellationToken)
+		{
+			await _ticker.StopAsync();
+		}
+
+		/// <inheritdoc />
+		public async ValueTask DisposeAsync()
+		{
+			_ticker.Dispose();
+			foreach (ITelemetrySinkInternal sink in _telemetrySinks)
+			{
+				await sink.DisposeAsync();
+			}
+		}
+
+		/// <inheritdoc />
+		public async ValueTask FlushAsync(CancellationToken cancellationToken)
+		{
+			foreach (ITelemetrySinkInternal sink in _telemetrySinks)
+			{
+				await sink.FlushAsync(cancellationToken);
+			}
+		}
+	}
+}
