Horde: Activate and register new telemetry service in DI for agent

diff --git a/Engine/Source/Programs/Horde/Horde.Agent/Program.cs b/Engine/Source/Programs/Horde/Horde.Agent/Program.cs
--- a/Engine/Source/Programs/Horde/Horde.Agent/Program.cs
+++ b/Engine/Source/Programs/Horde/Horde.Agent/Program.cs
@@ -183,6 +183,8 @@ namespace Horde.Agent
 				});
 
 			services.AddSingleton<GrpcService>();
+			services.AddSingleton<TelemetryService>();
+			services.AddHostedService(sp => sp.GetRequiredService<TelemetryService>());
 
 			services.AddSingleton<IJobExecutorFactory, PerforceExecutorFactory>();
 			services.AddSingleton<IJobExecutorFactory, WorkspaceExecutorFactory>();
diff --git a/Engine/Source/Programs/Horde/Horde.Agent/Services/TelemetryService.cs b/Engine/Source/Programs/Horde/Horde.Agent/Services/TelemetryService.cs
--- a/Engine/Source/Programs/Horde/Horde.Agent/Services/TelemetryService.cs
+++ b/Engine/Source/Programs/Horde/Horde.Agent/Services/TelemetryService.cs
@@ -204,7 +204,7 @@ class TelemetryService : BackgroundService
 	/// <summary>
 	/// Constructor
 	/// </summary>
-	TelemetryService(GrpcService grpcService, IOptions<AgentSettings> settings, ILogger<TelemetryService> logger)
+	public TelemetryService(GrpcService grpcService, IOptions<AgentSettings> settings, ILogger<TelemetryService> logger)
 	{
 		_grpcService = grpcService;
 		_agentSettings = settings.Value;
