Horde: Attach agent metadata to events

Calculates an agent ID based on this metadata to allow indexing of events in ClickHouse.

diff --git a/Engine/Source/Programs/Horde/Horde.Agent/Services/TelemetryService.cs b/Engine/Source/Programs/Horde/Horde.Agent/Services/TelemetryService.cs
--- a/Engine/Source/Programs/Horde/Horde.Agent/Services/TelemetryService.cs
+++ b/Engine/Source/Programs/Horde/Horde.Agent/Services/TelemetryService.cs
@@ -243,8 +243,8 @@ class TelemetryService : BackgroundService
 			AgentMemoryMetricsEvent memMetricsEvent = _systemMetrics.GetMemory().ToEvent();
 
 			SendTelemetryEventsRequest request = new();
-			request.Events.Add(new WrappedTelemetryEvent { Cpu = cpuMetricsEvent });
-			request.Events.Add(new WrappedTelemetryEvent { Mem = memMetricsEvent });
+			request.Events.Add(new WrappedTelemetryEvent { AgentMetadata = _agentMetadataEvent, Cpu = cpuMetricsEvent });
+			request.Events.Add(new WrappedTelemetryEvent { AgentMetadata = _agentMetadataEvent, Mem = memMetricsEvent });
 
 			if (DateTime.UtcNow > _lastTimeAgentMetadataSent + _agentMetadataReportInterval)
 			{
@@ -261,7 +261,7 @@ class TelemetryService : BackgroundService
 
 	AgentMetadataEvent GetAgentMetadataEvent()
 	{
-		return new AgentMetadataEvent
+		AgentMetadataEvent e = new()
 		{
 			Ip = null,
 			Hostname = _agentSettings.GetAgentName(),
@@ -273,6 +273,8 @@ class TelemetryService : BackgroundService
 			OsVersion = Environment.OSVersion.Version.ToString(),
 			Architecture = RuntimeInformation.OSArchitecture.ToString()
 		};
+		e.AgentId = e.CalculateAgentId();
+		return e;
 	}
 
 	private static string GetOs()
diff --git a/Engine/Source/Programs/Horde/Horde.Agent/Utilities/HordeRpc.cs b/Engine/Source/Programs/Horde/Horde.Agent/Utilities/HordeRpc.cs
--- a/Engine/Source/Programs/Horde/Horde.Agent/Utilities/HordeRpc.cs
+++ b/Engine/Source/Programs/Horde/Horde.Agent/Utilities/HordeRpc.cs
@@ -1,8 +1,11 @@
 // Copyright Epic Games, Inc. All Rights Reserved.
 
+using System;
 using System.Collections.Generic;
 using System.Diagnostics.CodeAnalysis;
+using System.IO;
 using System.Linq;
+using System.Security.Cryptography;
 using Google.Protobuf;
 using Google.Protobuf.Collections;
 
@@ -185,5 +188,43 @@ namespace HordeCommon.Rpc
 		}
 	}
 }
+
+namespace HordeCommon.Rpc.Messages.Telemetry
+{
+	partial class AgentMetadataEvent
+	{
+		/// <summary>
+		/// Calculate an agent ID
+		/// </summary>
+		/// <returns>A unique hash for all fields</returns>
+		public ulong CalculateAgentId()
+		{
+			using SHA256 sha256 = SHA256.Create();
+			using MemoryStream ms = new(200);
+			using BinaryWriter bw = new(ms);
+			
+			bw.Write(Ip);
+			bw.Write(Hostname);
+			bw.Write(Region);
+			bw.Write(AvailabilityZone);
+			bw.Write(Environment);
+			bw.Write(AgentVersion);
+			bw.Write(Os);
+			bw.Write(OsVersion);
+			bw.Write(Architecture);
+			
+			foreach (KeyValuePair<string,string> pair in Properties)
+			{
+				bw.Write(pair.Key);
+				bw.Write(pair.Value);
+			}
+			
+			ms.Position = 0;
+			byte[] hash = sha256.ComputeHash(ms);
+			return BitConverter.ToUInt64(hash, 0);
+		}
+	}
+}
+
 #pragma warning restore CA1716
 #pragma warning restore CS1591
\ No newline at end of file
