Horde: Fix proper custom binding of telemetry config

It was binding to an unused instance of server settings.

diff --git a/Engine/Source/Programs/Horde/Horde.Server/Startup.cs b/Engine/Source/Programs/Horde/Horde.Server/Startup.cs
--- a/Engine/Source/Programs/Horde/Horde.Server/Startup.cs
+++ b/Engine/Source/Programs/Horde/Horde.Server/Startup.cs
@@ -329,13 +329,17 @@ namespace Horde.Server
 		public void ConfigureServices(IServiceCollection services)
 		{
 			// IOptionsMonitor pattern for live updating of configuration settings
-			services.Configure<ServerSettings>(Configuration.GetSection("Horde"));
+			services.Configure<ServerSettings>(serverSettings =>
+			{
+				IConfigurationSection hordeSection = Configuration.GetSection("Horde");
+				hordeSection.Bind(serverSettings);
+				BindTelemetrySettings(serverSettings, hordeSection.GetSection("Telemetry"));
+			});
 
 			// Settings used for configuring services
 			IConfigurationSection configSection = Configuration.GetSection("Horde");
 			ServerSettings settings = new ServerSettings();
 			configSection.Bind(settings);
-			BindTelemetrySettings(settings, configSection.GetSection("Telemetry"));
 			settings.Validate();
 
 			if (settings.GlobalThreadPoolMinSize != null)
