Horde: Only send ClickHouse commands when there's data to send

diff --git a/Engine/Source/Programs/Horde/Horde.Server/Telemetry/ClickHouseTelemetrySink.cs b/Engine/Source/Programs/Horde/Horde.Server/Telemetry/ClickHouseTelemetrySink.cs
--- a/Engine/Source/Programs/Horde/Horde.Server/Telemetry/ClickHouseTelemetrySink.cs
+++ b/Engine/Source/Programs/Horde/Horde.Server/Telemetry/ClickHouseTelemetrySink.cs
@@ -81,7 +81,10 @@ namespace Horde.Server.Telemetry
 			WriteAgentMetadataSql(agentMetadataEvents, sb);
 			WriteAgentCpuSql(agentCpuEvents, sb);
 
-			await SendClickHouseQueryAsync(Encoding.UTF8.GetBytes(sb.ToString()), cancellationToken);
+			if (sb.Length > 0)
+			{
+				await SendClickHouseQueryAsync(Encoding.UTF8.GetBytes(sb.ToString()), cancellationToken);	
+			}
 		}
 
 		private async Task SendClickHouseQueryAsync(byte[] query, CancellationToken cancellationToken)
