Horde: Improve error logging for ClickHouse telemetry

diff --git a/Engine/Source/Programs/Horde/Horde.Server/Telemetry/ClickHouseTelemetrySink.cs b/Engine/Source/Programs/Horde/Horde.Server/Telemetry/ClickHouseTelemetrySink.cs
--- a/Engine/Source/Programs/Horde/Horde.Server/Telemetry/ClickHouseTelemetrySink.cs
+++ b/Engine/Source/Programs/Horde/Horde.Server/Telemetry/ClickHouseTelemetrySink.cs
@@ -103,8 +103,12 @@ namespace Horde.Server.Telemetry
 			}
 			else
 			{
-				string content = await response.Content.ReadAsStringAsync(cancellationToken);
-				_logger.LogError("Unable to send telemetry data to server (status {Code}): {Message}", (int)response.StatusCode, content);
+				string responseContent = await response.Content.ReadAsStringAsync(cancellationToken);
+				string queryStr = Encoding.UTF8.GetString(query);
+				queryStr = queryStr.Substring(0, Math.Min(queryStr.Length, 300));
+				
+				_logger.LogError("Unable to send telemetry data to server. StatusCode={StatusCode}) Response={Response} Query={Query}",
+					(int)response.StatusCode, responseContent, queryStr);
 			}
 		}
 
