Horde: Enable tools menu in prod
#rnx
#skipci

diff --git a/Engine/Source/Programs/Horde/HordeDashboard/src/components/TopNav.tsx b/Engine/Source/Programs/Horde/HordeDashboard/src/components/TopNav.tsx
--- a/Engine/Source/Programs/Horde/HordeDashboard/src/components/TopNav.tsx
+++ b/Engine/Source/Programs/Horde/HordeDashboard/src/components/TopNav.tsx
@@ -438,24 +438,24 @@ export const TopNav: React.FC<{ suppressServer?: boolean }> = observer(({ suppre
          }
       });
 
-      if (getSiteConfig().environment !== "production") {
+      const softwareItems: IContextualMenuItem[] = [];
+      softwareItems.push({
+         key: "software_tools",
+         text: "Tools",
+         link: `/tools`
+      });
 
-         const softwareItems: IContextualMenuItem[] = [];
-         softwareItems.push({
-            key: "software_tools",
-            text: "Tools",
-            link: `/tools`
-         });
+      subItems.push({
+         itemType: ContextualMenuItemType.Section,
+         key: `server_software`,
+         sectionProps: {
+            title: "Software",
+            items: softwareItems,
+            bottomDivider: true
+         }
+      });
 
-         subItems.push({
-            itemType: ContextualMenuItemType.Section,
-            key: `server_software`,
-            sectionProps: {
-               title: "Software",
-               items: softwareItems,
-               bottomDivider: true
-            }
-         });
+      if (getSiteConfig().environment !== "production") {
 
          const docsItems: IContextualMenuItem[] = [];
          docsItems.push({
