Horde: Allow downloading tools as single files if there are no other files or directories in them.

diff --git a/Engine/Source/Programs/Horde/Horde.Server/Tools/ToolsController.cs b/Engine/Source/Programs/Horde/Horde.Server/Tools/ToolsController.cs
--- a/Engine/Source/Programs/Horde/Horde.Server/Tools/ToolsController.cs
+++ b/Engine/Source/Programs/Horde/Horde.Server/Tools/ToolsController.cs
@@ -2,6 +2,7 @@
 
 using EpicGames.Core;
 using EpicGames.Horde.Storage;
+using EpicGames.Horde.Storage.Nodes;
 using Horde.Server.Acls;
 using Horde.Server.Server;
 using Horde.Server.Storage;
@@ -10,6 +11,8 @@ using HordeCommon;
 using Microsoft.AspNetCore.Authorization;
 using Microsoft.AspNetCore.Http;
 using Microsoft.AspNetCore.Mvc;
+using Microsoft.AspNetCore.StaticFiles;
+using Microsoft.Extensions.Caching.Memory;
 using Microsoft.Extensions.Logging;
 using Microsoft.Extensions.Options;
 using System;
@@ -175,6 +178,11 @@ namespace Horde.Server.Tools
 		/// Download the deployment data
 		/// </summary>
 		Download,
+
+		/// <summary>
+		/// Download the deployment data as a zip file
+		/// </summary>
+		Zip,
 	}
 
 	/// <summary>
@@ -267,16 +275,20 @@ namespace Horde.Server.Tools
 		readonly StorageService _storageService;
 		readonly IOptionsSnapshot<GlobalConfig> _globalConfig;
 		readonly IClock _clock;
+		readonly IMemoryCache _cache;
+		readonly ILogger _logger;
 
 		/// <summary>
 		/// Constructor
 		/// </summary>
-		public PublicToolsController(IToolCollection toolCollection, StorageService storageService, IClock clock, IOptionsSnapshot<GlobalConfig> globalConfig)
+		public PublicToolsController(IToolCollection toolCollection, StorageService storageService, IClock clock, IOptionsSnapshot<GlobalConfig> globalConfig, IMemoryCache cache, ILogger<ToolsController> logger)
 		{
 			_toolCollection = toolCollection;
 			_storageService = storageService;
 			_clock = clock;
 			_globalConfig = globalConfig;
+			_cache = cache;
+			_logger = logger;
 		}
 
 		/// <summary>
@@ -419,11 +431,27 @@ namespace Horde.Server.Tools
 				GetToolDeploymentResponse response = await GetDeploymentInfoResponseAsync(tool, deployment, cancellationToken);
 				return Ok(response);
 			}
-			else
+
+			IStorageClient client = await _toolCollection.GetStorageClientAsync(tool, cancellationToken);
+
+			TreeReader reader = new TreeReader(client, _cache, _logger);
+			DirectoryNode node = await reader.ReadNodeAsync<DirectoryNode>(deployment.RefName, DateTime.UtcNow - TimeSpan.FromDays(2.0), cancellationToken);
+
+			if (node.Directories.Count == 0 && node.Files.Count == 1 && action != GetToolAction.Zip)
 			{
-				Stream stream = await _toolCollection.GetDeploymentZipAsync(tool, deployment, cancellationToken);
-				return new FileStreamResult(stream, "application/zip");
+				FileEntry entry = node.Files.First();
+
+				string? contentType;
+				if (!new FileExtensionContentTypeProvider().TryGetContentType(entry.Name.ToString(), out contentType))
+				{
+					contentType = "application/octet-stream";
+				}
+
+				return new FileStreamResult(entry.AsStream(reader), contentType) { FileDownloadName = entry.Name.ToString() };
 			}
+
+			Stream stream = node.AsZipStream(reader);
+			return new FileStreamResult(stream, "application/zip") { FileDownloadName = $"{tool.Id}-{deployment.Version}.zip" };
 		}
 
 		private async Task<GetToolDeploymentResponse> GetDeploymentInfoResponseAsync(ITool tool, IToolDeployment deployment, CancellationToken cancellationToken)
