Horde: Fix static analysis warnings due to fields not being disposed, statements not on their own line.

diff --git a/Engine/Source/Programs/Horde/Horde.Agent/Services/TelemetryService.cs b/Engine/Source/Programs/Horde/Horde.Agent/Services/TelemetryService.cs
--- a/Engine/Source/Programs/Horde/Horde.Agent/Services/TelemetryService.cs
+++ b/Engine/Source/Programs/Horde/Horde.Agent/Services/TelemetryService.cs
@@ -96,7 +96,7 @@ public class MemoryMetrics
 /// <summary>
 /// OS agnostic interface for retrieving system metrics (CPU, memory etc)
 /// </summary>
-public interface ISystemMetrics
+public interface ISystemMetrics : IDisposable
 {
 	/// <summary>
 	/// Get CPU usage metrics
@@ -117,7 +117,7 @@ public interface ISystemMetrics
 /// <summary>
 /// Windows specific implementation for gathering system metrics
 /// </summary>
-public class WindowsSystemMetrics : ISystemMetrics
+public sealed class WindowsSystemMetrics : ISystemMetrics
 {
 	private const string ProcessorInfo = "Processor Information"; // Prefer this over "Processor" as it's more modern
 	private const string Memory = "Memory";
@@ -138,6 +138,15 @@ public class WindowsSystemMetrics : ISystemMetrics
 		GetCpu(); // Trigger this to ensure performance counter has a fetched value. Avoids an initial zero result when called later.
 	}
 
+	/// <inheritdoc/>
+	public void Dispose()
+	{
+		_memAvailableBytes.Dispose();
+		_procIdleTime.Dispose();
+		_procUserTime.Dispose();
+		_procPrivilegedTime.Dispose();
+	}
+
 	/// <inheritdoc />
 	public CpuMetrics GetCpu()
 	{
@@ -214,13 +223,20 @@ class TelemetryService : BackgroundService
 		// Calculate this once at startup as it should not change during lifetime of process
 		_agentMetadataEvent = GetAgentMetadataEvent();
 	}
+
+	/// <inheritdoc/>
+	public override void Dispose()
+	{
+		base.Dispose();
+		_systemMetrics?.Dispose();
+	}
  
 	/// <inheritdoc />
 	public override Task StartAsync(CancellationToken cancellationToken)
 	{
 		if (RuntimeInformation.IsOSPlatform(OSPlatform.Windows))
 		{
-			_systemMetrics = new WindowsSystemMetrics();
+			_systemMetrics ??= new WindowsSystemMetrics();
 		}
 		else
 		{
@@ -279,9 +295,18 @@ class TelemetryService : BackgroundService
 
 	private static string GetOs()
 	{
-		if (RuntimeInformation.IsOSPlatform(OSPlatform.Windows)) { return "Windows"; }
-		if (RuntimeInformation.IsOSPlatform(OSPlatform.Linux)) { return "Linux"; }
-		if (RuntimeInformation.IsOSPlatform(OSPlatform.OSX)) { return "macOS"; }
+		if (RuntimeInformation.IsOSPlatform(OSPlatform.Windows))
+		{ 
+			return "Windows"; 
+		}
+		if (RuntimeInformation.IsOSPlatform(OSPlatform.Linux))
+		{ 
+			return "Linux"; 
+		}
+		if (RuntimeInformation.IsOSPlatform(OSPlatform.OSX))
+		{ 
+			return "macOS"; 
+		}
 		return "Unknown";
 	}
 }
\ No newline at end of file
