Horde: Set timestamp for agent telemetry events

Also improve empty field handling.

diff --git a/Engine/Source/Programs/Horde/Horde.Agent/Services/TelemetryService.cs b/Engine/Source/Programs/Horde/Horde.Agent/Services/TelemetryService.cs
--- a/Engine/Source/Programs/Horde/Horde.Agent/Services/TelemetryService.cs
+++ b/Engine/Source/Programs/Horde/Horde.Agent/Services/TelemetryService.cs
@@ -6,6 +6,7 @@ using System.Threading.Tasks;
 using Microsoft.Extensions.Hosting;
 using System.Diagnostics;
 using System.Runtime.InteropServices;
+using Google.Protobuf.WellKnownTypes;
 using Grpc.Core;
 using Grpc.Net.Client;
 using HordeCommon.Rpc;
@@ -255,12 +256,26 @@ class TelemetryService : BackgroundService
 		while (!stoppingToken.IsCancellationRequested && _systemMetrics != null)
 		{
 			_logger.LogDebug("Sending telemetry events to server...");
-			AgentCpuMetricsEvent cpuMetricsEvent = _systemMetrics.GetCpu().ToEvent();
-			AgentMemoryMetricsEvent memMetricsEvent = _systemMetrics.GetMemory().ToEvent();
-
+			
 			SendTelemetryEventsRequest request = new();
-			request.Events.Add(new WrappedTelemetryEvent { AgentMetadata = _agentMetadataEvent, Cpu = cpuMetricsEvent });
-			request.Events.Add(new WrappedTelemetryEvent { AgentMetadata = _agentMetadataEvent, Mem = memMetricsEvent });
+			Timestamp utcNow = Timestamp.FromDateTime(DateTime.UtcNow);
+			ExecutionMetadata em = new ();
+			
+			{
+				AgentCpuMetricsEvent cpuMetricsEvent = _systemMetrics.GetCpu().ToEvent();
+				cpuMetricsEvent.AgentId = _agentMetadataEvent.AgentId;
+				cpuMetricsEvent.Timestamp = utcNow;
+				cpuMetricsEvent.ExecutionMetadata = em;
+				request.Events.Add(new WrappedTelemetryEvent { Cpu = cpuMetricsEvent });
+			}
+
+			{
+				AgentMemoryMetricsEvent memMetricsEvent = _systemMetrics.GetMemory().ToEvent();
+				memMetricsEvent.AgentId = _agentMetadataEvent.AgentId;
+				memMetricsEvent.Timestamp = utcNow;
+				memMetricsEvent.ExecutionMetadata = em;
+				request.Events.Add(new WrappedTelemetryEvent { Mem = memMetricsEvent });
+			}
 
 			if (DateTime.UtcNow > _lastTimeAgentMetadataSent + _agentMetadataReportInterval)
 			{
diff --git a/Engine/Source/Programs/Horde/Horde.Agent/Utilities/HordeRpc.cs b/Engine/Source/Programs/Horde/Horde.Agent/Utilities/HordeRpc.cs
--- a/Engine/Source/Programs/Horde/Horde.Agent/Utilities/HordeRpc.cs
+++ b/Engine/Source/Programs/Horde/Horde.Agent/Utilities/HordeRpc.cs
@@ -203,20 +203,20 @@ namespace HordeCommon.Rpc.Messages.Telemetry
 			using MemoryStream ms = new(200);
 			using BinaryWriter bw = new(ms);
 			
-			bw.Write(Ip);
-			bw.Write(Hostname);
-			bw.Write(Region);
-			bw.Write(AvailabilityZone);
-			bw.Write(Environment);
-			bw.Write(AgentVersion);
-			bw.Write(Os);
-			bw.Write(OsVersion);
-			bw.Write(Architecture);
+			bw.Write(Ip ?? "<empty ip>");
+			bw.Write(Hostname ?? "<empty hostname>");
+			bw.Write(Region ?? "<empty region>");
+			bw.Write(AvailabilityZone ?? "<empty az>");
+			bw.Write(Environment ?? "<empty env>");
+			bw.Write(AgentVersion ?? "<empty version>");
+			bw.Write(Os ?? "<empty os>");
+			bw.Write(OsVersion ?? "<empty os version>");
+			bw.Write(Architecture ?? "<empty os architecture>");
 			
 			foreach (KeyValuePair<string,string> pair in Properties)
 			{
-				bw.Write(pair.Key);
-				bw.Write(pair.Value);
+				bw.Write(pair.Key ?? "<empty key>");
+				bw.Write(pair.Value ?? "<empty value>");
 			}
 			
 			ms.Position = 0;
diff --git a/Engine/Source/Programs/Horde/Horde.Server/Telemetry/ClickHouseTelemetrySink.cs b/Engine/Source/Programs/Horde/Horde.Server/Telemetry/ClickHouseTelemetrySink.cs
--- a/Engine/Source/Programs/Horde/Horde.Server/Telemetry/ClickHouseTelemetrySink.cs
+++ b/Engine/Source/Programs/Horde/Horde.Server/Telemetry/ClickHouseTelemetrySink.cs
@@ -23,6 +23,8 @@ namespace Horde.Server.Telemetry
 		/// Name of the HTTP client for writing telemetry data
 		/// </summary>
 		public const string HttpClientName = "ClickHouseTelemetrySink";
+		
+		private const string EmptyObjectId = "000000000000000000000000";
 
 		private readonly IHttpClientFactory _httpClientFactory;
 		private readonly Uri? _uri;
@@ -124,11 +126,13 @@ namespace Horde.Server.Telemetry
 				sb.Append("INSERT INTO agentCpu (time, agentId, leaseId, jobId, jobBatchId, user, system, idle) VALUES \n");
 				foreach (AgentCpuMetricsEvent e in events)
 				{
-					ExecutionMetadata em = e.ExecutionMetadata;
+					ExecutionMetadata? em = e.ExecutionMetadata;
 					sb.AppendFormat("({0}, {1}, '{2}', '{3}', {4}, {5:F3}, {6:F3}, {7:F3}),\n",
-						e.Timestamp.Seconds, e.AgentId,
-						WriteEscapedObjectId(em.LeaseId), WriteEscapedObjectId(em.JobId),
-						SubResourceId.Parse(em.JobBatchId).Value, e.User, e.System, e.Idle);
+						e.Timestamp?.Seconds, e.AgentId,
+						WriteEscapedObjectId(em?.LeaseId ?? EmptyObjectId),
+						WriteEscapedObjectId(em?.JobId ?? EmptyObjectId),
+						SubResourceId.Parse(em?.JobBatchId ?? "0000").Value,
+						e.User, e.System, e.Idle);
 				}
 				sb.Remove(sb.Length - 2, 2); // Remove trailing comma and newline
 				sb.Append(';');
