Horde: Return a 404 response when artifact references have been deleted.

diff --git a/Engine/Source/Programs/Horde/Horde.Server/Artifacts/ArtifactsController.cs b/Engine/Source/Programs/Horde/Horde.Server/Artifacts/ArtifactsController.cs
--- a/Engine/Source/Programs/Horde/Horde.Server/Artifacts/ArtifactsController.cs
+++ b/Engine/Source/Programs/Horde/Horde.Server/Artifacts/ArtifactsController.cs
@@ -239,7 +239,16 @@ namespace Horde.Server.Artifacts
 			}
 
 			TreeReader reader = await _storageService.GetReaderAsync(artifact.NamespaceId, cancellationToken);
-			TreeNodeRef<DirectoryNode> directoryRef = await reader.ReadNodeRefAsync<DirectoryNode>(artifact.RefName, DateTime.UtcNow.AddHours(1.0), cancellationToken);
+
+			TreeNodeRef<DirectoryNode> directoryRef;
+			try
+			{
+				directoryRef = await reader.ReadNodeRefAsync<DirectoryNode>(artifact.RefName, DateTime.UtcNow.AddHours(1.0), cancellationToken);
+			}
+			catch (RefNameNotFoundException)
+			{
+				return NotFound(artifact.NamespaceId, artifact.RefName);
+			}
 
 			if (path != null)
 			{
diff --git a/Engine/Source/Programs/Horde/Horde.Server/Utilities/HordeControllerBase.cs b/Engine/Source/Programs/Horde/Horde.Server/Utilities/HordeControllerBase.cs
--- a/Engine/Source/Programs/Horde/Horde.Server/Utilities/HordeControllerBase.cs
+++ b/Engine/Source/Programs/Horde/Horde.Server/Utilities/HordeControllerBase.cs
@@ -243,6 +243,15 @@ namespace Horde.Server.Utilities
 			return NotFound("Namespace {NamespaceId} not found", namespaceId);
 		}
 
+		/// <summary>
+		/// Returns a 404 response for the given object
+		/// </summary>
+		[NonAction]
+		protected ActionResult NotFound(NamespaceId namespaceId, RefName refName)
+		{
+			return NotFound("Ref {RefName} in namespace {NamespaceId} not found", refName, namespaceId);
+		}
+
 		/// <summary>
 		/// Returns a 404 response for the given object
 		/// </summary>
