Horde: Artifact Browser - additional tweaks
#skipci

diff --git a/Engine/Source/Programs/Horde/HordeDashboard/src/components/jobDetailsV2/JobDetailArtifactsV3.tsx b/Engine/Source/Programs/Horde/HordeDashboard/src/components/jobDetailsV2/JobDetailArtifactsV3.tsx
--- a/Engine/Source/Programs/Horde/HordeDashboard/src/components/jobDetailsV2/JobDetailArtifactsV3.tsx
+++ b/Engine/Source/Programs/Horde/HordeDashboard/src/components/jobDetailsV2/JobDetailArtifactsV3.tsx
@@ -8,6 +8,7 @@ import backend from "../../backend";
 import { ArtifactContextType, GetArtifactDirectoryEntryResponse, GetArtifactDirectoryResponse, GetArtifactFileEntryResponse, GetArtifactResponseV2 } from "../../backend/Api";
 import { hordeClasses, modeColors } from "../../styles/Styles";
 import { JobDetailsV2 } from "./JobDetailsViewCommon";
+import dashboard from "../../backend/Dashboard";
 
 
 enum BrowserType {
@@ -27,10 +28,11 @@ type BrowserItem = {
 
 class ArtifactsHandler {
 
-   constructor(details: JobDetailsV2, stepId: string) {
+   constructor(details: JobDetailsV2, stepId: string, contextType: ArtifactContextType) {
       makeObservable(this);
       this.details = details;
       this.stepId = stepId;
+      this.context = contextType;
 
       this.set();
    }
@@ -61,19 +63,16 @@ class ArtifactsHandler {
       let a = artifacts.find(a => a.type === "step-saved");
       if (a) {
          this.contexts.push("step-saved");
-         this.context = "step-saved";
       }
 
       a = artifacts.find(a => a.type === "step-output");
       if (a) {
          this.contexts.push("step-output");
-         this.context = this.context ?? "step-saved";
       }
 
       a = artifacts.find(a => a.type === "step-trace");
       if (a) {
          this.contexts.push("step-trace");
-         this.context = this.context ?? "step-trace";
       }
 
       if (!this.context) {
@@ -96,23 +95,6 @@ class ArtifactsHandler {
 
    }
 
-   async setContext(c: ArtifactContextType) {
-
-
-      let a = this.artifacts?.find(a => a.type === c)!;
-      if (!a) {
-         return;
-      }
-
-      this.context = c;
-
-      this.artifact = a;
-      this.path = undefined;
-      this.browse = await backend.getBrowseArtifacts(a.id);
-      this.updateReady();
-
-   }
-
    async browseTo(path: string) {
 
       if (!this.artifact) {
@@ -186,7 +168,6 @@ class ArtifactsHandler {
       this.browse = undefined;
       this.artifact = undefined;
       this.artifacts = undefined;
-      this.context = undefined;
       this.contexts = undefined;
       this.stepId = "";
    }
@@ -205,9 +186,9 @@ class ArtifactsHandler {
    artifact?: GetArtifactResponseV2;
    artifacts?: GetArtifactResponseV2[];
 
-   contexts?: ArtifactContextType[];
+   private contexts?: ArtifactContextType[];
 
-   context?: ArtifactContextType;
+   readonly context: ArtifactContextType;
 
    stepId: string;
 }
@@ -215,14 +196,24 @@ class ArtifactsHandler {
 const styles = mergeStyleSets({
    list: {
       selectors: {
+         'a': {
+            height: "unset !important",
+         },
          '.ms-List-cell': {
+
             borderTop: "1px solid #EDEBE9",
             borderRight: "1px solid #EDEBE9",
             borderLeft: "1px solid #EDEBE9"
          },
          '.ms-List-cell:nth-last-child(-n + 1)': {
             borderBottom: "1px solid #EDEBE9"
-         }
+         },
+         ".ms-DetailsRow #artifactview": {
+            opacity: 0
+         },
+         ".ms-DetailsRow:hover #artifactview": {
+            opacity: 1
+         },
       }
    }
 });
@@ -235,9 +226,17 @@ const BrowseBreadCrumbs: React.FC<{ handler: ArtifactsHandler }> = observer(({ h
 
    const fontSize = 13;
 
+   let rootName = "Step";
+   if (handler.context === "step-output") {
+      rootName = "Output";
+   }
+   if (handler.context === "step-trace") {
+      rootName = "Trace";
+   }
+
    if (!handler.path) {
-      return <Stack><Text style={{ fontSize: fontSize, paddingLeft: 4, paddingRight: 4 }}>/
-      </Text>
+      return <Stack className="horde-no-darktheme">
+         <Text style={{ fontSize: fontSize, fontWeight: 600, paddingLeft: 2, paddingRight: 4 }}>{rootName} /</Text>
       </Stack>
    }
 
@@ -245,17 +244,23 @@ const BrowseBreadCrumbs: React.FC<{ handler: ArtifactsHandler }> = observer(({ h
 
    const elements = ppath.map((e, index) => {
       const path = ppath.slice(0, index + 1).join("/");
-      return <Stack horizontal onClick={() => handler.browseTo(path)} style={{ cursor: "pointer" }}>
+
+      const last = index === (ppath.length - 1);
+
+      const color = last ? undefined : (dashboard.darktheme ? "#55B7FF" : "#0078D4");
+      const cursor = last ? undefined : "pointer";
+
+      return <Stack horizontal onClick={() => handler.browseTo(path)} style={{ cursor: cursor }}>
          <Stack>
-            <Text style={{ fontSize: fontSize, fontFamily: "Horde Open Sans Light" }}>{e}</Text>
+            <Text style={{ fontSize: fontSize, color: color, fontWeight: last ? 600 : undefined }}>{e}</Text>
          </Stack>
-         {index !== (ppath.length - 1) && <Stack style={{ paddingLeft: 4, paddingRight: 4 }}><Text style={{ fontSize: fontSize, fontFamily: "Horde Open Sans Light" }}>/</Text></Stack>}
+         {index !== (ppath.length - 1) && <Stack style={{ paddingLeft: 4, paddingRight: 4 }}><Text style={{ fontSize: fontSize, fontWeight: 600 }}>/</Text></Stack>}
       </Stack>
    });
 
    elements.unshift(<Stack style={{ cursor: "pointer" }} onClick={() => handler.browseTo("")}>
-      <Stack style={{ paddingLeft: 4, paddingRight: 4 }}>
-         <Text style={{ fontSize: fontSize }}>/</Text>
+      <Stack style={{ paddingLeft: 2, paddingRight: 4 }}>
+         <Text style={{ fontSize: fontSize, color:  (dashboard.darktheme ? "#55B7FF" : "#0078D4"), fontWeight: 600 }}>{rootName} /</Text>
       </Stack>
    </Stack>);
 
@@ -307,13 +312,13 @@ const DownloadButton: React.FC<{ handler: ArtifactsHandler }> = observer(({ hand
    }
 
    if (selection.directoriesSelected || selection.filesSelected > 1) {
-      buttonText = `Download Zip (${sizeText})`;
+      buttonText = `Download (${sizeText})`;
    } else if (selection.filesSelected === 1) {
-      buttonText = `Download File (${sizeText})`;
+      buttonText = `Download (${sizeText})`;
    }
 
    return <Stack>
-      <PrimaryButton styles={{ root: { fontFamily: 'Horde Open Sans Semibold !important' } }} disabled={!selection.size} onClick={async () => {
+      <PrimaryButton styles={{ root: { fontFamily: 'Horde Open Sans SemiBold !important' } }} disabled={!selection.size} onClick={async () => {
 
          const selection = handler.currentSelection.items;
 
@@ -393,18 +398,6 @@ const JobDetailArtifactsInner: React.FC<{ handler: ArtifactsHandler }> = observe
       return null;
    }
 
-   const coptions: IDropdownOption[] = [];
-   if (handler.hasContext("step-saved")) {
-      coptions.push({ key: 'step-saved', text: 'Step' })
-   }
-   if (handler.hasContext("step-output")) {
-      coptions.push({ key: 'step-output', text: 'Output' })
-   }
-
-   if (handler.hasContext("step-trace")) {
-      coptions.push({ key: 'step-trace', text: 'Trace' })
-   }
-
    const items: BrowserItem[] = [];
 
    if (handler.path?.length) {
@@ -415,7 +408,7 @@ const JobDetailArtifactsInner: React.FC<{ handler: ArtifactsHandler }> = observe
 
       function recurseDirectories(dir: GetArtifactDirectoryEntryResponse, flattened: GetArtifactDirectoryEntryResponse[]) {
          if (!dir.directories?.length) {
-            const name = flattened.length ? flattened.map(d => d.name).join("/") + "/" + dir.name : dir.name;            
+            const name = flattened.length ? flattened.map(d => d.name).join("/") + "/" + dir.name : dir.name;
             items.push({ key: d.hash, text: name, icon: "Folder", type: BrowserType.Directory, size: d.length });
          } else {
             flattened.push(dir);
@@ -433,10 +426,9 @@ const JobDetailArtifactsInner: React.FC<{ handler: ArtifactsHandler }> = observe
    });
 
    const columns: IColumn[] = [
-      { key: 'column1', name: 'Name', minWidth: 768, maxWidth: 768, isResizable: false, isPadded: false },
-      { key: 'column2', name: 'View', minWidth: 32, maxWidth: 32, isResizable: false, isPadded: false },
-      { key: 'column3', name: 'Download', minWidth: 32, maxWidth: 32, isResizable: false, isPadded: false },
-      { key: 'column4', name: 'Size', minWidth: 128, maxWidth: 128, isResizable: false, isPadded: false },
+      { key: 'column1', name: 'Name', minWidth: 794, maxWidth: 794, isResizable: false, isPadded: false },
+      { key: 'column2', name: 'Size', minWidth: 128, maxWidth: 128, isResizable: false, isPadded: false },
+      { key: 'column3', name: 'View_Download', minWidth: 64, maxWidth: 64, isResizable: false, isPadded: false }
    ];
 
    const renderItem = (item: any, index?: number, column?: IColumn) => {
@@ -449,52 +441,42 @@ const JobDetailArtifactsInner: React.FC<{ handler: ArtifactsHandler }> = observe
          if (!item.size) {
             return null;
          }
-         return <Stack horizontalAlign="end">
+         return <Stack horizontalAlign="end" verticalAlign="center" verticalFill>
             <Text>{formatBytes(item.size, (item.size < (1024 * 1024)) ? 0 : 1)}</Text>
          </Stack>
 
       }
 
-      if (column.name === "View") {
-         if (item.type !== BrowserType.File) {
-            return null;
-         }
-         const path = encodeURI(handler.path + "/" + item.text);
-         const server = backend.serverUrl;
-         const href = `${server}/api/v2/artifacts/${handler.artifact!.id}/file?path=${path}&inline=true`;
-
-         return <Stack>
-            <IconButton href={href} target="_blank" style={{ height: 20, color: "#106EBE" }} iconProps={{ iconName: "Eye", styles: { root: { fontSize: "15px" } } }} />
-         </Stack>
-      }
+      if (column.name === "View_Download") {
 
-      if (column.name === "Download") {
          if (item.type !== BrowserType.File) {
             return null;
          }
+
          const path = encodeURI(handler.path + "/" + item.text);
          const server = backend.serverUrl;
          const href = `${server}/api/v2/artifacts/${handler.artifact!.id}/file?path=${path}`;
 
 
-         return <Stack>
-            <IconButton href={href} styles={{ root: { height: 20, color: "#106EBE" } }} iconProps={{ iconName: "CloudDownload", styles: { root: { fontSize: "14px" } } }} />
+         return <Stack data-selection-disabled verticalAlign="center" verticalFill horizontal horizontalAlign="end" style={{ paddingTop: 0, paddingBottom: 0 }}>
+            <IconButton id="artifactview" href={`${href}&inline=true`} target="_blank" style={{ paddingTop: 1, color: "#106EBE" }} iconProps={{ iconName: "Eye", styles: { root: { fontSize: "14px" } } }} />
+            <IconButton id="artifactview" href={href} target="_blank" style={{ paddingTop: 1, color: "#106EBE" }} iconProps={{ iconName: "CloudDownload", styles: { root: { fontSize: "14px" } } }} />
          </Stack>
       }
 
       if (column.name === "Name") {
 
-         let href: string | undefined;
-         let target: string | undefined;
+         const path = (item.text as string).split("/");
 
-         if (item.type === BrowserType.File) {
-            target = "_blank";
-            const path = encodeURI(handler.path + "/" + item.text);
-            const server = backend.serverUrl;
-            href = `${server}/api/v2/artifacts/${handler.artifact!.id}/file?path=${path}&inline=true`;
-         }
+         const pathElements = path.map((t, index) => {
+            const last = index == (path.length - 1);
+            let color = last ? (dashboard.darktheme ? "#FFFFFF" :"#605E5C") : undefined;
+            const font = last ? undefined : "Horde Open Sans Light";
+            const sep = last ? undefined : "/"
+            return <Text styles={{ root: { fontFamily: font } }} style={{ color: color }}>{t}{sep}</Text>
+         })
 
-         return <FluentLink href={href} target={target} style={{ color: "unset", height: "100%", width: "100%" }}><Stack style={{ cursor: "pointer" }} onClick={(ev) => {
+         return <Stack verticalFill verticalAlign="center" style={{ cursor: "pointer" }} onClick={(ev) => {
 
             if (item.type === BrowserType.Directory) {
                const nbrowse = handler.path ? `${handler.path}/${item.text}` : item.text;
@@ -507,35 +489,28 @@ const JobDetailArtifactsInner: React.FC<{ handler: ArtifactsHandler }> = observe
             }
 
          }}>
-            {item.type !== BrowserType.NavigateUp && <Stack data-selection-disabled horizontal tokens={{ childrenGap: 8 }}>
+            {item.type !== BrowserType.NavigateUp && <Stack data-selection-disabled={item.type !== BrowserType.File} horizontal tokens={{ childrenGap: 8 }}>
                <Stack>
-                  <FontIcon style={{ fontSize: 16 }} iconName={item.icon} />
+                  <FontIcon style={{ paddingTop: 1, fontSize: 16 }} iconName={item.icon} />
                </Stack>
-               <Stack>
-                  <Text>{item.text}</Text>
+               <Stack horizontal>
+                  {pathElements}
                </Stack>
             </Stack>}
-            {item.type === BrowserType.NavigateUp && <Stack data-selection-disabled style={{ height: "100%", width: "100%" }}>
+            {item.type === BrowserType.NavigateUp && <Stack data-selection-disabled verticalFill verticalAlign="center">
                <Text>..</Text>
             </Stack>
             }
          </Stack>
-         </FluentLink>
+
       }
 
       return null;
 
    }
 
-   return <Stack key={ `jobdetailartifacts_${idcounter++}`} tokens={{ childrenGap: 12 }}>
+   return <Stack key={`jobdetailartifacts_${idcounter++}`} tokens={{ childrenGap: 12 }}>
       <Stack horizontal verticalAlign="center" tokens={{ childrenGap: 18 }} style={{ paddingBottom: 12 }}>
-         <Stack>
-            <Dropdown className={hordeClasses.modal} style={{ width: 128 }} options={coptions} selectedKey={handler.context} onChange={(ev, option) => {
-               if (option?.key) {
-                  handler.setContext(option.key as ArtifactContextType);
-               }
-            }} />
-         </Stack>
          <Stack>
             <BrowseBreadCrumbs handler={handler} />
          </Stack>
@@ -546,6 +521,7 @@ const JobDetailArtifactsInner: React.FC<{ handler: ArtifactsHandler }> = observe
          <ScrollablePane style={{ height: 492 + 160 }}>
             <SelectionZone selection={handler.selection!}>
                <DetailsList
+                  styles={{ root: { overflowX: "hidden" } }}
                   className={styles.list}
                   isHeaderVisible={false}
                   compact={true}
@@ -569,9 +545,9 @@ const JobDetailArtifactsInner: React.FC<{ handler: ArtifactsHandler }> = observe
 
 
 
-export const JobArtifactsModal: React.FC<{ jobDetails: JobDetailsV2; stepId: string, onClose: () => void }> = ({ jobDetails, stepId, onClose }) => {
+export const JobArtifactsModal: React.FC<{ jobDetails: JobDetailsV2; stepId: string, contextType: ArtifactContextType, onClose: () => void }> = ({ jobDetails, stepId, contextType, onClose }) => {
 
-   const [handler] = useState(new ArtifactsHandler(jobDetails, stepId));
+   const [handler] = useState(new ArtifactsHandler(jobDetails, stepId, contextType));
 
    useEffect(() => {
       return () => {
@@ -581,11 +557,11 @@ export const JobArtifactsModal: React.FC<{ jobDetails: JobDetailsV2; stepId: str
 
    return <Stack>
       <Modal isOpen={true} isBlocking={true} topOffsetFixed={true} styles={{ main: { padding: 8, width: 1180, height: 820, hasBeenOpened: false, top: "80px", position: "absolute" } }} onDismiss={() => onClose()} className={hordeClasses.modal}>
-         <Stack styles={{ root: { paddingTop: 18, paddingRight: 12 } }}>
+         <Stack className="horde-no-darktheme" styles={{ root: { paddingTop: 10, paddingRight: 12 } }}>
             <Stack style={{ paddingLeft: 24, paddingRight: 24 }}>
                <Stack tokens={{ childrenGap: 12 }} style={{ height: 800 }}>
-                  <Stack horizontal>
-                     <Stack>
+                  <Stack horizontal verticalAlign="start">
+                     <Stack style={{ paddingTop: 3 }}>
                         <Text variant="mediumPlus" styles={{ root: { fontFamily: "Horde Open Sans SemiBold" } }}>Artifacts</Text>
                      </Stack>
                      <Stack grow />
diff --git a/Engine/Source/Programs/Horde/HordeDashboard/src/components/jobDetailsV2/JobOperationsBar.tsx b/Engine/Source/Programs/Horde/HordeDashboard/src/components/jobDetailsV2/JobOperationsBar.tsx
--- a/Engine/Source/Programs/Horde/HordeDashboard/src/components/jobDetailsV2/JobOperationsBar.tsx
+++ b/Engine/Source/Programs/Horde/HordeDashboard/src/components/jobDetailsV2/JobOperationsBar.tsx
@@ -4,7 +4,7 @@ import { CommandBar, CommandBarButton, ICommandBarItemProps, IContextualMenuItem
 import { observer } from 'mobx-react-lite';
 import React, { useState } from 'react';
 import { Link, useNavigate } from 'react-router-dom';
-import { JobState, StepData } from '../../backend/Api';
+import { ArtifactContextType, JobState, StepData } from '../../backend/Api';
 import dashboard from '../../backend/Dashboard';
 import { hordeClasses } from '../../styles/Styles';
 import { EditJobModal } from '../EditJobModal';
@@ -29,7 +29,6 @@ export const JobOperations: React.FC<{ jobDetails: JobDetailsV2 }> = observer(({
    const navigate = useNavigate();
    const [abortShown, setAbortShown] = useState(false);
    const [editShown, setEditShown] = useState(false);
-   const [artifactsShown, setArtifactsShown] = useState(false);
 
    const [parametersState, setParametersState] = useState(query.get("newbuild") ? ParameterState.Clone : ParameterState.Hidden);
 
@@ -47,8 +46,6 @@ export const JobOperations: React.FC<{ jobDetails: JobDetailsV2 }> = observer(({
       return null;
    }
 
-   const stepArtifacts = jobDetails.stepArtifacts.get(stepId ?? "");
-
    const abortDisabled = jobData.state === JobState.Complete;
    const runAgainDisabled = false; /*jobDetails.jobdata?.state !== JobState.Complete*/
 
@@ -122,7 +119,6 @@ export const JobOperations: React.FC<{ jobDetails: JobDetailsV2 }> = observer(({
    }
 
    return <Stack>
-      {artifactsShown && !!step && <JobArtifactsModal stepId={step.id} jobDetails={jobDetails} onClose={() => { setArtifactsShown(false); }} />}
       <AbortJobModal jobDetails={jobDetails} show={abortShown} onClose={() => { setAbortShown(false); }} />
       <EditJobModal jobData={jobDetails.jobData} show={editShown} onClose={() => { setEditShown(false); }} />
       <NewBuild streamId={jobDetails.stream!.id} jobDetails={jobDetails} readOnly={parametersState === ParameterState.Parameters}
@@ -147,11 +143,9 @@ export const JobOperations: React.FC<{ jobDetails: JobDetailsV2 }> = observer(({
             </Stack>
             </Link>}
 
-            {!!stepId && !!jobData?.useArtifactsV2 && <div onClick={() => { setArtifactsShown(true) }}>
-               <Stack styles={{ root: { paddingTop: 4 } }}>
-                  <CommandBarButton disabled={!stepArtifacts?.length} className={hordeClasses.commandBarSmall} styles={{ root: { padding: "10px 8px 10px 8px" } }} iconProps={{ iconName: "CloudDownload" }} text="Artifacts" />
-               </Stack>
-            </div>}
+            {!!stepId && <Stack styles={{ root: { paddingTop: 4 } }}>
+               <StepArtifactsOperations jobDetails={jobDetails} stepId={stepId} />
+            </Stack>}
 
             <Stack className={hordeClasses.commandBarSmall} styles={{ root: { paddingTop: 4 } }}>
                {notifications}
@@ -180,6 +174,89 @@ export const JobOperations: React.FC<{ jobDetails: JobDetailsV2 }> = observer(({
 
 });
 
+const StepArtifactsOperations: React.FC<{ jobDetails: JobDetailsV2, stepId: string }> = observer(({ jobDetails, stepId }) => {
+
+   const [artifactsShown, setArtifactsShown] = useState<ArtifactContextType | undefined>(undefined);
+
+   const jobData = jobDetails.jobData;
+
+   if (!jobData?.useArtifactsV2) {
+      return null;
+   }
+
+   // subscribe
+   if (dashboard.updated) { }
+   if (jobDetails.updated) { }
+
+   const stepArtifacts = jobDetails.stepArtifacts.get(stepId ?? "");
+
+   const step = jobDetails.stepById(stepId)
+   if (!step) {
+      return null;
+   }
+
+   const atypes = new Map<ArtifactContextType, number>();
+
+   stepArtifacts?.forEach(a => {
+      let c = atypes.get(a.type) ?? 0;
+      c++;
+      atypes.set(a.type, c);
+   });
+
+   const opsList: IContextualMenuItem[] = [];
+
+   opsList.push({
+      key: 'stepops_artifacts_step',
+      text: "Step Artifacts",
+      iconProps: { iconName: "Folder" },
+      disabled: !atypes.get("step-saved"),
+      onClick: () => { setArtifactsShown("step-saved")}
+   });
+
+   opsList.push({
+      key: 'stepops_artifacts_output',
+      text: "Output Artifacts",
+      iconProps: { iconName: "MenuOpen" },
+      disabled: !atypes.get("step-output"),
+      onClick: () => { setArtifactsShown("step-output")}
+   });
+
+   opsList.push({
+      key: 'stepops_artifacts_trace',
+      text: "Trace Artifacts",
+      iconProps: { iconName: "SearchTemplate" },
+      onClick: () => { setArtifactsShown("step-trace")}      
+   });
+
+   const opsItems: ICommandBarItemProps[] = [
+      {
+         key: 'stepops_artifacts',
+         text: "Artifacts",
+         iconProps: { iconName: "CloudDownload" },
+         disabled: !stepArtifacts?.length,
+         subMenuProps: {
+            items: opsList
+         }
+      }
+   ];
+
+   return <Stack>
+      {!!artifactsShown && <JobArtifactsModal stepId={step.id} jobDetails={jobDetails} contextType={artifactsShown} onClose={() => { setArtifactsShown(undefined); }} />}
+      <Stack horizontal styles={{ root: { paddingLeft: 0 } }}>
+         <Stack grow />
+         <Stack horizontal>
+            <Stack className={hordeClasses.commandBarSmall} style={{ paddingRight: 0 }}>
+               <CommandBar
+                  items={opsItems}
+                  onReduceData={() => undefined} />
+            </Stack>
+         </Stack>
+      </Stack>
+   </Stack>
+
+
+});
+
 const StepOperations: React.FC<{ jobDetails: JobDetailsV2, stepId: string }> = observer(({ jobDetails, stepId }) => {
 
    const [shown, setShown] = useState<{ abortShown?: boolean, retryShown?: boolean, pauseShown?: boolean }>({});
