Horde: Support copying CL to clipboard on trendsview

diff --git a/Engine/Source/Programs/Horde/HordeDashboard/src/components/jobDetailsV2/JobDetailStepTrends.tsx b/Engine/Source/Programs/Horde/HordeDashboard/src/components/jobDetailsV2/JobDetailStepTrends.tsx
--- a/Engine/Source/Programs/Horde/HordeDashboard/src/components/jobDetailsV2/JobDetailStepTrends.tsx
+++ b/Engine/Source/Programs/Horde/HordeDashboard/src/components/jobDetailsV2/JobDetailStepTrends.tsx
@@ -17,6 +17,7 @@ import { msecToElapsed } from '../../base/utilities/timeUtils';
 import { hordeClasses } from '../../styles/Styles';
 import { useQuery } from '../JobDetailCommon';
 import { JobDataView, JobDetailsV2 } from './JobDetailsViewCommon';
+import { copyToClipboard } from '../../base/utilities/clipboard';
 
 const sideRail: ISideRailLink = { text: "Trends", url: "rail_trends" };
 
@@ -379,7 +380,7 @@ function preventChartScroll(this: Document, e: WheelEvent) {
    e.returnValue = false
 }
 
-export const StepTrendsPanel: React.FC<{ jobDetails: JobDetailsV2}> = observer(({ jobDetails}) => {
+export const StepTrendsPanel: React.FC<{ jobDetails: JobDetailsV2 }> = observer(({ jobDetails }) => {
    const query = useQuery();
 
    const stepId = query.get("step") ? query.get("step")! : undefined;
@@ -392,15 +393,15 @@ export const StepTrendsPanel: React.FC<{ jobDetails: JobDetailsV2}> = observer((
          dataView?.clear();
       };
    }, [dataView]);
-   
+
    dataView.subscribe();
 
    if (jobDetails?.jobData) {
-      dataView.initialize(jobDetails.jobData?.preflightChange ? undefined : [sideRail]);       
+      dataView.initialize(jobDetails.jobData?.preflightChange ? undefined : [sideRail]);
    } else {
       jobDetails?.subscribe();
    }
-   
+
    if (jobDetails.jobData == null || !jobDetails?.viewsReady || jobDetails.jobData.preflightChange) {
       return null;
    }
@@ -429,6 +430,7 @@ const StepTrendsPanelInner: React.FC<{ jobDetails: JobDetailsV2; dataView: Trend
    const [selectedTemplateRef, setSelectedTemplateRef] = useState<GetTemplateRefResponse | undefined>(undefined);
    const [includeWaitTimes, setIncludeWaitTimes] = useState<boolean>(false);
    const [includeStepDependencies, setIncludeStepDependencies] = useState<boolean>(false);
+   const [copyCLToClipboard, setCopyCLToClipboard] = useState<string | undefined>(undefined);
 
    const [pageInfo, setPageInfo] = useState<PageInfo>({ pageNumber: 1, zoomLevel: 10, pageSize: 10 });
 
@@ -467,12 +469,12 @@ const StepTrendsPanelInner: React.FC<{ jobDetails: JobDetailsV2; dataView: Trend
       };
       dataView.populateChartDataCache(query);
    }
-   
+
    if (dataView.cacheData?.chartData) {
-      let totalValues = Object.values(dataView.cacheData.chartData!.labelData).map(c => c.length).reduce((a, b) => Math.max(a, b));   
+      let totalValues = Object.values(dataView.cacheData.chartData!.labelData).map(c => c.length).reduce((a, b) => Math.max(a, b));
       pageInfo.pageSize = totalValues / 10;
    }
-   
+
 
 
    const filterChartData = (chartDataCache?: ChartDataCache) => {
@@ -565,7 +567,11 @@ const StepTrendsPanelInner: React.FC<{ jobDetails: JobDetailsV2; dataView: Trend
                      y: stepRunTime,
                      yAxisCalloutData: calloutData,
                      xAxisCalloutData: `CL ${change}`,
-                     //onDataPointClick: () => selectDataPoint(helperData)
+                     onDataPointClick: () => {
+                        copyToClipboard(change); setCopyCLToClipboard(change); setTimeout(() => {
+                           setCopyCLToClipboard(undefined);
+                        }, 3000)
+                     }
                   });
                }
             }
@@ -889,12 +895,15 @@ const StepTrendsPanelInner: React.FC<{ jobDetails: JobDetailsV2; dataView: Trend
             <Stack horizontal horizontalAlign="space-between" styles={{ root: { minHeight: 32, paddingRight: 8 } }}>
                <Text variant="mediumPlus" styles={{ root: { fontFamily: "Horde Open Sans SemiBold" } }}>{titleText}</Text>
                {!!step && <Toggle label="Include dependencies" checked={includeStepDependencies} onText="On" offText="Off" onChange={() => setIncludeStepDependencies(!includeStepDependencies)} />}
-               <Toggle label="Include wait times" checked={includeWaitTimes} onText="On" offText="Off" onChange={() => setIncludeWaitTimes(!includeWaitTimes)} />               
+               <Toggle label="Include wait times" checked={includeWaitTimes} onText="On" offText="Off" onChange={() => setIncludeWaitTimes(!includeWaitTimes)} />
             </Stack>
             <Stack styles={{ root: { position: 'relative', top: -10 } }}>
                {statusDiv}
             </Stack>
-            <Stack styles={{ root: { width: 1370, minHeight: 400, height: 400 } }} onWheel={onChartScroll} onMouseEnter={onChartWheelEnter} onMouseLeave={onChartWheelLeave}>
+            <Stack styles={{ root: { width: 1370, minHeight: 400, height: 400, position: "relative" } }} onWheel={onChartScroll} onMouseEnter={onChartWheelEnter} onMouseLeave={onChartWheelLeave}>
+               {!!copyCLToClipboard && <Stack style={{ position: "absolute", top: -12, left: 580 }}>
+                  <Text variant="mediumPlus">Copied CL {copyCLToClipboard} to the clipboard</Text>
+               </Stack>}
                {placeholderDiv}
             </Stack>
 
