Horde: Use a struct that wraps a string to represent ACL actions, rather than an enum. This allows more flexibility and extensibility in where ACL actions can be defined.

diff --git a/Engine/Source/Programs/Horde/Horde.Server.Tests/DeviceServiceTests.cs b/Engine/Source/Programs/Horde/Horde.Server.Tests/DeviceServiceTests.cs
--- a/Engine/Source/Programs/Horde/Horde.Server.Tests/DeviceServiceTests.cs
+++ b/Engine/Source/Programs/Horde/Horde.Server.Tests/DeviceServiceTests.cs
@@ -36,9 +36,12 @@ namespace Horde.Server.Tests
 					ControllerContext controllerContext = new ControllerContext();
 					controllerContext.HttpContext = new DefaultHttpContext();
 					controllerContext.HttpContext.User = new ClaimsPrincipal(new ClaimsIdentity(
-						new List<Claim> {new Claim(ServerSettings.AdminClaimType, ServerSettings.AdminClaimValue),
-					new Claim(ClaimTypes.Name, "TestUser"),
-					new Claim(HordeClaimTypes.UserId, user.Id.ToString()) }
+						new List<Claim>
+						{ 
+							HordeClaims.AdminClaim.ToClaim(),
+							new Claim(ClaimTypes.Name, "TestUser"),
+							new Claim(HordeClaimTypes.UserId, user.Id.ToString()) 
+						}
 						, "TestAuthType"));
 					_deviceController.ControllerContext = controllerContext;
 
diff --git a/Engine/Source/Programs/Horde/Horde.Server.Tests/RpcServiceTest.cs b/Engine/Source/Programs/Horde/Horde.Server.Tests/RpcServiceTest.cs
--- a/Engine/Source/Programs/Horde/Horde.Server.Tests/RpcServiceTest.cs
+++ b/Engine/Source/Programs/Horde/Horde.Server.Tests/RpcServiceTest.cs
@@ -53,7 +53,7 @@ namespace Horde.Server.Tests
 	[TestClass]
 	public class RpcServiceTest : TestSetup
 	{
-		private readonly ServerCallContext _adminContext = new ServerCallContextStub("app-horde-admins");
+		private readonly ServerCallContext _adminContext = new ServerCallContextStub(HordeClaims.AdminClaim.ToClaim());
 
 		sealed class HttpContextStub : HttpContext
 		{
@@ -69,11 +69,11 @@ namespace Horde.Server.Tests
 			public override ClaimsPrincipal User { get; set; }
 			public override WebSocketManager WebSockets { get; } = null!;
 
-			public HttpContextStub(string roleClaimType)
+			public HttpContextStub(Claim roleClaimType)
 			{
 				User = new ClaimsPrincipal(new ClaimsIdentity(new List<Claim>
 				{
-					new Claim(HordeClaimTypes.Role, roleClaimType),
+					roleClaimType
 				}, "TestAuthType"));
 			}
 			
@@ -97,7 +97,7 @@ namespace Horde.Server.Tests
 			{
 				return new ServerCallContextStub(new ClaimsPrincipal(new ClaimsIdentity(new List<Claim>
 				{
-					new Claim(HordeClaimTypes.Role, "app-horde-admins"),
+					HordeClaims.AdminClaim.ToClaim(),
 					new Claim(HordeClaimTypes.AgentSessionId, agentSessionId),
 				}, "TestAuthType")));
 			}
@@ -106,11 +106,11 @@ namespace Horde.Server.Tests
 			{
 				return new ServerCallContextStub(new ClaimsPrincipal(new ClaimsIdentity(new List<Claim>
 				{
-					new Claim(HordeClaimTypes.Role, "app-horde-admins"),
+					HordeClaims.AdminClaim.ToClaim()
 				}, "TestAuthType")));
 			}
 
-			public ServerCallContextStub(string roleClaimType)
+			public ServerCallContextStub(Claim roleClaimType)
 			{
 				// The GetHttpContext extension falls back to getting the HttpContext from UserState
 				// We can piggyback on that behavior during tests
@@ -431,7 +431,7 @@ namespace Horde.Server.Tests
 			SessionId sessionId = SessionId.GenerateNewId();
 			ServerCallContext context = new ServerCallContextStub(new ClaimsPrincipal(new ClaimsIdentity(new List<Claim>
 			{
-				new Claim(HordeClaimTypes.Role, "app-horde-admins"),
+				HordeClaims.AdminClaim.ToClaim(),
 				new Claim(HordeClaimTypes.AgentSessionId, sessionId.ToString()),
 			}, "TestAuthType")));
 
diff --git a/Engine/Source/Programs/Horde/Horde.Server.Tests/SwaggerTest.cs b/Engine/Source/Programs/Horde/Horde.Server.Tests/SwaggerTest.cs
--- a/Engine/Source/Programs/Horde/Horde.Server.Tests/SwaggerTest.cs
+++ b/Engine/Source/Programs/Horde/Horde.Server.Tests/SwaggerTest.cs
@@ -14,7 +14,7 @@ namespace Horde.Server.Tests
         [TestMethod]
         public async Task ValidateSwagger()
         {
-	        HttpResponseMessage res = await Client.GetAsync(new Uri("/swagger/v1/swagger.json"));
+	        HttpResponseMessage res = await Client.GetAsync(new Uri("https://localhost:5001/swagger/v1/swagger.json"));
 	        if (!res.IsSuccessStatusCode)
 	        {
 		        string rawJson = await res.Content.ReadAsStringAsync();
diff --git a/Engine/Source/Programs/Horde/Horde.Server.Tests/TestDataTests.cs b/Engine/Source/Programs/Horde/Horde.Server.Tests/TestDataTests.cs
--- a/Engine/Source/Programs/Horde/Horde.Server.Tests/TestDataTests.cs
+++ b/Engine/Source/Programs/Horde/Horde.Server.Tests/TestDataTests.cs
@@ -52,9 +52,9 @@ namespace Horde.Server.Tests
 					ControllerContext controllerContext = new ControllerContext();
 					controllerContext.HttpContext = new DefaultHttpContext();
 					controllerContext.HttpContext.User = new ClaimsPrincipal(new ClaimsIdentity(
-						new List<Claim> {new Claim(ServerSettings.AdminClaimType, ServerSettings.AdminClaimValue),
-					new Claim(ClaimTypes.Name, "TestUser"),
-					new Claim(HordeClaimTypes.UserId, user.Id.ToString()) }
+						new List<Claim> { HordeClaims.AdminClaim.ToClaim(),
+						new Claim(ClaimTypes.Name, "TestUser"),
+						new Claim(HordeClaimTypes.UserId, user.Id.ToString()) }
 						, "TestAuthType"));
 					_testDataController.ControllerContext = controllerContext;
 
diff --git a/Engine/Source/Programs/Horde/Horde.Server.Tests/TestSetup.cs b/Engine/Source/Programs/Horde/Horde.Server.Tests/TestSetup.cs
--- a/Engine/Source/Programs/Horde/Horde.Server.Tests/TestSetup.cs
+++ b/Engine/Source/Programs/Horde/Horde.Server.Tests/TestSetup.cs
@@ -166,8 +166,6 @@ namespace Horde.Server.Tests
 			{
 			}
 
-			settings.AdminClaimType = HordeClaimTypes.Role;
-			settings.AdminClaimValue = "app-horde-admins";
 			settings.WithAws = true;
 
 			settings.ForceConfigUpdateOnStartup = true;
@@ -328,12 +326,12 @@ namespace Horde.Server.Tests
 			return controller;
 		}
 		
-		private ControllerContext GetControllerContext()
+		private static ControllerContext GetControllerContext()
 		{
 			ControllerContext controllerContext = new ControllerContext();
 			controllerContext.HttpContext = new DefaultHttpContext();
 			controllerContext.HttpContext.User = new ClaimsPrincipal(new ClaimsIdentity(
-				new List<Claim> {new Claim(ServerSettings.AdminClaimType, ServerSettings.AdminClaimValue)}, "TestAuthType"));
+				new List<Claim> { HordeClaims.AdminClaim.ToClaim() }, "TestAuthType"));
 			return controllerContext;
 		}
 		
diff --git a/Engine/Source/Programs/Horde/Horde.Server/Acls/Acl.cs b/Engine/Source/Programs/Horde/Horde.Server/Acls/Acl.cs
--- a/Engine/Source/Programs/Horde/Horde.Server/Acls/Acl.cs
+++ b/Engine/Source/Programs/Horde/Horde.Server/Acls/Acl.cs
@@ -1,523 +1,48 @@
 // Copyright Epic Games, Inc. All Rights Reserved.
 
 using System;
-using System.Collections.Generic;
-using System.Linq;
-using System.Linq.Expressions;
-using System.Security.Claims;
-using MongoDB.Bson;
-using MongoDB.Bson.Serialization;
-using MongoDB.Bson.Serialization.Attributes;
-using MongoDB.Driver;
+using System.ComponentModel;
+using System.Globalization;
+using System.Text.Json;
+using System.Text.Json.Serialization;
+using EpicGames.Core;
 
 namespace Horde.Server.Acls
 {
 	/// <summary>
-	/// Set of actions that can be performed by a user. NOTE: This enum is sensitive to ordering. Do not change values.
+	/// Wraps a string used to describe an ACL action
 	/// </summary>
-	public enum AclAction
+	/// <param name="Name">Name of the action</param>
+	[JsonSchemaString]
+	[JsonConverter(typeof(AclActionJsonConverter))]
+	[TypeConverter(typeof(AclActionTypeConverter))]
+	public record struct AclAction(string Name)
 	{
-		//// PROJECTS ////
+		/// <inheritdoc/>
+		public override string ToString() => Name;
+	}
 
-        #region Projects
-
-        /// <summary>
-        /// Allows the creation of new projects
-        /// </summary>
-        CreateProject,
-
-		/// <summary>
-		/// Allows deletion of projects.
-		/// </summary>
-		DeleteProject,
-		
-		/// <summary>
-		/// Modify attributes of a project (name, categories, etc...)
-		/// </summary>
-		UpdateProject,
-
-		/// <summary>
-		/// View information about a project
-		/// </summary>
-		ViewProject,
-
-        #endregion
-
-		//// STREAMS ////
-
-        #region Streams
-
-        /// <summary>
-        /// Allows the creation of new streams within a project
-        /// </summary>
-        CreateStream,
-
-		/// <summary>
-		/// Allows updating a stream (agent types, templates, schedules)
-		/// </summary>
-		UpdateStream,
-
-		/// <summary>
-		/// Allows deleting a stream
-		/// </summary>
-		DeleteStream,
-
-		/// <summary>
-		/// Ability to view a stream
-		/// </summary>
-		ViewStream,
-
-		/// <summary>
-		/// View changes submitted to a stream. NOTE: this returns responses from the server's Perforce account, which may be a priviledged user.
-		/// </summary>
-		ViewChanges,
-
-		/// <summary>
-		/// Update records in commit queues for a stream
-		/// </summary>
-		UpdateCommitQueues,
-
-		/// <summary>
-		/// View commit queues for a stream
-		/// </summary>
-		ViewCommitQueues,
-
-        #endregion
-
-		//// JOBS ////
-
-        #region Jobs
-
-        /// <summary>
-        /// Ability to start new jobs
-        /// </summary>
-        CreateJob,
-
-		/// <summary>
-		/// Rename a job, modify its priority, etc...
-		/// </summary>
-		UpdateJob,
-
-		/// <summary>
-		/// Delete a job properties
-		/// </summary>
-		DeleteJob,
-
-		/// <summary>
-		/// Allows updating a job metadata (name, changelist number, step properties, new groups, job states, etc...). Typically granted to agents. Not user facing.
-		/// </summary>
-		ExecuteJob,
-
-		/// <summary>
-		/// Ability to retry a failed job step
-		/// </summary>
-		RetryJobStep,
-
-		/// <summary>
-		/// Ability to view a job
-		/// </summary>
-		ViewJob,
-
-		#endregion
-
-		//// ARTIFACTS ////
-
-		#region Artifacts
-
-		/// <summary>
-		/// Permission to read from an artifact
-		/// </summary>
-		ReadArtifact,
-
-		/// <summary>
-		/// Permission to write to an artifact
-		/// </summary>
-		WriteArtifact,
-
-		#endregion
-
-		//// EVENTS ////
-
-		#region Events
-
-		/// <summary>
-		/// Ability to create events
-		/// </summary>
-		CreateEvent,
-
-		/// <summary>
-		/// Ability to view events
-		/// </summary>
-		ViewEvent,
-
-        #endregion
-
-		//// AGENTS ////
-
-        #region Agents
-
-        /// <summary>
-        /// Ability to create an agent. This may be done explicitly, or granted to agents to allow them to self-register.
-        /// </summary>
-        CreateAgent,
-
-		/// <summary>
-		/// Update an agent's name, pools, etc...
-		/// </summary>
-		UpdateAgent,
-
-		/// <summary>
-		/// Soft-delete an agent
-		/// </summary>
-		DeleteAgent,
-
-		/// <summary>
-		/// View an agent
-		/// </summary>
-		ViewAgent,
-
-		/// <summary>
-		/// List the available agents
-		/// </summary>
-		ListAgents,
-
-        #endregion
-
-		//// POOLS ////
-
-        #region Pools
-
-        /// <summary>
-        /// Create a global pool of agents
-        /// </summary>
-        CreatePool,
-
-		/// <summary>
-		/// Modify an agent pool
-		/// </summary>
-		UpdatePool,
-
-		/// <summary>
-		/// Delete an agent pool
-		/// </summary>
-		DeletePool,
-
-		/// <summary>
-		/// Ability to view a pool
-		/// </summary>
-		ViewPool,
-
-		/// <summary>
-		/// View all the available agent pools
-		/// </summary>
-		ListPools,
-
-        #endregion
-
-		//// SESSIONS ////
-
-        #region Sessions
-
-        /// <summary>
-        /// Granted to agents to call CreateSession, which returns a bearer token identifying themselves valid to call UpdateSesssion via gRPC.
-        /// </summary>
-        CreateSession,
-
-		/// <summary>
-		/// Allows viewing information about an agent session
-		/// </summary>
-		ViewSession,
-
-        #endregion
-
-        //// CREDENTIALS ////
-
-        #region Credentials
-
-        /// <summary>
-        /// Create a new credential
-        /// </summary>
-        CreateCredential,
-
-		/// <summary>
-		/// Delete a credential
-		/// </summary>
-		DeleteCredential,
-
-		/// <summary>
-		/// Modify an existing credential
-		/// </summary>
-		UpdateCredential,
-
-		/// <summary>
-		/// Enumerates all the available credentials
-		/// </summary>
-		ListCredentials,
-
-		/// <summary>
-		/// View a credential
-		/// </summary>
-		ViewCredential,
-
-        #endregion
-
-        //// LEASES ////
-
-        #region Leases
-
-        /// <summary>
-        /// View all the leases that an agent has worked on
-        /// </summary>
-        ViewLeases,
-
-        #endregion
-
-        //// COMMITS ////
-
-        #region Commits
-
-        /// <summary>
-        /// Add a new commit
-        /// </summary>
-        AddCommit,
-
-		/// <summary>
-		/// List all the commits that have been added
-		/// </summary>
-		FindCommits,
-
-		/// <summary>
-		/// View the commits for a particular stream
-		/// </summary>
-		ViewCommits,
-
-        #endregion
-
-        //// ISSUES ////
-
-        #region Issues
-
-        /// <summary>
-        /// View a build health issue
-        /// </summary>
-        ViewIssue,
-
-		/// <summary>
-		/// Mark an issue as fixed via the p4fix endpoint
-		/// </summary>
-		IssueFixViaPerforce,
-
-        #endregion
-
-        //// TEMPLATES ////
-
-        #region Templates
-
-        /// <summary>
-        /// View template associated with a stream
-        /// </summary>
-        ViewTemplate,
-
-        #endregion
-
-        //// LOGS ////
-
-        #region Logs
-
-        /// <summary>
-        /// Ability to create a log. Implicitly granted to agents.
-        /// </summary>
-        CreateLog,
-
-		/// <summary>
-		/// Ability to update log metadata
-		/// </summary>
-		UpdateLog,
-
-		/// <summary>
-		/// Ability to view a log contents
-		/// </summary>
-		ViewLog,
-
-		/// <summary>
-		/// Ability to write log data
-		/// </summary>
-		WriteLogData,
-
-        #endregion
-
-        //// ARTIFACTS ////
-
-        #region Artifacts
-
-        /// <summary>
-        /// Ability to create an artifact. Typically just for debugging; agents have this access for a particular session.
-        /// </summary>
-        UploadArtifact,
-
-		/// <summary>
-		/// Ability to download an artifact
-		/// </summary>
-		DownloadArtifact,
-
-        #endregion
-
-        //// SOFTWARE ////
-
-        #region Software
-
-        /// <summary>
-        /// Ability to upload new versions of the agent software
-        /// </summary>
-        UploadSoftware,
-
-		/// <summary>
-		/// Ability to download the agent software
-		/// </summary>
-		DownloadSoftware,
-
-		/// <summary>
-		/// Ability to delete agent software
-		/// </summary>
-		DeleteSoftware,
-
-		#endregion
-
-		//// TOOLS ////
-
-		#region Tools
-
-		/// <summary>
-		/// Ability to download a tool
-		/// </summary>
-		DownloadTool,
-
-		/// <summary>
-		/// Ability to upload new tool versions
-		/// </summary>
-		UploadTool,
-
-		#endregion
-
-		//// ADMIN ////
-
-		#region Admin
-
-		/// <summary>
-		/// Ability to read any data from the server. Always inherited.
-		/// </summary>
-		AdminRead,
-
-		/// <summary>
-		/// Ability to write any data to the server.
-		/// </summary>
-		AdminWrite,
-
-		/// <summary>
-		/// Ability to impersonate another user
-		/// </summary>
-		Impersonate,
-
-		/// <summary>
-		/// View estimated costs for particular operations
-		/// </summary>
-		ViewCosts,
-
-        #endregion
-
-        //// PERMISSIONS ////
-
-        #region Permissions
-
-        /// <summary>
-        /// Ability to view permissions on an object
-        /// </summary>
-        ViewPermissions,
-
-		/// <summary>
-		/// Ability to change permissions on an object
-		/// </summary>
-		ChangePermissions,
-
-		/// <summary>
-		/// Issue bearer token for the current user
-		/// </summary>
-		IssueBearerToken,
-
-        #endregion
-
-        //// NOTIFICATIONS ////
-
-        #region Notifications
-
-        /// <summary>
-        /// Ability to subscribe to notifications
-        /// </summary>
-        CreateSubscription,
-
-        #endregion
-
-        //// DEVICES ////
-
-        #region Devices
-
-        /// <summary>
-        /// Ability to read devices
-        /// </summary>
-        DeviceRead,
-
-		/// <summary>
-		/// Ability to write devices
-		/// </summary>
-		DeviceWrite,
-
-        #endregion
-
-        //// COMPUTE ////
-
-        #region Compute
-
-        /// <summary>
-        /// User can add tasks to the compute cluster
-        /// </summary>
-        AddComputeTasks,
-
-		/// <summary>
-		/// User can poll for compute results
-		/// </summary>
-		ViewComputeTasks,
-
-        #endregion
-
-        //// STORAGE ////
-
-        #region Storage
-
-        /// <summary>
-        /// Ability to read blobs from the storage service
-        /// </summary>
-        ReadBlobs,
-
-		/// <summary>
-		/// Ability to write blobs to the storage service
-		/// </summary>
-		WriteBlobs,
-
-		/// <summary>
-		/// Ability to read refs from the storage service
-		/// </summary>
-		ReadRefs,
+	/// <summary>
+	/// Type converter for NamespaceId to and from JSON
+	/// </summary>
+	sealed class AclActionJsonConverter : JsonConverter<AclAction>
+	{
+		/// <inheritdoc/>
+		public override AclAction Read(ref Utf8JsonReader reader, Type typeToConvert, JsonSerializerOptions options) => new AclAction(reader.GetString()!);
 
-		/// <summary>
-		/// Ability to write refs to the storage service
-		/// </summary>
-		WriteRefs, 
+		/// <inheritdoc/>
+		public override void Write(Utf8JsonWriter writer, AclAction value, JsonSerializerOptions options) => writer.WriteStringValue(value.Name);
+	}
 
-		/// <summary>
-		/// Ability to delete refs
-		/// </summary>
-		DeleteRefs,
+	/// <summary>
+	/// Type converter from strings to NamespaceId objects
+	/// </summary>
+	sealed class AclActionTypeConverter : TypeConverter
+	{
+		/// <inheritdoc/>
+		public override bool CanConvertFrom(ITypeDescriptorContext? context, Type sourceType) => sourceType == typeof(string);
 
-        #endregion
-    }
+		/// <inheritdoc/>
+		public override object ConvertFrom(ITypeDescriptorContext? context, CultureInfo? culture, object value) => new AclAction((string)value);
+	}
 }
diff --git a/Engine/Source/Programs/Horde/Horde.Server/Acls/AclConfig.cs b/Engine/Source/Programs/Horde/Horde.Server/Acls/AclConfig.cs
--- a/Engine/Source/Programs/Horde/Horde.Server/Acls/AclConfig.cs
+++ b/Engine/Source/Programs/Horde/Horde.Server/Acls/AclConfig.cs
@@ -148,5 +148,10 @@ namespace Horde.Server.Acls
 			Type = type;
 			Value = value;
 		}
+
+		/// <summary>
+		/// Converts this object to a regular <see cref="Claim"/> object.
+		/// </summary>
+		public Claim ToClaim() => new Claim(Type, Value);
 	}
 }
diff --git a/Engine/Source/Programs/Horde/Horde.Server/Acls/AclService.cs b/Engine/Source/Programs/Horde/Horde.Server/Acls/AclService.cs
--- a/Engine/Source/Programs/Horde/Horde.Server/Acls/AclService.cs
+++ b/Engine/Source/Programs/Horde/Horde.Server/Acls/AclService.cs
@@ -98,6 +98,11 @@ namespace Horde.Server.Acls
 
     internal static class ClaimExtensions
 	{
+		public static bool HasAdminClaim(this ClaimsPrincipal user)
+		{
+			return user.HasClaim(HordeClaims.AdminClaim.Type, HordeClaims.AdminClaim.Value);
+		}
+
 		public static bool HasLeaseClaim(this ClaimsPrincipal user, LeaseId leaseId)
 		{
 			return user.HasClaim(HordeClaimTypes.Lease, leaseId.ToString());
diff --git a/Engine/Source/Programs/Horde/Horde.Server/Acls/IAclScope.cs b/Engine/Source/Programs/Horde/Horde.Server/Acls/IAclScope.cs
--- a/Engine/Source/Programs/Horde/Horde.Server/Acls/IAclScope.cs
+++ b/Engine/Source/Programs/Horde/Horde.Server/Acls/IAclScope.cs
@@ -105,6 +105,11 @@ namespace Horde.Server.Acls
 		/// <returns>True if the given principal is allowed to perform a particular action</returns>
 		public static bool Authorize(this IAclScope scope, AclAction action, ClaimsPrincipal principal)
 		{
+			if (principal.HasAdminClaim())
+			{
+				return true;
+			}
+
 			for (IAclScope? next = scope; next != null; next = next.ParentScope)
 			{
 				if (next.Acl != null)
diff --git a/Engine/Source/Programs/Horde/Horde.Server/Agents/AgentAclAction.cs b/Engine/Source/Programs/Horde/Horde.Server/Agents/AgentAclAction.cs
new file mode 100644
--- /dev/null
+++ b/Engine/Source/Programs/Horde/Horde.Server/Agents/AgentAclAction.cs
@@ -0,0 +1,37 @@
+// Copyright Epic Games, Inc. All Rights Reserved.
+
+using Horde.Server.Acls;
+
+namespace Horde.Server.Agents
+{
+	/// <summary>
+	/// ACL actions which apply to agents
+	/// </summary>
+	public static class AgentAclAction
+	{
+		/// <summary>
+		/// Ability to create an agent. This may be done explicitly, or granted to agents to allow them to self-register.
+		/// </summary>
+		public static AclAction CreateAgent { get; } = new AclAction("CreateAgent");
+
+		/// <summary>
+		/// Update an agent's name, pools, etc...
+		/// </summary>
+		public static AclAction UpdateAgent { get; } = new AclAction("UpdateAgent");
+
+		/// <summary>
+		/// Soft-delete an agent
+		/// </summary>
+		public static AclAction DeleteAgent { get; } = new AclAction("DeleteAgent");
+
+		/// <summary>
+		/// View an agent
+		/// </summary>
+		public static AclAction ViewAgent { get; } = new AclAction("ViewAgent");
+
+		/// <summary>
+		/// List the available agents
+		/// </summary>
+		public static AclAction ListAgents { get; } = new AclAction("ListAgents");
+	}
+}
diff --git a/Engine/Source/Programs/Horde/Horde.Server/Agents/AgentsController.cs b/Engine/Source/Programs/Horde/Horde.Server/Agents/AgentsController.cs
--- a/Engine/Source/Programs/Horde/Horde.Server/Agents/AgentsController.cs
+++ b/Engine/Source/Programs/Horde/Horde.Server/Agents/AgentsController.cs
@@ -48,9 +48,9 @@ namespace Horde.Server.Agents
 		[Route("/api/v1/agents")]
 		public async Task<ActionResult<CreateAgentResponse>> CreateAgentAsync([FromBody] CreateAgentRequest request)
 		{
-			if(!_globalConfig.Value.Authorize(AclAction.CreateAgent, User))
+			if(!_globalConfig.Value.Authorize(AgentAclAction.CreateAgent, User))
 			{
-				return Forbid(AclAction.CreateAgent);
+				return Forbid(AgentAclAction.CreateAgent);
 			}
 
 			IAgent agent = await _agentService.CreateAgentAsync(request.Name, request.Enabled, request.Pools?.ConvertAll(x => new PoolId(x)));
@@ -72,9 +72,9 @@ namespace Horde.Server.Agents
 		[ProducesResponseType(typeof(List<GetAgentResponse>), 200)]
 		public async Task<ActionResult<List<object>>> FindAgentsAsync([FromQuery] PoolId? poolId = null, [FromQuery] int? index = null, [FromQuery] int? count = null, [FromQuery] DateTimeOffset? modifiedAfter = null, [FromQuery] PropertyFilter? filter = null)
 		{
-			if (!_globalConfig.Value.Authorize(AclAction.ListAgents, User))
+			if (!_globalConfig.Value.Authorize(AgentAclAction.ListAgents, User))
 			{
-				return Forbid(AclAction.ListAgents);
+				return Forbid(AgentAclAction.ListAgents);
 			}
 
 			List<IAgent> agents = await _agentService.FindAgentsAsync(poolId, modifiedAfter?.UtcDateTime, index, count);
@@ -99,9 +99,9 @@ namespace Horde.Server.Agents
 		[ProducesResponseType(typeof(GetAgentResponse), 200)]
 		public async Task<ActionResult<object>> GetAgentAsync(AgentId agentId, [FromQuery] PropertyFilter? filter = null)
 		{
-			if (!_globalConfig.Value.Authorize(AclAction.ViewAgent, User))
+			if (!_globalConfig.Value.Authorize(AgentAclAction.ViewAgent, User))
 			{
-				return Forbid(AclAction.ViewAgent, agentId);
+				return Forbid(AgentAclAction.ViewAgent, agentId);
 			}
 
 			IAgent? agent = await _agentService.GetAgentAsync(agentId);
@@ -119,7 +119,7 @@ namespace Horde.Server.Agents
 		async ValueTask<object> GetAgentResponseAsync(IAgent agent, PropertyFilter? filter = null)
 		{
 			double? rate = null;
-			if (_globalConfig.Value.Authorize(AclAction.ViewCosts, User))
+			if (_globalConfig.Value.Authorize(AdminAclAction.ViewCosts, User))
 			{
 				rate = await _agentService.GetRateAsync(agent.Id);
 			}
@@ -144,9 +144,9 @@ namespace Horde.Server.Agents
 		[Route("/api/v1/agents/{agentId}")]
 		public async Task<ActionResult> UpdateAgentAsync(AgentId agentId, [FromBody] UpdateAgentRequest update)
 		{
-			if (!_globalConfig.Value.Authorize(AclAction.UpdateAgent, User))
+			if (!_globalConfig.Value.Authorize(AgentAclAction.UpdateAgent, User))
 			{
-				return Forbid(AclAction.UpdateAgent, agentId);
+				return Forbid(AgentAclAction.UpdateAgent, agentId);
 			}
 
 			List<PoolId>? updatePools = update.Pools?.ConvertAll(x => new PoolId(x));
@@ -214,9 +214,9 @@ namespace Horde.Server.Agents
 		[Route("/api/v1/agents/{agentId}")]
 		public async Task<ActionResult> DeleteAgentAsync(AgentId agentId)
 		{
-			if (!_globalConfig.Value.Authorize(AclAction.DeleteAgent, User))
+			if (!_globalConfig.Value.Authorize(AgentAclAction.DeleteAgent, User))
 			{
-				return Forbid(AclAction.DeleteAgent, agentId);
+				return Forbid(AgentAclAction.DeleteAgent, agentId);
 			}
 
 			IAgent? agent = await _agentService.GetAgentAsync(agentId);
@@ -261,9 +261,9 @@ namespace Horde.Server.Agents
 		[Route("/api/v1/agents/{agentId}/sessions")]
 		public async Task<ActionResult<List<GetAgentSessionResponse>>> FindSessionsAsync(AgentId agentId, [FromQuery] DateTimeOffset? startTime, [FromQuery] DateTimeOffset? finishTime, [FromQuery] int index = 0, [FromQuery] int count = 50)
 		{
-			if (!_globalConfig.Value.Authorize(AclAction.ViewSession, User))
+			if (!_globalConfig.Value.Authorize(SessionAclAction.ViewSession, User))
 			{
-				return Forbid(AclAction.ViewSession, agentId);
+				return Forbid(SessionAclAction.ViewSession, agentId);
 			}
 
 			IAgent? agent = await _agentService.GetAgentAsync(agentId);
@@ -286,9 +286,9 @@ namespace Horde.Server.Agents
 		[Route("/api/v1/agents/{agentId}/sessions/{sessionId}")]
 		public async Task<ActionResult<GetAgentSessionResponse>> GetSessionAsync(AgentId agentId, SessionId sessionId)
 		{
-			if (!_globalConfig.Value.Authorize(AclAction.ViewSession, User))
+			if (!_globalConfig.Value.Authorize(SessionAclAction.ViewSession, User))
 			{
-				return Forbid(AclAction.ViewSession, agentId);
+				return Forbid(SessionAclAction.ViewSession, agentId);
 			}
 
 			IAgent? agent = await _agentService.GetAgentAsync(agentId);
@@ -322,9 +322,9 @@ namespace Horde.Server.Agents
 		[ProducesResponseType(200, Type = typeof(List<GetAgentLeaseResponse>))]
 		public async Task<ActionResult<List<object>>> FindLeasesAsync(AgentId agentId, [FromQuery] SessionId? sessionId, [FromQuery] DateTimeOffset? startTime, [FromQuery] DateTimeOffset? finishTime, [FromQuery] int index = 0, [FromQuery] int count = 1000, [FromQuery] PropertyFilter? filter = null)
 		{
-			if (!_globalConfig.Value.Authorize(AclAction.ViewLeases, User))
+			if (!_globalConfig.Value.Authorize(LeaseAclAction.ViewLeases, User))
 			{
-				return Forbid(AclAction.ViewLeases, agentId);
+				return Forbid(LeaseAclAction.ViewLeases, agentId);
 			}
 
 			IAgent? agent = await _agentService.GetAgentAsync(agentId);
@@ -336,7 +336,7 @@ namespace Horde.Server.Agents
 			List<ILease> leases = await _agentService.FindLeasesAsync(agentId, sessionId, startTime?.UtcDateTime, finishTime?.UtcDateTime, index, count);
 
 			double? agentRate = null;
-			if (_globalConfig.Value.Authorize(AclAction.ViewCosts, User))
+			if (_globalConfig.Value.Authorize(AdminAclAction.ViewCosts, User))
 			{
 				agentRate = await _agentService.GetRateAsync(agentId);
 			}
@@ -361,9 +361,9 @@ namespace Horde.Server.Agents
 		[Route("/api/v1/agents/{agentId}/leases/{leaseId}")]
 		public async Task<ActionResult<GetAgentLeaseResponse>> GetLeaseAsync(AgentId agentId, LeaseId leaseId)
 		{
-			if (!_globalConfig.Value.Authorize(AclAction.ViewLeases, User))
+			if (!_globalConfig.Value.Authorize(LeaseAclAction.ViewLeases, User))
 			{
-				return Forbid(AclAction.ViewLeases, agentId);
+				return Forbid(LeaseAclAction.ViewLeases, agentId);
 			}
 
 			IAgent? agent = await _agentService.GetAgentAsync(agentId);
@@ -379,7 +379,7 @@ namespace Horde.Server.Agents
 			}
 
 			double? agentRate = null;
-			if (_globalConfig.Value.Authorize(AclAction.ViewCosts, User))
+			if (_globalConfig.Value.Authorize(AdminAclAction.ViewCosts, User))
 			{
 				agentRate = await _agentService.GetRateAsync(agentId);
 			}
diff --git a/Engine/Source/Programs/Horde/Horde.Server/Agents/Leases/LeaseAclAction.cs b/Engine/Source/Programs/Horde/Horde.Server/Agents/Leases/LeaseAclAction.cs
new file mode 100644
--- /dev/null
+++ b/Engine/Source/Programs/Horde/Horde.Server/Agents/Leases/LeaseAclAction.cs
@@ -0,0 +1,17 @@
+// Copyright Epic Games, Inc. All Rights Reserved.
+
+using Horde.Server.Acls;
+
+namespace Horde.Server.Agents.Leases
+{
+	/// <summary>
+	/// ACL actions which apply to leases
+	/// </summary>
+	public static class LeaseAclAction
+	{
+		/// <summary>
+		/// View all the leases that an agent has worked on
+		/// </summary>
+		public static AclAction ViewLeases { get; } = new AclAction("ViewLeases");
+	}
+}
diff --git a/Engine/Source/Programs/Horde/Horde.Server/Agents/Leases/LeasesController.cs b/Engine/Source/Programs/Horde/Horde.Server/Agents/Leases/LeasesController.cs
--- a/Engine/Source/Programs/Horde/Horde.Server/Agents/Leases/LeasesController.cs
+++ b/Engine/Source/Programs/Horde/Horde.Server/Agents/Leases/LeasesController.cs
@@ -63,12 +63,12 @@ namespace Horde.Server.Agents.Leases
 			[FromQuery] int count = 1000,
 			[FromQuery] PropertyFilter? filter = null)
 		{
-			if (!_globalConfig.Value.Authorize(AclAction.ViewLeases, User))
+			if (!_globalConfig.Value.Authorize(LeaseAclAction.ViewLeases, User))
 			{
-				return Forbid(AclAction.ViewLeases);
+				return Forbid(LeaseAclAction.ViewLeases);
 			}
 
-			bool includeCosts = _globalConfig.Value.Authorize(AclAction.ViewCosts, User);
+			bool includeCosts = _globalConfig.Value.Authorize(AdminAclAction.ViewCosts, User);
 
 			List<ILease> leases;
 			if (minFinishTime == null && maxFinishTime == null)
@@ -112,9 +112,9 @@ namespace Horde.Server.Agents.Leases
 		[Route("/api/v1/leases/{leaseId}")]
 		public async Task<ActionResult<GetAgentLeaseResponse>> GetLeaseAsync(LeaseId leaseId)
 		{
-			if (!_globalConfig.Value.Authorize(AclAction.ViewLeases, User))
+			if (!_globalConfig.Value.Authorize(LeaseAclAction.ViewLeases, User))
 			{
-				return Forbid(AclAction.ViewLeases);
+				return Forbid(LeaseAclAction.ViewLeases);
 			}
 
 			ILease? lease = await _agentService.GetLeaseAsync(leaseId);
@@ -130,7 +130,7 @@ namespace Horde.Server.Agents.Leases
 			}
 
 			double? agentRate = null;
-			if (_globalConfig.Value.Authorize(AclAction.ViewCosts, User))
+			if (_globalConfig.Value.Authorize(AdminAclAction.ViewCosts, User))
 			{
 				agentRate = await _agentService.GetRateAsync(agent.Id);
 			}
@@ -149,9 +149,9 @@ namespace Horde.Server.Agents.Leases
 		[Route("/api/v1/leases/{leaseId}")]
 		public async Task<ActionResult> UpdateLeaseAsync(LeaseId leaseId, [FromBody] UpdateLeaseRequest request)
 		{
-			if (!_globalConfig.Value.Authorize(AclAction.AdminWrite, User))
+			if (!_globalConfig.Value.Authorize(AdminAclAction.AdminWrite, User))
 			{
-				return Forbid(AclAction.AdminWrite);
+				return Forbid(AdminAclAction.AdminWrite);
 			}
 
 			// only update supported right now is abort
diff --git a/Engine/Source/Programs/Horde/Horde.Server/Agents/Pools/PoolAclAction.cs b/Engine/Source/Programs/Horde/Horde.Server/Agents/Pools/PoolAclAction.cs
new file mode 100644
--- /dev/null
+++ b/Engine/Source/Programs/Horde/Horde.Server/Agents/Pools/PoolAclAction.cs
@@ -0,0 +1,37 @@
+// Copyright Epic Games, Inc. All Rights Reserved.
+
+using Horde.Server.Acls;
+
+namespace Horde.Server.Agents.Pools
+{
+	/// <summary>
+	/// ACL actions for pools
+	/// </summary>
+	public static class PoolAclAction
+	{
+		/// <summary>
+		/// Create a global pool of agents
+		/// </summary>
+		public static AclAction CreatePool = new AclAction("CreatePool");
+
+		/// <summary>
+		/// Modify an agent pool
+		/// </summary>
+		public static AclAction UpdatePool = new AclAction("UpdatePool");
+
+		/// <summary>
+		/// Delete an agent pool
+		/// </summary>
+		public static AclAction DeletePool = new AclAction("DeletePool");
+
+		/// <summary>
+		/// Ability to view a pool
+		/// </summary>
+		public static AclAction ViewPool = new AclAction("ViewPool");
+
+		/// <summary>
+		/// View all the available agent pools
+		/// </summary>
+		public static AclAction ListPools = new AclAction("ListPools");
+	}
+}
diff --git a/Engine/Source/Programs/Horde/Horde.Server/Agents/Pools/PoolsController.cs b/Engine/Source/Programs/Horde/Horde.Server/Agents/Pools/PoolsController.cs
--- a/Engine/Source/Programs/Horde/Horde.Server/Agents/Pools/PoolsController.cs
+++ b/Engine/Source/Programs/Horde/Horde.Server/Agents/Pools/PoolsController.cs
@@ -43,9 +43,9 @@ namespace Horde.Server.Agents.Pools
 		[Route("/api/v1/pools")]
 		public async Task<ActionResult<CreatePoolResponse>> CreatePoolAsync([FromBody] CreatePoolRequest create)
 		{
-			if(!_globalConfig.Value.Authorize(AclAction.CreatePool, User))
+			if(!_globalConfig.Value.Authorize(PoolAclAction.CreatePool, User))
 			{
-				return Forbid(AclAction.CreatePool);
+				return Forbid(PoolAclAction.CreatePool);
 			}
 
 			List<PoolSizeStrategyInfo>? sizeStrategies = create.SizeStrategies?.Select(x => x.Convert()).ToList();
@@ -74,9 +74,9 @@ namespace Horde.Server.Agents.Pools
 		[ProducesResponseType(typeof(List<GetPoolResponse>), 200)]
 		public async Task<ActionResult<List<object>>> GetPoolsAsync([FromQuery] PropertyFilter? filter = null)
 		{
-			if (!_globalConfig.Value.Authorize(AclAction.ListPools, User))
+			if (!_globalConfig.Value.Authorize(PoolAclAction.ListPools, User))
 			{
-				return Forbid(AclAction.ListPools);
+				return Forbid(PoolAclAction.ListPools);
 			}
 
 			List<IPool> pools = await _poolService.GetPoolsAsync();
@@ -100,9 +100,9 @@ namespace Horde.Server.Agents.Pools
 		[ProducesResponseType(typeof(GetPoolResponse), 200)]
 		public async Task<ActionResult<object>> GetPoolAsync(string poolId, [FromQuery] PropertyFilter? filter = null)
 		{
-			if (!_globalConfig.Value.Authorize(AclAction.ViewPool, User))
+			if (!_globalConfig.Value.Authorize(PoolAclAction.ViewPool, User))
 			{
-				return Forbid(AclAction.ViewPool);
+				return Forbid(PoolAclAction.ViewPool);
 			}
 
 			PoolId poolIdValue = new PoolId(poolId);
@@ -126,9 +126,9 @@ namespace Horde.Server.Agents.Pools
 		[Route("/api/v1/pools/{poolId}")]
 		public async Task<ActionResult> UpdatePoolAsync(string poolId, [FromBody] UpdatePoolRequest update)
 		{
-			if (!_globalConfig.Value.Authorize(AclAction.UpdatePool, User))
+			if (!_globalConfig.Value.Authorize(PoolAclAction.UpdatePool, User))
 			{
-				return Forbid(AclAction.UpdatePool);
+				return Forbid(PoolAclAction.UpdatePool);
 			}
 
 			PoolId poolIdValue = new PoolId(poolId);
@@ -160,9 +160,9 @@ namespace Horde.Server.Agents.Pools
 		[Route("/api/v1/pools/{poolId}")]
 		public async Task<ActionResult> DeletePoolAsync(string poolId)
 		{
-			if (!_globalConfig.Value.Authorize(AclAction.DeletePool, User))
+			if (!_globalConfig.Value.Authorize(PoolAclAction.DeletePool, User))
 			{
-				return Forbid(AclAction.DeletePool);
+				return Forbid(PoolAclAction.DeletePool);
 			}
 
 			PoolId poolIdValue = new PoolId(poolId);
@@ -182,9 +182,9 @@ namespace Horde.Server.Agents.Pools
 		[Route("/api/v1/pools")]
 		public async Task<ActionResult> UpdatePoolAsync([FromBody] List<BatchUpdatePoolRequest> batchUpdates)
 		{
-			if (!_globalConfig.Value.Authorize(AclAction.UpdatePool, User))
+			if (!_globalConfig.Value.Authorize(PoolAclAction.UpdatePool, User))
 			{
-				return Forbid(AclAction.UpdatePool);
+				return Forbid(PoolAclAction.UpdatePool);
 			}
 
 			foreach (BatchUpdatePoolRequest update in batchUpdates)
diff --git a/Engine/Source/Programs/Horde/Horde.Server/Agents/Sessions/SessionAclAction.cs b/Engine/Source/Programs/Horde/Horde.Server/Agents/Sessions/SessionAclAction.cs
new file mode 100644
--- /dev/null
+++ b/Engine/Source/Programs/Horde/Horde.Server/Agents/Sessions/SessionAclAction.cs
@@ -0,0 +1,22 @@
+// Copyright Epic Games, Inc. All Rights Reserved.
+
+using Horde.Server.Acls;
+
+namespace Horde.Server.Agents.Sessions
+{
+	/// <summary>
+	/// Actions valid for agent sessions
+	/// </summary>
+	public class SessionAclAction
+	{
+		/// <summary>
+		/// Granted to agents to call CreateSession, which returns a bearer token identifying themselves valid to call UpdateSesssion via gRPC.
+		/// </summary>
+		public static AclAction CreateSession = new AclAction("CreateSession");
+
+		/// <summary>
+		/// Allows viewing information about an agent session
+		/// </summary>
+		public static AclAction ViewSession = new AclAction("ViewSession");
+	}
+}
diff --git a/Engine/Source/Programs/Horde/Horde.Server/Agents/Software/AgentSoftwareAclAction.cs b/Engine/Source/Programs/Horde/Horde.Server/Agents/Software/AgentSoftwareAclAction.cs
new file mode 100644
--- /dev/null
+++ b/Engine/Source/Programs/Horde/Horde.Server/Agents/Software/AgentSoftwareAclAction.cs
@@ -0,0 +1,27 @@
+// Copyright Epic Games, Inc. All Rights Reserved.
+
+using Horde.Server.Acls;
+
+namespace Horde.Server.Agents.Software
+{
+	/// <summary>
+	/// ACL actions for agent software
+	/// </summary>
+	public static class AgentSoftwareAclAction
+	{
+		/// <summary>
+		/// Ability to upload new versions of the agent software
+		/// </summary>
+		public static AclAction UploadSoftware { get; } = new AclAction("UploadSoftware");
+
+		/// <summary>
+		/// Ability to download the agent software
+		/// </summary>
+		public static AclAction DownloadSoftware { get; } = new AclAction("DownloadSoftware");
+
+		/// <summary>
+		/// Ability to delete agent software
+		/// </summary>
+		public static AclAction DeleteSoftware { get; } = new AclAction("DeleteSoftware");
+	}
+}
diff --git a/Engine/Source/Programs/Horde/Horde.Server/Agents/Software/AgentSoftwareController.cs b/Engine/Source/Programs/Horde/Horde.Server/Agents/Software/AgentSoftwareController.cs
--- a/Engine/Source/Programs/Horde/Horde.Server/Agents/Software/AgentSoftwareController.cs
+++ b/Engine/Source/Programs/Horde/Horde.Server/Agents/Software/AgentSoftwareController.cs
@@ -59,7 +59,7 @@ namespace Horde.Server.Agents.Software
 		[ProducesResponseType(typeof(GetAgentSoftwareChannelResponse), 200)]
 		public async Task<ActionResult<object>> FindSoftwareAsync()
 		{
-			if (!_globalConfig.Value.Authorize(AclAction.DownloadSoftware, User))
+			if (!_globalConfig.Value.Authorize(AgentSoftwareAclAction.DownloadSoftware, User))
 			{
 				return Forbid();
 			}
@@ -89,7 +89,7 @@ namespace Horde.Server.Agents.Software
 		[Route("/api/v1/agentsoftware/default/zip")]
 		public async Task<ActionResult> GetArchiveAsync(CancellationToken cancellationToken)
 		{
-			if (!_globalConfig.Value.Authorize(AclAction.DownloadSoftware, User))
+			if (!_globalConfig.Value.Authorize(AgentSoftwareAclAction.DownloadSoftware, User))
 			{
 				return Forbid();
 			}
diff --git a/Engine/Source/Programs/Horde/Horde.Server/Artifacts/ArtifactAclAction.cs b/Engine/Source/Programs/Horde/Horde.Server/Artifacts/ArtifactAclAction.cs
new file mode 100644
--- /dev/null
+++ b/Engine/Source/Programs/Horde/Horde.Server/Artifacts/ArtifactAclAction.cs
@@ -0,0 +1,32 @@
+// Copyright Epic Games, Inc. All Rights Reserved.
+
+using Horde.Server.Acls;
+
+namespace Horde.Server.Artifacts
+{
+	/// <summary>
+	/// ACL actions for artifacts
+	/// </summary>
+	public static class ArtifactAclAction
+	{
+		/// <summary>
+		/// Permission to read from an artifact
+		/// </summary>
+		public static AclAction ReadArtifact { get; } = new AclAction("ReadArtifact");
+
+		/// <summary>
+		/// Permission to write to an artifact
+		/// </summary>
+		public static AclAction WriteArtifact { get; } = new AclAction("WriteArtifact");
+
+		/// <summary>
+		/// Ability to create an artifact. Typically just for debugging; agents have this access for a particular session.
+		/// </summary>
+		public static AclAction UploadArtifact { get; } = new AclAction("UploadArtifact");
+
+		/// <summary>
+		/// Ability to download an artifact
+		/// </summary>
+		public static AclAction DownloadArtifact { get; } = new AclAction("DownloadArtifact");
+	}
+}
diff --git a/Engine/Source/Programs/Horde/Horde.Server/Artifacts/ArtifactsController.cs b/Engine/Source/Programs/Horde/Horde.Server/Artifacts/ArtifactsController.cs
--- a/Engine/Source/Programs/Horde/Horde.Server/Artifacts/ArtifactsController.cs
+++ b/Engine/Source/Programs/Horde/Horde.Server/Artifacts/ArtifactsController.cs
@@ -154,9 +154,9 @@ namespace Horde.Server.Artifacts
 			{
 				return NotFound(id);
 			}
-			if (!_globalConfig.Authorize(artifact.AclScope, AclAction.ReadArtifact, User))
+			if (!_globalConfig.Authorize(artifact.AclScope, ArtifactAclAction.ReadArtifact, User))
 			{
-				return Forbid(AclAction.ReadArtifact, artifact.AclScope);
+				return Forbid(ArtifactAclAction.ReadArtifact, artifact.AclScope);
 			}
 
 			return PropertyFilter.Apply(new GetArtifactResponse(artifact), filter);
@@ -180,9 +180,9 @@ namespace Horde.Server.Artifacts
 			{
 				return NotFound(id);
 			}
-			if (!_globalConfig.Authorize(artifact.AclScope, AclAction.ReadArtifact, User))
+			if (!_globalConfig.Authorize(artifact.AclScope, ArtifactAclAction.ReadArtifact, User))
 			{
-				return Forbid(AclAction.ReadArtifact, artifact.AclScope);
+				return Forbid(ArtifactAclAction.ReadArtifact, artifact.AclScope);
 			}
 			if (!locator.BlobId.WithinFolder(artifact.RefName.Text))
 			{
@@ -207,9 +207,9 @@ namespace Horde.Server.Artifacts
 			{
 				return NotFound(id);
 			}
-			if (!_globalConfig.Authorize(artifact.AclScope, AclAction.ReadArtifact, User))
+			if (!_globalConfig.Authorize(artifact.AclScope, ArtifactAclAction.ReadArtifact, User))
 			{
-				return Forbid(AclAction.ReadArtifact, artifact.AclScope);
+				return Forbid(ArtifactAclAction.ReadArtifact, artifact.AclScope);
 			}
 
 			return StorageController.ReadRefInternalAsync(_storageService, artifact.NamespaceId, artifact.RefName, cancellationToken);
@@ -233,9 +233,9 @@ namespace Horde.Server.Artifacts
 			{
 				return NotFound(id);
 			}
-			if (!_globalConfig.Authorize(artifact.AclScope, AclAction.ReadArtifact, User))
+			if (!_globalConfig.Authorize(artifact.AclScope, ArtifactAclAction.ReadArtifact, User))
 			{
-				return Forbid(AclAction.ReadArtifact, artifact.AclScope);
+				return Forbid(ArtifactAclAction.ReadArtifact, artifact.AclScope);
 			}
 
 			TreeReader reader = await _storageService.GetReaderAsync(artifact.NamespaceId, cancellationToken);
@@ -328,9 +328,9 @@ namespace Horde.Server.Artifacts
 			{
 				return NotFound(id);
 			}
-			if (!_globalConfig.Authorize(artifact.AclScope, AclAction.ReadArtifact, User))
+			if (!_globalConfig.Authorize(artifact.AclScope, ArtifactAclAction.ReadArtifact, User))
 			{
-				return Forbid(AclAction.ReadArtifact, artifact.AclScope);
+				return Forbid(ArtifactAclAction.ReadArtifact, artifact.AclScope);
 			}
 
 			TreeReader reader = await _storageService.GetReaderAsync(artifact.NamespaceId, cancellationToken);
@@ -395,9 +395,9 @@ namespace Horde.Server.Artifacts
 			{
 				return NotFound(id);
 			}
-			if (!_globalConfig.Authorize(artifact.AclScope, AclAction.ReadArtifact, User))
+			if (!_globalConfig.Authorize(artifact.AclScope, ArtifactAclAction.ReadArtifact, User))
 			{
-				return Forbid(AclAction.ReadArtifact, artifact.AclScope);
+				return Forbid(ArtifactAclAction.ReadArtifact, artifact.AclScope);
 			}
 
 			FileFilter? filter = null;
@@ -434,7 +434,7 @@ namespace Horde.Server.Artifacts
 			FindArtifactsResponse response = new FindArtifactsResponse();
 			await foreach (IArtifact artifact in _artifactCollection.FindAsync(ids, keys, HttpContext.RequestAborted))
 			{
-				if (_globalConfig.Authorize(artifact.AclScope, AclAction.ReadArtifact, User))
+				if (_globalConfig.Authorize(artifact.AclScope, ArtifactAclAction.ReadArtifact, User))
 				{
 					response.Artifacts.Add(new GetArtifactResponse(artifact));
 				}
diff --git a/Engine/Source/Programs/Horde/Horde.Server/Authentication/AnonymousHandler.cs b/Engine/Source/Programs/Horde/Horde.Server/Authentication/AnonymousHandler.cs
--- a/Engine/Source/Programs/Horde/Horde.Server/Authentication/AnonymousHandler.cs
+++ b/Engine/Source/Programs/Horde/Horde.Server/Authentication/AnonymousHandler.cs
@@ -15,8 +15,6 @@ namespace Horde.Server.Authentication
 {
 	class AnonymousAuthenticationOptions : AuthenticationSchemeOptions
 	{
-		public string? AdminClaimType { get; set; }
-		public string? AdminClaimValue { get; set; }
 	}
 
 	class AnonymousAuthenticationHandler : AuthenticationHandler<AnonymousAuthenticationOptions>
@@ -36,11 +34,7 @@ namespace Horde.Server.Authentication
 			List<Claim> claims = new List<Claim>();
 			claims.Add(new Claim(ClaimTypes.Name, AuthenticationScheme));
 			claims.Add(new Claim(HordeClaimTypes.UserId, UserId.Anonymous.ToString()));
-
-			if (Options.AdminClaimType != null && Options.AdminClaimValue != null)
-			{
-				claims.Add(new Claim(Options.AdminClaimType, Options.AdminClaimValue));
-			}
+			claims.Add(new Claim(HordeClaims.AdminClaim.Type, HordeClaims.AdminClaim.Value));
 
 			ClaimsIdentity identity = new ClaimsIdentity(claims, Scheme.Name);
 
diff --git a/Engine/Source/Programs/Horde/Horde.Server/Compute/ComputeAclAction.cs b/Engine/Source/Programs/Horde/Horde.Server/Compute/ComputeAclAction.cs
new file mode 100644
--- /dev/null
+++ b/Engine/Source/Programs/Horde/Horde.Server/Compute/ComputeAclAction.cs
@@ -0,0 +1,17 @@
+// Copyright Epic Games, Inc. All Rights Reserved.
+
+using Horde.Server.Acls;
+
+namespace Horde.Server.Compute
+{
+	/// <summary>
+	/// ACL actions for compute service
+	/// </summary>
+	public static class ComputeAclAction
+	{
+		/// <summary>
+		/// User can add tasks to the compute cluster
+		/// </summary>
+		public static AclAction AddComputeTasks { get; } = new AclAction("AddComputeTasks");
+	}
+}
diff --git a/Engine/Source/Programs/Horde/Horde.Server/Compute/ComputeController.cs b/Engine/Source/Programs/Horde/Horde.Server/Compute/ComputeController.cs
--- a/Engine/Source/Programs/Horde/Horde.Server/Compute/ComputeController.cs
+++ b/Engine/Source/Programs/Horde/Horde.Server/Compute/ComputeController.cs
@@ -98,9 +98,9 @@ namespace Horde.Server.Compute
 			{
 				return NotFound(clusterId);
 			}
-			if(!clusterConfig.Authorize(AclAction.AddComputeTasks, User))
+			if(!clusterConfig.Authorize(ComputeAclAction.AddComputeTasks, User))
 			{
-				return Forbid(AclAction.AddComputeTasks, clusterId);
+				return Forbid(ComputeAclAction.AddComputeTasks, clusterId);
 			}
 
 			Requirements requirements = request.Requirements ?? new Requirements();
diff --git a/Engine/Source/Programs/Horde/Horde.Server/Dashboard/DashboardController.cs b/Engine/Source/Programs/Horde/Horde.Server/Dashboard/DashboardController.cs
--- a/Engine/Source/Programs/Horde/Horde.Server/Dashboard/DashboardController.cs
+++ b/Engine/Source/Programs/Horde/Horde.Server/Dashboard/DashboardController.cs
@@ -152,7 +152,7 @@ namespace Horde.Server.Dashboard
 		[Route("/api/v1/dashboard/preview")]
 		public async Task<ActionResult<GetDashboardPreviewResponse>> CreateDashbordPreview([FromBody] CreateDashboardPreviewRequest request)
 		{
-			if (!_globalConfig.Value.Authorize(AclAction.AdminWrite, User))
+			if (!_globalConfig.Value.Authorize(AdminAclAction.AdminWrite, User))
 			{
 				return Forbid();
 			}
@@ -182,7 +182,7 @@ namespace Horde.Server.Dashboard
 		[Route("/api/v1/dashboard/preview")]
 		public async Task<ActionResult<GetDashboardPreviewResponse>> UpdateDashbordPreview([FromBody] UpdateDashboardPreviewRequest request)
 		{
-			if (!_globalConfig.Value.Authorize(AclAction.AdminWrite, User))
+			if (!_globalConfig.Value.Authorize(AdminAclAction.AdminWrite, User))
 			{
 				return Forbid();
 			}
diff --git a/Engine/Source/Programs/Horde/Horde.Server/Devices/DeviceAclAction.cs b/Engine/Source/Programs/Horde/Horde.Server/Devices/DeviceAclAction.cs
new file mode 100644
--- /dev/null
+++ b/Engine/Source/Programs/Horde/Horde.Server/Devices/DeviceAclAction.cs
@@ -0,0 +1,22 @@
+// Copyright Epic Games, Inc. All Rights Reserved.
+
+using Horde.Server.Acls;
+
+namespace Horde.Server.Devices
+{
+	/// <summary>
+	/// ACL actions which apply to devices
+	/// </summary>
+	public static class DeviceAclAction
+	{
+		/// <summary>
+		/// Ability to read devices
+		/// </summary>
+		public static AclAction DeviceRead { get; } = new AclAction("DeviceRead");
+
+		/// <summary>
+		/// Ability to write devices
+		/// </summary>
+		public static AclAction DeviceWrite { get; } = new AclAction("DeviceWrite");
+	}
+}
diff --git a/Engine/Source/Programs/Horde/Horde.Server/Devices/DeviceService.cs b/Engine/Source/Programs/Horde/Horde.Server/Devices/DeviceService.cs
--- a/Engine/Source/Programs/Horde/Horde.Server/Devices/DeviceService.cs
+++ b/Engine/Source/Programs/Horde/Horde.Server/Devices/DeviceService.cs
@@ -558,7 +558,7 @@ namespace Horde.Server.Devices
 		{
 			// This is deprecated and device auth should be going through GetUserPoolAuthorizationsAsync
 
-			if (action == AclAction.DeviceRead)
+			if (action == DeviceAclAction.DeviceRead)
 			{
 				return true;
 			}
@@ -568,7 +568,7 @@ namespace Horde.Server.Devices
 				return true;
 			}
 
-			if (globalConfig.Authorize(AclAction.AdminWrite, user))
+			if (globalConfig.Authorize(AdminAclAction.AdminWrite, user))
 			{
 				return true;
 			}
@@ -605,8 +605,8 @@ namespace Horde.Server.Devices
 
 				if (globalConfig.TryGetProject(projectId, out ProjectConfig? projectConfig))
 				{
-					deviceRead.Add(projectId, projectConfig.Authorize(AclAction.DeviceRead, user));
-					deviceWrite.Add(projectId, projectConfig.Authorize(AclAction.DeviceWrite, user));
+					deviceRead.Add(projectId, projectConfig.Authorize(DeviceAclAction.DeviceRead, user));
+					deviceWrite.Add(projectId, projectConfig.Authorize(DeviceAclAction.DeviceWrite, user));
 				}
 			}
 
@@ -615,7 +615,7 @@ namespace Horde.Server.Devices
 
 			if (!globalPoolAccess)
 			{
-				if (globalConfig.Authorize(AclAction.AdminWrite, user))
+				if (globalConfig.Authorize(AdminAclAction.AdminWrite, user))
 				{
 					globalPoolAccess = true;
 				}
diff --git a/Engine/Source/Programs/Horde/Horde.Server/Devices/DevicesController.cs b/Engine/Source/Programs/Horde/Horde.Server/Devices/DevicesController.cs
--- a/Engine/Source/Programs/Horde/Horde.Server/Devices/DevicesController.cs
+++ b/Engine/Source/Programs/Horde/Horde.Server/Devices/DevicesController.cs
@@ -300,7 +300,7 @@ namespace Horde.Server.Devices
 		[Route("/api/v2/devices/{deviceId}")]
 		public async Task<ActionResult> DeleteDeviceAsync(string deviceId)
 		{
-			if (!DeviceService.Authorize(AclAction.DeviceWrite, User, _globalConfig.Value))
+			if (!DeviceService.Authorize(DeviceAclAction.DeviceWrite, User, _globalConfig.Value))
 			{
 				return Forbid();
 			}
@@ -334,7 +334,7 @@ namespace Horde.Server.Devices
 		[Route("/api/v2/devices/platforms")]
 		public async Task<ActionResult<CreateDevicePlatformResponse>> CreatePlatformAsync([FromBody] CreateDevicePlatformRequest request)
 		{
-			if (!DeviceService.Authorize(AclAction.DeviceWrite, User, _globalConfig.Value))
+			if (!DeviceService.Authorize(DeviceAclAction.DeviceWrite, User, _globalConfig.Value))
 			{
 				return Forbid();
 			}
@@ -359,7 +359,7 @@ namespace Horde.Server.Devices
 		[Route("/api/v2/devices/platforms/{platformId}")]
 		public async Task<ActionResult<CreateDevicePlatformResponse>> UpdatePlatformAsync(string platformId, [FromBody] UpdateDevicePlatformRequest request)
 		{
-			if (!DeviceService.Authorize(AclAction.DeviceWrite, User, _globalConfig.Value))
+			if (!DeviceService.Authorize(DeviceAclAction.DeviceWrite, User, _globalConfig.Value))
 			{
 				return Forbid();
 			}
@@ -378,7 +378,7 @@ namespace Horde.Server.Devices
 		[ProducesResponseType(typeof(List<GetDevicePlatformResponse>), 200)]
 		public async Task<ActionResult<List<object>>> GetDevicePlatformsAsync()
 		{
-			if (!DeviceService.Authorize(AclAction.DeviceRead, User, _globalConfig.Value))
+			if (!DeviceService.Authorize(DeviceAclAction.DeviceRead, User, _globalConfig.Value))
 			{
 				return Forbid();
 			}
@@ -406,7 +406,7 @@ namespace Horde.Server.Devices
 		[Route("/api/v2/devices/pools")]
 		public async Task<ActionResult<CreateDevicePoolResponse>> CreatePoolAsync([FromBody] CreateDevicePoolRequest request)
 		{
-			if (!DeviceService.Authorize(AclAction.DeviceWrite, User, _globalConfig.Value))
+			if (!DeviceService.Authorize(DeviceAclAction.DeviceWrite, User, _globalConfig.Value))
 			{
 				return Forbid();
 			}
@@ -494,7 +494,7 @@ namespace Horde.Server.Devices
 			[FromQuery] int index = 0,
 			[FromQuery] int count = 1024)
 		{
-			if (!DeviceService.Authorize(AclAction.DeviceRead, User, _globalConfig.Value))
+			if (!DeviceService.Authorize(DeviceAclAction.DeviceRead, User, _globalConfig.Value))
 			{
 				return Forbid();
 			}
@@ -542,7 +542,7 @@ namespace Horde.Server.Devices
 		[ProducesResponseType(typeof(CreateDeviceReservationResponse), 200)]
 		public async Task<ActionResult<CreateDeviceReservationResponse>> CreateDeviceReservation([FromBody] CreateDeviceReservationRequest request)
 		{
-			if (!DeviceService.Authorize(AclAction.DeviceWrite, User, _globalConfig.Value))
+			if (!DeviceService.Authorize(DeviceAclAction.DeviceWrite, User, _globalConfig.Value))
 			{
 				return Forbid();
 			}
@@ -624,7 +624,7 @@ namespace Horde.Server.Devices
 		[ProducesResponseType(typeof(List<GetDeviceReservationResponse>), 200)]
 		public async Task<ActionResult<List<GetDeviceReservationResponse>>> GetDeviceReservations()
 		{
-			if (!DeviceService.Authorize(AclAction.DeviceRead, User, _globalConfig.Value))
+			if (!DeviceService.Authorize(DeviceAclAction.DeviceRead, User, _globalConfig.Value))
 			{
 				return Forbid();
 			}
@@ -666,7 +666,7 @@ namespace Horde.Server.Devices
 		[Route("/api/v2/devices/reservations/{reservationId}")]
 		public async Task<ActionResult> UpdateReservationAsync(string reservationId)
 		{
-			if (!DeviceService.Authorize(AclAction.DeviceWrite, User, _globalConfig.Value))
+			if (!DeviceService.Authorize(DeviceAclAction.DeviceWrite, User, _globalConfig.Value))
 			{
 				return Forbid();
 			}
@@ -691,7 +691,7 @@ namespace Horde.Server.Devices
 		[Route("/api/v2/devices/reservations/{reservationId}")]
 		public async Task<ActionResult> DeleteReservationAsync(string reservationId)
 		{
-			if (!DeviceService.Authorize(AclAction.DeviceWrite, User, _globalConfig.Value))
+			if (!DeviceService.Authorize(DeviceAclAction.DeviceWrite, User, _globalConfig.Value))
 			{
 				return Forbid();
 			}
@@ -726,7 +726,7 @@ namespace Horde.Server.Devices
 			[FromQuery] int index = 0,
 			[FromQuery] int count = 1024)
 		{
-			if (!DeviceService.Authorize(AclAction.DeviceRead, User, _globalConfig.Value))
+			if (!DeviceService.Authorize(DeviceAclAction.DeviceRead, User, _globalConfig.Value))
 			{
 				return Forbid();
 			}
diff --git a/Engine/Source/Programs/Horde/Horde.Server/Issues/IssuesController.cs b/Engine/Source/Programs/Horde/Horde.Server/Issues/IssuesController.cs
--- a/Engine/Source/Programs/Horde/Horde.Server/Issues/IssuesController.cs
+++ b/Engine/Source/Programs/Horde/Horde.Server/Issues/IssuesController.cs
@@ -92,9 +92,9 @@ namespace Horde.Server.Issues
 				{
 					return NotFound(streamId.Value);
 				}
-				if (!streamConfig.Authorize(AclAction.ViewStream, User))
+				if (!streamConfig.Authorize(StreamAclAction.ViewStream, User))
 				{
-					return Forbid(AclAction.ViewStream, streamId.Value);
+					return Forbid(StreamAclAction.ViewStream, streamId.Value);
 				}
 
 				List<IIssueSpan> spans = await _issueCollection.FindSpansAsync(null, ids, streamId.Value, minChange, maxChange, resolved);
@@ -245,9 +245,9 @@ namespace Horde.Server.Issues
 				{
 					return NotFound(jobId.Value);
 				}
-				if(!_globalConfig.Value.Authorize(job, AclAction.ViewJob, User))
+				if(!_globalConfig.Value.Authorize(job, JobAclAction.ViewJob, User))
 				{
-					return Forbid(AclAction.ViewJob, jobId.Value);
+					return Forbid(JobAclAction.ViewJob, jobId.Value);
 				}
 
 				IGraph graph = await _jobService.GetGraphAsync(job);
@@ -360,7 +360,7 @@ namespace Horde.Server.Issues
 			List<object> responses = new List<object>();
 			foreach (IGrouping<StreamId, IIssueSpan> spanGroup in issue.Spans.GroupBy(x => x.StreamId))
 			{
-				if (_globalConfig.Value.TryGetStream(spanGroup.Key, out StreamConfig? streamConfig) && streamConfig.Authorize(AclAction.ViewStream, User))
+				if (_globalConfig.Value.TryGetStream(spanGroup.Key, out StreamConfig? streamConfig) && streamConfig.Authorize(StreamAclAction.ViewStream, User))
 				{
 					HashSet<ObjectId> spanIds = new HashSet<ObjectId>(spanGroup.Select(x => x.Id));
 					List<IIssueStep> steps = issue.Steps.Where(x => spanIds.Contains(x.SpanId)).ToList();
@@ -391,9 +391,9 @@ namespace Horde.Server.Issues
 			{
 				return NotFound(streamId);
 			}
-			if (!streamConfig.Authorize(AclAction.ViewStream, User))
+			if (!streamConfig.Authorize(StreamAclAction.ViewStream, User))
 			{
-				return Forbid(AclAction.ViewStream, streamId);
+				return Forbid(StreamAclAction.ViewStream, streamId);
 			}
 
 			List<IIssueSpan> spans = details.Spans.Where(x => x.StreamId == streamId).ToList();
@@ -505,7 +505,7 @@ namespace Horde.Server.Issues
 					logFile = await _logFileService.GetLogFileAsync(logEvent.LogId, cancellationToken);
 					logFiles[logEvent.LogId] = logFile;
 				}
-				if (logFile != null && await _jobService.AuthorizeAsync(logFile.JobId, AclAction.ViewLog, User, _globalConfig.Value))
+				if (logFile != null && await _jobService.AuthorizeAsync(logFile.JobId, LogAclAction.ViewLog, User, _globalConfig.Value))
 				{
 					ILogEventData data = await _logFileService.GetEventDataAsync(logFile, logEvent.LineIndex, logEvent.LineCount, cancellationToken);
 					GetLogEventResponse response = new GetLogEventResponse(logEvent, data, issueId);
@@ -524,7 +524,7 @@ namespace Horde.Server.Issues
 		{
 			foreach (StreamId streamId in issue.Spans.Select(x => x.StreamId).Distinct())
 			{
-				if (_globalConfig.Value.TryGetStream(streamId, out StreamConfig? streamConfig) && streamConfig.Authorize(AclAction.ViewStream, User))
+				if (_globalConfig.Value.TryGetStream(streamId, out StreamConfig? streamConfig) && streamConfig.Authorize(StreamAclAction.ViewStream, User))
 				{
 					return true;
 				}
@@ -620,9 +620,9 @@ namespace Horde.Server.Issues
 			{
 				return NotFound(streamId);
 			}
-			if (!streamConfig.Authorize(AclAction.ViewStream, User))
+			if (!streamConfig.Authorize(StreamAclAction.ViewStream, User))
 			{
-				return Forbid(AclAction.ViewStream, streamId);
+				return Forbid(StreamAclAction.ViewStream, streamId);
 			}
 
 			List<GetExternalIssueResponse> response = new List<GetExternalIssueResponse>();
@@ -659,9 +659,9 @@ namespace Horde.Server.Issues
 			{
 				return NotFound(streamIdValue);
 			}
-			if (!streamConfig.Authorize(AclAction.ViewStream, User))
+			if (!streamConfig.Authorize(StreamAclAction.ViewStream, User))
 			{
-				return Forbid(AclAction.ViewStream, streamIdValue);
+				return Forbid(StreamAclAction.ViewStream, streamIdValue);
 			}
 
 
@@ -702,9 +702,9 @@ namespace Horde.Server.Issues
 			{
 				return NotFound(streamIdValue);
 			}
-			if (!streamConfig.Authorize(AclAction.ViewStream, User))
+			if (!streamConfig.Authorize(StreamAclAction.ViewStream, User))
 			{
-				return Forbid(AclAction.ViewStream, streamIdValue);
+				return Forbid(StreamAclAction.ViewStream, streamIdValue);
 			}
 
 			List<IExternalIssueProject> projects = await _externalIssueService.GetProjects(streamConfig);
diff --git a/Engine/Source/Programs/Horde/Horde.Server/Jobs/Artifacts/ArtifactsControllerV1.cs b/Engine/Source/Programs/Horde/Horde.Server/Jobs/Artifacts/ArtifactsControllerV1.cs
--- a/Engine/Source/Programs/Horde/Horde.Server/Jobs/Artifacts/ArtifactsControllerV1.cs
+++ b/Engine/Source/Programs/Horde/Horde.Server/Jobs/Artifacts/ArtifactsControllerV1.cs
@@ -10,6 +10,7 @@ using System.Net.Mime;
 using System.Security.Claims;
 using System.Threading.Tasks;
 using Horde.Server.Acls;
+using Horde.Server.Artifacts;
 using Horde.Server.Jobs.Graphs;
 using Horde.Server.Server;
 using Horde.Server.Streams;
@@ -71,7 +72,7 @@ namespace Horde.Server.Jobs.Artifacts
 			{
 				return NotFound(job.StreamId);
 			}
-			if (!streamConfig.Authorize(AclAction.UploadArtifact, User))
+			if (!streamConfig.Authorize(ArtifactAclAction.UploadArtifact, User))
 			{
 				return Forbid();
 			}
@@ -113,7 +114,7 @@ namespace Horde.Server.Jobs.Artifacts
 			{
 				return NotFound();
 			}
-			if (!await _jobService.AuthorizeAsync(artifact.JobId, AclAction.UploadArtifact, User, _globalConfig.Value))
+			if (!await _jobService.AuthorizeAsync(artifact.JobId, ArtifactAclAction.UploadArtifact, User, _globalConfig.Value))
 			{
 				return Forbid();
 			}
@@ -136,7 +137,7 @@ namespace Horde.Server.Jobs.Artifacts
 		[ProducesResponseType(typeof(List<GetJobArtifactResponse>), 200)]
 		public async Task<ActionResult<List<object>>> GetArtifacts([FromQuery] JobId jobId, [FromQuery] string? stepId = null, [FromQuery] bool code = false, [FromQuery] PropertyFilter? filter = null)
 		{
-			if (!await _jobService.AuthorizeAsync(jobId, AclAction.DownloadArtifact, User, _globalConfig.Value))
+			if (!await _jobService.AuthorizeAsync(jobId, ArtifactAclAction.DownloadArtifact, User, _globalConfig.Value))
 			{
 				return Forbid();
 			}
@@ -186,7 +187,7 @@ namespace Horde.Server.Jobs.Artifacts
 			{
 				return NotFound();
 			}
-			if (!await _jobService.AuthorizeAsync(artifact.JobId, AclAction.DownloadArtifact, User, _globalConfig.Value))
+			if (!await _jobService.AuthorizeAsync(artifact.JobId, ArtifactAclAction.DownloadArtifact, User, _globalConfig.Value))
 			{
 				return Forbid();
 			}
@@ -210,7 +211,7 @@ namespace Horde.Server.Jobs.Artifacts
 			{
 				return NotFound();
 			}
-			if (!await _jobService.AuthorizeAsync(artifact.JobId, AclAction.DownloadArtifact, User, _globalConfig.Value))
+			if (!await _jobService.AuthorizeAsync(artifact.JobId, ArtifactAclAction.DownloadArtifact, User, _globalConfig.Value))
 			{
 				return Forbid();
 			}
@@ -232,7 +233,7 @@ namespace Horde.Server.Jobs.Artifacts
 		{
 			SubResourceId stepIdValue = stepId.ToSubResourceId();
 
-			if (!await _jobService.AuthorizeAsync(jobId, AclAction.DownloadArtifact, User, _globalConfig.Value))
+			if (!await _jobService.AuthorizeAsync(jobId, ArtifactAclAction.DownloadArtifact, User, _globalConfig.Value))
 			{
 				return Forbid();
 			}
@@ -307,7 +308,7 @@ namespace Horde.Server.Jobs.Artifacts
 			{
 				return NotFound();
 			}
-			if (!_globalConfig.Value.Authorize(job, AclAction.DownloadArtifact, User))
+			if (!_globalConfig.Value.Authorize(job, ArtifactAclAction.DownloadArtifact, User))
 			{
 				return Forbid();
 			}
diff --git a/Engine/Source/Programs/Horde/Horde.Server/Jobs/JobAclAction.cs b/Engine/Source/Programs/Horde/Horde.Server/Jobs/JobAclAction.cs
new file mode 100644
--- /dev/null
+++ b/Engine/Source/Programs/Horde/Horde.Server/Jobs/JobAclAction.cs
@@ -0,0 +1,42 @@
+// Copyright Epic Games, Inc. All Rights Reserved.
+
+using Horde.Server.Acls;
+
+namespace Horde.Server.Jobs
+{
+	/// <summary>
+	/// ACL actions valid for manipulating jobs
+	/// </summary>
+	public static class JobAclAction
+	{
+		/// <summary>
+		/// Ability to start new jobs
+		/// </summary>
+		public static AclAction CreateJob = new AclAction("CreateJob");
+
+		/// <summary>
+		/// Rename a job, modify its priority, etc...
+		/// </summary>
+		public static AclAction UpdateJob = new AclAction("UpdateJob");
+
+		/// <summary>
+		/// Delete a job properties
+		/// </summary>
+		public static AclAction DeleteJob = new AclAction("DeleteJob");
+
+		/// <summary>
+		/// Allows updating a job metadata (name, changelist number, step properties, new groups, job states, etc...). Typically granted to agents. Not user facing.
+		/// </summary>
+		public static AclAction ExecuteJob = new AclAction("ExecuteJob");
+
+		/// <summary>
+		/// Ability to retry a failed job step
+		/// </summary>
+		public static AclAction RetryJobStep = new AclAction("RetryJobStep");
+
+		/// <summary>
+		/// Ability to view a job
+		/// </summary>
+		public static AclAction ViewJob = new AclAction("ViewJob");
+	}
+}
diff --git a/Engine/Source/Programs/Horde/Horde.Server/Jobs/JobRpcService.cs b/Engine/Source/Programs/Horde/Horde.Server/Jobs/JobRpcService.cs
--- a/Engine/Source/Programs/Horde/Horde.Server/Jobs/JobRpcService.cs
+++ b/Engine/Source/Programs/Horde/Horde.Server/Jobs/JobRpcService.cs
@@ -235,7 +235,7 @@ namespace Horde.Server.Jobs
 			{
 				throw new StructuredRpcException(StatusCode.NotFound, "Stream {StreamId} does not exist", request.StreamId);
 			}
-			if (!streamConfig.Authorize(AclAction.ViewStream, context.GetHttpContext().User))
+			if (!streamConfig.Authorize(StreamAclAction.ViewStream, context.GetHttpContext().User))
 			{
 				throw new StructuredRpcException(StatusCode.PermissionDenied, "Not authenticated to access stream {StreamId}", request.StreamId);
 			}
@@ -750,7 +750,7 @@ namespace Horde.Server.Jobs
 		/// <returns>Information about the new agent</returns>
 		public async Task<Empty> CreateEvents(CreateEventsRequest request, ServerCallContext context)
 		{
-			if (!_globalConfig.Value.Authorize(AclAction.CreateEvent, context.GetHttpContext().User))
+			if (!_globalConfig.Value.Authorize(LogAclAction.CreateEvent, context.GetHttpContext().User))
 			{
 				throw new StructuredRpcException(StatusCode.NotFound, "Access denied");
 			}
diff --git a/Engine/Source/Programs/Horde/Horde.Server/Jobs/JobsController.cs b/Engine/Source/Programs/Horde/Horde.Server/Jobs/JobsController.cs
--- a/Engine/Source/Programs/Horde/Horde.Server/Jobs/JobsController.cs
+++ b/Engine/Source/Programs/Horde/Horde.Server/Jobs/JobsController.cs
@@ -84,9 +84,9 @@ namespace Horde.Server.Jobs
 			{
 				return NotFound(create.StreamId);
 			}
-			if (!streamConfig.Authorize(AclAction.CreateJob, User))
+			if (!streamConfig.Authorize(JobAclAction.CreateJob, User))
 			{
-				return Forbid(AclAction.CreateJob, streamConfig.Id);
+				return Forbid(JobAclAction.CreateJob, streamConfig.Id);
 			}
 
 			// Get the name of the template ref
@@ -98,9 +98,9 @@ namespace Horde.Server.Jobs
 			{
 				return BadRequest($"Template {create.TemplateId} is not available for stream {streamConfig.Id}");
 			}
-			if (!templateRefConfig.Authorize(AclAction.CreateJob, User))
+			if (!templateRefConfig.Authorize(JobAclAction.CreateJob, User))
 			{
-				return Forbid(AclAction.CreateJob, streamConfig.Id);
+				return Forbid(JobAclAction.CreateJob, streamConfig.Id);
 			}
 
 			ITemplate? template = await _templateCollection.GetOrAddAsync(templateRefConfig);
@@ -251,9 +251,9 @@ namespace Horde.Server.Jobs
 			{
 				return NotFound(jobId);
 			}
-			if (!_globalConfig.Value.Authorize(job, AclAction.DeleteJob, User))
+			if (!_globalConfig.Value.Authorize(job, JobAclAction.DeleteJob, User))
 			{
-				return Forbid(AclAction.DeleteJob, jobId);
+				return Forbid(JobAclAction.DeleteJob, jobId);
 			}
 			if (!await _jobService.DeleteJobAsync(job))
 			{
@@ -283,9 +283,9 @@ namespace Horde.Server.Jobs
 			{
 				return NotFound(job.StreamId);
 			}
-			if (!streamConfig.Authorize(AclAction.UpdateJob, User))
+			if (!streamConfig.Authorize(JobAclAction.UpdateJob, User))
 			{
-				return Forbid(AclAction.UpdateJob, jobId);
+				return Forbid(JobAclAction.UpdateJob, jobId);
 			}
 
 			// Convert legacy behavior of clearing out the argument to setting the aborted flag
@@ -330,9 +330,9 @@ namespace Horde.Server.Jobs
 			{
 				return NotFound(job.StreamId);
 			}
-			if (!streamConfig.Authorize(AclAction.CreateSubscription, User))
+			if (!streamConfig.Authorize(NotificationAclAction.CreateSubscription, User))
 			{
-				return Forbid(AclAction.CreateSubscription, jobId);
+				return Forbid(NotificationAclAction.CreateSubscription, jobId);
 			}
 
 			ObjectId triggerId = job.NotificationTriggerId ?? ObjectId.GenerateNewId();
@@ -367,9 +367,9 @@ namespace Horde.Server.Jobs
 			{
 				return NotFound(job.StreamId);
 			}
-			if (!streamConfig.Authorize(AclAction.CreateSubscription, User))
+			if (!streamConfig.Authorize(NotificationAclAction.CreateSubscription, User))
 			{
-				return Forbid(AclAction.CreateSubscription, jobId);
+				return Forbid(NotificationAclAction.CreateSubscription, jobId);
 			}
 
 			INotificationSubscription? subscription;
@@ -407,9 +407,9 @@ namespace Horde.Server.Jobs
 			{
 				return NotFound(job.StreamId);
 			}
-			if (!streamConfig.Authorize(AclAction.ViewJob, User))
+			if (!streamConfig.Authorize(JobAclAction.ViewJob, User))
 			{
-				return Forbid(AclAction.ViewJob, job.StreamId);
+				return Forbid(JobAclAction.ViewJob, job.StreamId);
 			}
 			if (modifiedAfter != null && job.UpdateTimeUtc <= modifiedAfter.Value)
 			{
@@ -417,7 +417,7 @@ namespace Horde.Server.Jobs
 			}
 
 			IGraph graph = await _jobService.GetGraphAsync(job);
-			bool includeCosts = streamConfig.Authorize(AclAction.ViewCosts, User);
+			bool includeCosts = streamConfig.Authorize(AdminAclAction.ViewCosts, User);
 			return await CreateJobResponseAsync(job, graph, includeCosts, filter);
 		}
 
@@ -551,9 +551,9 @@ namespace Horde.Server.Jobs
 			{
 				return NotFound(job.StreamId);
 			}
-			if (!streamConfig.Authorize(AclAction.ViewJob, User))
+			if (!streamConfig.Authorize(JobAclAction.ViewJob, User))
 			{
-				return Forbid(AclAction.ViewJob, jobId);
+				return Forbid(JobAclAction.ViewJob, jobId);
 			}
 
 			IGraph graph = await _jobService.GetGraphAsync(job);
@@ -582,9 +582,9 @@ namespace Horde.Server.Jobs
 			{
 				return NotFound(job.StreamId);
 			}
-			if (!streamConfig.Authorize(AclAction.ViewJob, User))
+			if (!streamConfig.Authorize(JobAclAction.ViewJob, User))
 			{
-				return Forbid(AclAction.ViewJob, jobId);
+				return Forbid(JobAclAction.ViewJob, jobId);
 			}
 
 			IJobTiming jobTiming = await _jobService.GetJobTimingAsync(job);
@@ -653,7 +653,7 @@ namespace Horde.Server.Jobs
 			List<IJob> jobs = await _jobService.FindJobsByStreamWithTemplatesAsync(new StreamId(streamId), templateRefIds, count: count, consistentRead: false);
 
 			Dictionary<string, GetJobTimingResponse> jobTimings = await jobs.ToAsyncEnumerable()
-				.Where(job => _globalConfig.Value.Authorize(job, AclAction.ViewJob, User))
+				.Where(job => _globalConfig.Value.Authorize(job, JobAclAction.ViewJob, User))
 				.ToDictionaryAwaitAsync(x => ValueTask.FromResult(x.Id.ToString()), async job =>
 				{
 					IJobTiming jobTiming = await _jobService.GetJobTimingAsync(job);
@@ -686,9 +686,9 @@ namespace Horde.Server.Jobs
 			{
 				return NotFound(job.StreamId);
 			}
-			if (!streamConfig.Authorize(AclAction.ViewJob, User))
+			if (!streamConfig.Authorize(JobAclAction.ViewJob, User))
 			{
-				return Forbid(AclAction.ViewJob, jobId);
+				return Forbid(JobAclAction.ViewJob, jobId);
 			}
 
 			ITemplate? template = await _templateCollection.GetAsync(job.TemplateHash);
@@ -828,9 +828,9 @@ namespace Horde.Server.Jobs
 			foreach(IGrouping<StreamId, IJob> grouping in jobs.GroupBy(x => x.StreamId))
 			{
 				StreamConfig? streamConfig;
-				if (_globalConfig.Value.TryGetStream(grouping.Key, out streamConfig) && streamConfig.Authorize(AclAction.ViewJob, User))
+				if (_globalConfig.Value.TryGetStream(grouping.Key, out streamConfig) && streamConfig.Authorize(JobAclAction.ViewJob, User))
 				{
-					bool includeCosts = streamConfig.Authorize(AclAction.ViewCosts, User);
+					bool includeCosts = streamConfig.Authorize(AdminAclAction.ViewCosts, User);
 					foreach (IJob job in grouping)
 					{
 						IGraph graph = await _jobService.GetGraphAsync(job);
@@ -864,9 +864,9 @@ namespace Horde.Server.Jobs
 				{
 					return NotFound(job.StreamId);
 				}
-				if (!streamConfig.Authorize(AclAction.ExecuteJob, User))
+				if (!streamConfig.Authorize(JobAclAction.ExecuteJob, User))
 				{
-					return Forbid(AclAction.ExecuteJob, jobId);
+					return Forbid(JobAclAction.ExecuteJob, jobId);
 				}
 
 				IGraph graph = await _jobService.GetGraphAsync(job);
@@ -902,9 +902,9 @@ namespace Horde.Server.Jobs
 			{
 				return NotFound(job.StreamId);
 			}
-			if (!streamConfig.Authorize(AclAction.ViewJob, User))
+			if (!streamConfig.Authorize(JobAclAction.ViewJob, User))
 			{
-				return Forbid(AclAction.ViewJob, jobId);
+				return Forbid(JobAclAction.ViewJob, jobId);
 			}
 
 			IGraph graph = await _jobService.GetGraphAsync(job);
@@ -934,9 +934,9 @@ namespace Horde.Server.Jobs
 			{
 				return NotFound(job.StreamId);
 			}
-			if (!streamConfig.Authorize(AclAction.ViewJob, User))
+			if (!streamConfig.Authorize(JobAclAction.ViewJob, User))
 			{
-				return Forbid(AclAction.ViewJob, jobId);
+				return Forbid(JobAclAction.ViewJob, jobId);
 			}
 
 			IGraph graph = await _jobService.GetGraphAsync(job);
@@ -970,9 +970,9 @@ namespace Horde.Server.Jobs
 			{
 				return NotFound(job.StreamId);
 			}
-			if (!streamConfig.Authorize(AclAction.ViewJob, User))
+			if (!streamConfig.Authorize(JobAclAction.ViewJob, User))
 			{
-				return Forbid(AclAction.ViewJob, jobId);
+				return Forbid(JobAclAction.ViewJob, jobId);
 			}
 
 			IGraph graph = await _jobService.GetGraphAsync(job);
@@ -1007,9 +1007,9 @@ namespace Horde.Server.Jobs
 			{
 				return NotFound(job.StreamId);
 			}
-			if (!streamConfig.Authorize(AclAction.ViewJob, User))
+			if (!streamConfig.Authorize(JobAclAction.ViewJob, User))
 			{
-				return Forbid(AclAction.ViewJob, jobId);
+				return Forbid(JobAclAction.ViewJob, jobId);
 			}
 
 			IGraph graph = await _jobService.GetGraphAsync(job);
@@ -1047,12 +1047,12 @@ namespace Horde.Server.Jobs
 			{
 				return NotFound(job.StreamId);
 			}
-			if (!streamConfig.Authorize(AclAction.ViewJob, User))
+			if (!streamConfig.Authorize(JobAclAction.ViewJob, User))
 			{
-				return Forbid(AclAction.ViewJob, job.StreamId);
+				return Forbid(JobAclAction.ViewJob, job.StreamId);
 			}
 
-			bool includeCosts = streamConfig.Authorize(AclAction.ViewCosts, User);
+			bool includeCosts = streamConfig.Authorize(AdminAclAction.ViewCosts, User);
 
 			List<object> responses = new List<object>();
 			foreach (IJobStepBatch batch in job.Batches)
@@ -1124,12 +1124,12 @@ namespace Horde.Server.Jobs
 			{
 				return NotFound(job.StreamId);
 			}
-			if (!streamConfig.Authorize(AclAction.ViewJob, User))
+			if (!streamConfig.Authorize(JobAclAction.ViewJob, User))
 			{
-				return Forbid(AclAction.ViewJob, job.StreamId);
+				return Forbid(JobAclAction.ViewJob, job.StreamId);
 			}
 
-			bool includeCosts = streamConfig.Authorize(AclAction.ViewCosts, User);
+			bool includeCosts = streamConfig.Authorize(AdminAclAction.ViewCosts, User);
 			foreach (IJobStepBatch batch in job.Batches)
 			{
 				if (batch.Id == batchId)
@@ -1165,9 +1165,9 @@ namespace Horde.Server.Jobs
 			{
 				return NotFound(job.StreamId);
 			}
-			if (!streamConfig.Authorize(AclAction.ViewJob, User))
+			if (!streamConfig.Authorize(JobAclAction.ViewJob, User))
 			{
-				return Forbid(AclAction.ViewJob, jobId);
+				return Forbid(JobAclAction.ViewJob, jobId);
 			}
 
 			foreach (IJobStepBatch batch in job.Batches)
@@ -1226,16 +1226,16 @@ namespace Horde.Server.Jobs
 
 			if (request.Retry != null || request.Priority != null)
 			{
-				if (!streamConfig.Authorize(AclAction.RetryJobStep, User))
+				if (!streamConfig.Authorize(JobAclAction.RetryJobStep, User))
 				{
-					return Forbid(AclAction.RetryJobStep, jobId);
+					return Forbid(JobAclAction.RetryJobStep, jobId);
 				}
 			}
 			if (request.Properties != null)
 			{
-				if (!streamConfig.Authorize(AclAction.UpdateJob, User))
+				if (!streamConfig.Authorize(JobAclAction.UpdateJob, User))
 				{
-					return Forbid(AclAction.UpdateJob, jobId);
+					return Forbid(JobAclAction.UpdateJob, jobId);
 				}
 			}
 
@@ -1322,9 +1322,9 @@ namespace Horde.Server.Jobs
 			{
 				return NotFound(job.StreamId);
 			}
-			if (!streamConfig.Authorize(AclAction.ViewJob, User))
+			if (!streamConfig.Authorize(JobAclAction.ViewJob, User))
 			{
-				return Forbid(AclAction.ViewJob, jobId);
+				return Forbid(JobAclAction.ViewJob, jobId);
 			}
 
 			foreach (IJobStepBatch batch in job.Batches)
@@ -1369,9 +1369,9 @@ namespace Horde.Server.Jobs
 			{
 				return NotFound(job.StreamId);
 			}
-			if (!streamConfig.Authorize(AclAction.CreateSubscription, User))
+			if (!streamConfig.Authorize(NotificationAclAction.CreateSubscription, User))
 			{
-				return Forbid(AclAction.CreateSubscription, jobId);
+				return Forbid(NotificationAclAction.CreateSubscription, jobId);
 			}
 
 			if (!job.TryGetBatch(batchId, out IJobStepBatch? batch))
@@ -1430,9 +1430,9 @@ namespace Horde.Server.Jobs
 			{
 				return NotFound(job.StreamId);
 			}
-			if (!streamConfig.Authorize(AclAction.CreateSubscription, User))
+			if (!streamConfig.Authorize(NotificationAclAction.CreateSubscription, User))
 			{
-				return Forbid(AclAction.CreateSubscription, jobId);
+				return Forbid(NotificationAclAction.CreateSubscription, jobId);
 			}
 
 			INotificationSubscription? subscription;
@@ -1470,9 +1470,9 @@ namespace Horde.Server.Jobs
 			{
 				return NotFound(job.StreamId);
 			}
-			if (!streamConfig.Authorize(AclAction.ViewJob, User))
+			if (!streamConfig.Authorize(JobAclAction.ViewJob, User))
 			{
-				return Forbid(AclAction.ViewJob, jobId);
+				return Forbid(JobAclAction.ViewJob, jobId);
 			}
 
 			if (!job.TryGetBatch(batchId, out IJobStepBatch? batch))
@@ -1516,9 +1516,9 @@ namespace Horde.Server.Jobs
 			{
 				return NotFound(job.StreamId);
 			}
-			if (!streamConfig.Authorize(AclAction.ViewJob, User))
+			if (!streamConfig.Authorize(JobAclAction.ViewJob, User))
 			{
-				return Forbid(AclAction.ViewJob, jobId);
+				return Forbid(JobAclAction.ViewJob, jobId);
 			}
 
 			if (!job.TryGetBatch(batchId, out IJobStepBatch? batch))
@@ -1565,9 +1565,9 @@ namespace Horde.Server.Jobs
 				{
 					return NotFound(job.StreamId);
 				}
-				if (!streamConfig.Authorize(AclAction.CreateSubscription, User))
+				if (!streamConfig.Authorize(NotificationAclAction.CreateSubscription, User))
 				{
-					return Forbid(AclAction.CreateSubscription, jobId);
+					return Forbid(NotificationAclAction.CreateSubscription, jobId);
 				}
 
 				ObjectId newTriggerId;
diff --git a/Engine/Source/Programs/Horde/Horde.Server/Jobs/TestData/TestDataController.cs b/Engine/Source/Programs/Horde/Horde.Server/Jobs/TestData/TestDataController.cs
--- a/Engine/Source/Programs/Horde/Horde.Server/Jobs/TestData/TestDataController.cs
+++ b/Engine/Source/Programs/Horde/Horde.Server/Jobs/TestData/TestDataController.cs
@@ -129,7 +129,7 @@ namespace Horde.Server.Jobs.TestData
 			// authorize streams
 			foreach (StreamId streamId in streamIdValues)
 			{
-				if (_globalConfig.Value.TryGetStream(streamId, out StreamConfig? streamConfig) && streamConfig.Authorize(AclAction.ViewJob, User))
+				if (_globalConfig.Value.TryGetStream(streamId, out StreamConfig? streamConfig) && streamConfig.Authorize(JobAclAction.ViewJob, User))
 				{
 					queryStreams.Add(streamId);
 				}
@@ -269,7 +269,7 @@ namespace Horde.Server.Jobs.TestData
 			// authorize streams
 			foreach (StreamId streamId in streamIdValues)
 			{
-				if (_globalConfig.Value.TryGetStream(streamId, out StreamConfig? streamConfig) && streamConfig.Authorize(AclAction.ViewJob, User))
+				if (_globalConfig.Value.TryGetStream(streamId, out StreamConfig? streamConfig) && streamConfig.Authorize(JobAclAction.ViewJob, User))
 				{
 					queryStreams.Add(streamId);
 				}
@@ -300,7 +300,7 @@ namespace Horde.Server.Jobs.TestData
 			{
 				return NotFound();
 			}
-			if (!_globalConfig.Value.Authorize(job, AclAction.UpdateJob, User))
+			if (!_globalConfig.Value.Authorize(job, JobAclAction.UpdateJob, User))
 			{
 				return Forbid();
 			}
@@ -344,7 +344,7 @@ namespace Horde.Server.Jobs.TestData
 			List<ITestData> documents = await _testDataCollection.FindAsync(streamIdValue, minChange, maxChange, (jobId == null)? null : JobId.Parse(jobId), jobStepId?.ToSubResourceId(), key, index, count);
 			foreach (ITestData document in documents)
 			{
-				if (await _jobService.AuthorizeAsync(document.JobId, AclAction.ViewJob, User, _globalConfig.Value))
+				if (await _jobService.AuthorizeAsync(document.JobId, JobAclAction.ViewJob, User, _globalConfig.Value))
 				{
 					results.Add(PropertyFilter.Apply(new GetTestDataResponse(document), filter));
 				}
@@ -369,7 +369,7 @@ namespace Horde.Server.Jobs.TestData
 			{
 				return NotFound();
 			}
-			if (!await _jobService.AuthorizeAsync(testData.JobId, AclAction.ViewJob, User, _globalConfig.Value))
+			if (!await _jobService.AuthorizeAsync(testData.JobId, JobAclAction.ViewJob, User, _globalConfig.Value))
 			{
 				return Forbid();
 			}
diff --git a/Engine/Source/Programs/Horde/Horde.Server/Logs/LogAclAction.cs b/Engine/Source/Programs/Horde/Horde.Server/Logs/LogAclAction.cs
new file mode 100644
--- /dev/null
+++ b/Engine/Source/Programs/Horde/Horde.Server/Logs/LogAclAction.cs
@@ -0,0 +1,42 @@
+// Copyright Epic Games, Inc. All Rights Reserved.
+
+using Horde.Server.Acls;
+
+namespace Horde.Server.Logs
+{
+	/// <summary>
+	/// ACL actions for logs
+	/// </summary>
+	public static class LogAclAction
+	{
+		/// <summary>
+		/// Ability to create a log. Implicitly granted to agents.
+		/// </summary>
+		public static AclAction CreateLog { get; } = new AclAction("CreateLog");
+
+		/// <summary>
+		/// Ability to update log metadata
+		/// </summary>
+		public static AclAction UpdateLog { get; } = new AclAction("UpdateLog");
+
+		/// <summary>
+		/// Ability to view a log contents
+		/// </summary>
+		public static AclAction ViewLog { get; } = new AclAction("ViewLog");
+
+		/// <summary>
+		/// Ability to write log data
+		/// </summary>
+		public static AclAction WriteLogData { get; } = new AclAction("WriteLogData");
+
+		/// <summary>
+		/// Ability to create events
+		/// </summary>
+		public static AclAction CreateEvent { get; } = new AclAction("CreateEvent");
+
+		/// <summary>
+		/// Ability to view events
+		/// </summary>
+		public static AclAction ViewEvent { get; } = new AclAction("ViewEvent");
+	}
+}
diff --git a/Engine/Source/Programs/Horde/Horde.Server/Logs/LogsController.cs b/Engine/Source/Programs/Horde/Horde.Server/Logs/LogsController.cs
--- a/Engine/Source/Programs/Horde/Horde.Server/Logs/LogsController.cs
+++ b/Engine/Source/Programs/Horde/Horde.Server/Logs/LogsController.cs
@@ -12,6 +12,7 @@ using EpicGames.Core;
 using EpicGames.Horde.Logs;
 using EpicGames.Horde.Storage;
 using Horde.Server.Acls;
+using Horde.Server.Agents.Sessions;
 using Horde.Server.Issues;
 using Horde.Server.Jobs;
 using Horde.Server.Logs.Data;
@@ -85,7 +86,7 @@ namespace Horde.Server.Logs
 			{
 				return NotFound();
 			}
-			if (!await AuthorizeAsync(logFile, AclAction.ViewLog, User))
+			if (!await AuthorizeAsync(logFile, LogAclAction.ViewLog, User))
 			{
 				return Forbid();
 			}
@@ -111,7 +112,7 @@ namespace Horde.Server.Logs
 			{
 				return NotFound();
 			}
-			if (!await AuthorizeAsync(logFile, AclAction.WriteLogData, User))
+			if (!await AuthorizeAsync(logFile, LogAclAction.WriteLogData, User))
 			{
 				return Forbid();
 			}
@@ -142,7 +143,7 @@ namespace Horde.Server.Logs
 			{
 				return NotFound();
 			}
-			if (!await AuthorizeAsync(logFile, AclAction.ViewLog, User))
+			if (!await AuthorizeAsync(logFile, LogAclAction.ViewLog, User))
 			{
 				return Forbid();
 			}
@@ -177,7 +178,7 @@ namespace Horde.Server.Logs
 			{
 				return NotFound();
 			}
-			if (!await AuthorizeAsync(logFile, AclAction.ViewLog, User))
+			if (!await AuthorizeAsync(logFile, LogAclAction.ViewLog, User))
 			{
 				return Forbid();
 			}
@@ -272,7 +273,7 @@ namespace Horde.Server.Logs
 			{
 				return NotFound();
 			}
-			if (!await AuthorizeAsync(logFile, AclAction.ViewLog, User))
+			if (!await AuthorizeAsync(logFile, LogAclAction.ViewLog, User))
 			{
 				return Forbid();
 			}
@@ -301,7 +302,7 @@ namespace Horde.Server.Logs
 			{
 				return NotFound();
 			}
-			if (!await AuthorizeAsync(logFile, AclAction.ViewLog, User))
+			if (!await AuthorizeAsync(logFile, LogAclAction.ViewLog, User))
 			{
 				return Forbid();
 			}
@@ -345,7 +346,7 @@ namespace Horde.Server.Logs
 			{
 				return NotFound();
 			}
-			if (!await AuthorizeAsync(logFile, AclAction.WriteLogData, User))
+			if (!await AuthorizeAsync(logFile, LogAclAction.WriteLogData, User))
 			{
 				return Forbid();
 			}
@@ -380,7 +381,7 @@ namespace Horde.Server.Logs
 			{
 				return true;
 			}
-			if (action == AclAction.ViewLog && logFile.SessionId != null && globalConfig.Authorize(AclAction.ViewSession, user))
+			if (action == LogAclAction.ViewLog && logFile.SessionId != null && globalConfig.Authorize(SessionAclAction.ViewSession, user))
 			{
 				return true;
 			}
diff --git a/Engine/Source/Programs/Horde/Horde.Server/Notifications/NotificationAclAction.cs b/Engine/Source/Programs/Horde/Horde.Server/Notifications/NotificationAclAction.cs
new file mode 100644
--- /dev/null
+++ b/Engine/Source/Programs/Horde/Horde.Server/Notifications/NotificationAclAction.cs
@@ -0,0 +1,17 @@
+// Copyright Epic Games, Inc. All Rights Reserved.
+
+using Horde.Server.Acls;
+
+namespace Horde.Server.Notifications
+{
+	/// <summary>
+	/// ACL actions which apply to notifications
+	/// </summary>
+	public static class NotificationAclAction
+	{
+		/// <summary>
+		/// Ability to subscribe to notifications
+		/// </summary>
+		public static AclAction CreateSubscription = new AclAction("CreateSubscription");
+	}
+}
diff --git a/Engine/Source/Programs/Horde/Horde.Server/Notifications/SubscriptionsController.cs b/Engine/Source/Programs/Horde/Horde.Server/Notifications/SubscriptionsController.cs
--- a/Engine/Source/Programs/Horde/Horde.Server/Notifications/SubscriptionsController.cs
+++ b/Engine/Source/Programs/Horde/Horde.Server/Notifications/SubscriptionsController.cs
@@ -131,7 +131,7 @@ namespace Horde.Server.Notifications
 				{
 					return BadRequest($"Invalid user id: '{subscription.UserId}'.");
 				}
-				if (authorizedUsers.Add(newUserId) && !_globalConfig.Value.Authorize(AclAction.Impersonate, User))
+				if (authorizedUsers.Add(newUserId) && !_globalConfig.Value.Authorize(AdminAclAction.Impersonate, User))
 				{
 					return Forbid();
 				}
diff --git a/Engine/Source/Programs/Horde/Horde.Server/Perforce/PerforceController.cs b/Engine/Source/Programs/Horde/Horde.Server/Perforce/PerforceController.cs
--- a/Engine/Source/Programs/Horde/Horde.Server/Perforce/PerforceController.cs
+++ b/Engine/Source/Programs/Horde/Horde.Server/Perforce/PerforceController.cs
@@ -62,7 +62,7 @@ namespace Horde.Server.Perforce
 		{
 			GlobalConfig globalConfig = _globalConfig.Value;
 
-			if (!globalConfig.Authorize(AclAction.AdminRead, User))
+			if (!globalConfig.Authorize(AdminAclAction.AdminRead, User))
 			{
 				return Forbid();
 			}
diff --git a/Engine/Source/Programs/Horde/Horde.Server/Projects/ProjectAclAction.cs b/Engine/Source/Programs/Horde/Horde.Server/Projects/ProjectAclAction.cs
new file mode 100644
--- /dev/null
+++ b/Engine/Source/Programs/Horde/Horde.Server/Projects/ProjectAclAction.cs
@@ -0,0 +1,32 @@
+// Copyright Epic Games, Inc. All Rights Reserved.
+
+using Horde.Server.Acls;
+
+namespace Horde.Server.Projects
+{
+	/// <summary>
+	/// ACL actions that can be performed on projects
+	/// </summary>
+	public static class ProjectAclAction
+	{
+		/// <summary>
+		/// Allows the creation of new projects
+		/// </summary>
+		public static AclAction CreateProject { get; } = new AclAction("CreateProject");
+
+		/// <summary>
+		/// Allows deletion of projects.
+		/// </summary>
+		public static AclAction DeleteProject { get; } = new AclAction("DeleteProject");
+
+		/// <summary>
+		/// Modify attributes of a project (name, categories, etc...)
+		/// </summary>
+		public static AclAction UpdateProject { get; } = new AclAction("UpdateProject");
+
+		/// <summary>
+		/// View information about a project
+		/// </summary>
+		public static AclAction ViewProject { get; } = new AclAction("ViewProject");
+	}
+}
diff --git a/Engine/Source/Programs/Horde/Horde.Server/Projects/ProjectsController.cs b/Engine/Source/Programs/Horde/Horde.Server/Projects/ProjectsController.cs
--- a/Engine/Source/Programs/Horde/Horde.Server/Projects/ProjectsController.cs
+++ b/Engine/Source/Programs/Horde/Horde.Server/Projects/ProjectsController.cs
@@ -52,12 +52,12 @@ namespace Horde.Server.Projects
 			List<object> responses = new List<object>();
 			foreach (ProjectConfig projectConfig in globalConfig.Projects)
 			{
-				if (projectConfig.Authorize(AclAction.ViewProject, User))
+				if (projectConfig.Authorize(ProjectAclAction.ViewProject, User))
 				{
 					List<StreamConfig>? visibleStreams = null;
 					if (includeStreams || includeCategories)
 					{
-						visibleStreams = projectConfig.Streams.Where(x => x.Authorize(AclAction.ViewStream, User)).ToList();
+						visibleStreams = projectConfig.Streams.Where(x => x.Authorize(StreamAclAction.ViewStream, User)).ToList();
 					}
 					responses.Add(GetProjectResponse.FromConfig(projectConfig, includeStreams, includeCategories, visibleStreams).ApplyFilter(filter));
 				}
@@ -81,9 +81,9 @@ namespace Horde.Server.Projects
 			{
 				return NotFound(projectId);
 			}
-			if (!projectConfig.Authorize(AclAction.ViewProject, User))
+			if (!projectConfig.Authorize(ProjectAclAction.ViewProject, User))
 			{
-				return Forbid(AclAction.ViewProject, projectId);
+				return Forbid(ProjectAclAction.ViewProject, projectId);
 			}
 
 			bool includeStreams = PropertyFilter.Includes(filter, nameof(GetProjectResponse.Streams));
@@ -92,7 +92,7 @@ namespace Horde.Server.Projects
 			List<StreamConfig>? visibleStreams = null;
 			if (includeStreams || includeCategories)
 			{
-				visibleStreams = projectConfig.Streams.Where(x => x.Authorize(AclAction.ViewStream, User)).ToList();
+				visibleStreams = projectConfig.Streams.Where(x => x.Authorize(StreamAclAction.ViewStream, User)).ToList();
 			}
 
 			return GetProjectResponse.FromConfig(projectConfig, includeStreams, includeCategories, visibleStreams).ApplyFilter(filter);
@@ -112,9 +112,9 @@ namespace Horde.Server.Projects
 			{
 				return NotFound(projectId);
 			}
-			if (!projectConfig.Authorize(AclAction.ViewProject, User))
+			if (!projectConfig.Authorize(ProjectAclAction.ViewProject, User))
 			{
-				return Forbid(AclAction.ViewProject, projectId);
+				return Forbid(ProjectAclAction.ViewProject, projectId);
 			}
 
 			ConfigResource? logoResource = projectConfig.Logo;
diff --git a/Engine/Source/Programs/Horde/Horde.Server/Secrets/CredentialAclAction.cs b/Engine/Source/Programs/Horde/Horde.Server/Secrets/CredentialAclAction.cs
new file mode 100644
--- /dev/null
+++ b/Engine/Source/Programs/Horde/Horde.Server/Secrets/CredentialAclAction.cs
@@ -0,0 +1,37 @@
+// Copyright Epic Games, Inc. All Rights Reserved.
+
+using Horde.Server.Acls;
+
+namespace Horde.Server.Secrets
+{
+	/// <summary>
+	/// ACL actions for credentials
+	/// </summary>
+	public static class CredentialAclAction
+	{
+		/// <summary>
+		/// Create a new credential
+		/// </summary>
+		public static AclAction CreateCredential { get; } = new AclAction("CreateCredential");
+
+		/// <summary>
+		/// Delete a credential
+		/// </summary>
+		public static AclAction DeleteCredential { get; } = new AclAction("DeleteCredential");
+
+		/// <summary>
+		/// Modify an existing credential
+		/// </summary>
+		public static AclAction UpdateCredential { get; } = new AclAction("UpdateCredential");
+
+		/// <summary>
+		/// Enumerates all the available credentials
+		/// </summary>
+		public static AclAction ListCredentials { get; } = new AclAction("ListCredentials");
+
+		/// <summary>
+		/// View a credential
+		/// </summary>
+		public static AclAction ViewCredential { get; } = new AclAction("ViewCredential");
+	}
+}
diff --git a/Engine/Source/Programs/Horde/Horde.Server/Secrets/CredentialsController.cs b/Engine/Source/Programs/Horde/Horde.Server/Secrets/CredentialsController.cs
--- a/Engine/Source/Programs/Horde/Horde.Server/Secrets/CredentialsController.cs
+++ b/Engine/Source/Programs/Horde/Horde.Server/Secrets/CredentialsController.cs
@@ -42,7 +42,7 @@ namespace Horde.Server.Secrets
 		[Route("/api/v1/credentials")]
 		public async Task<ActionResult<CreateCredentialResponse>> CreateCredentialAsync([FromBody] CreateCredentialRequest create)
 		{
-			if(!_globalConfig.Value.Authorize(AclAction.CreateCredential, User))
+			if(!_globalConfig.Value.Authorize(CredentialAclAction.CreateCredential, User))
 			{
 				return Forbid();
 			}
@@ -62,7 +62,7 @@ namespace Horde.Server.Secrets
 		[ProducesResponseType(typeof(List<GetCredentialResponse>), 200)]
 		public async Task<ActionResult<object>> FindCredentialAsync([FromQuery] string? name = null, [FromQuery] PropertyFilter? filter = null)
 		{
-			if (!_globalConfig.Value.Authorize(AclAction.ListCredentials, User))
+			if (!_globalConfig.Value.Authorize(CredentialAclAction.ListCredentials, User))
 			{
 				return Forbid();
 			}
@@ -72,7 +72,7 @@ namespace Horde.Server.Secrets
 			List<object> responses = new List<object>();
 			foreach (Credential credential in credentials)
 			{
-				if (_globalConfig.Value.Authorize(AclAction.ViewCredential, User))
+				if (_globalConfig.Value.Authorize(CredentialAclAction.ViewCredential, User))
 				{
 					responses.Add(new GetCredentialResponse(credential).ApplyFilter(filter));
 				}
@@ -99,7 +99,7 @@ namespace Horde.Server.Secrets
 				return NotFound();
 			}
 
-			if (!_globalConfig.Value.Authorize(AclAction.ViewCredential, User))
+			if (!_globalConfig.Value.Authorize(CredentialAclAction.ViewCredential, User))
 			{
 				return Forbid();
 			}
@@ -125,11 +125,7 @@ namespace Horde.Server.Secrets
 				return NotFound();
 			}
 
-			if (!_globalConfig.Value.Authorize(AclAction.UpdateCredential, User))
-			{
-				return Forbid();
-			}
-			if (update.Acl != null && !_globalConfig.Value.Authorize(AclAction.ChangePermissions, User))
+			if (!_globalConfig.Value.Authorize(CredentialAclAction.UpdateCredential, User))
 			{
 				return Forbid();
 			}
@@ -154,7 +150,7 @@ namespace Horde.Server.Secrets
 			{
 				return NotFound();
 			}
-			if (!_globalConfig.Value.Authorize(AclAction.DeleteCredential, User))
+			if (!_globalConfig.Value.Authorize(CredentialAclAction.DeleteCredential, User))
 			{
 				return Forbid();
 			}
diff --git a/Engine/Source/Programs/Horde/Horde.Server/Server/AccountController.cs b/Engine/Source/Programs/Horde/Horde.Server/Server/AccountController.cs
--- a/Engine/Source/Programs/Horde/Horde.Server/Server/AccountController.cs
+++ b/Engine/Source/Programs/Horde/Horde.Server/Server/AccountController.cs
@@ -73,7 +73,7 @@ namespace Horde.Server.Server
 			if (User.Identity?.IsAuthenticated ?? false)
 			{
 				content.Append(CultureInfo.InvariantCulture, $"<p>User <b>{User.Identity?.Name}</b> is logged in. <a href=\"/account/logout\">Log out</a></p>");
-				if (_globalConfig.Value.Authorize(AclAction.AdminWrite, User))
+				if (_globalConfig.Value.Authorize(AdminAclAction.AdminWrite, User))
 				{
 					content.Append("<p>");
 					content.Append("<a href=\"/api/v1/admin/token\">Get bearer token</a><br/>");
diff --git a/Engine/Source/Programs/Horde/Horde.Server/Server/AdminAclAction.cs b/Engine/Source/Programs/Horde/Horde.Server/Server/AdminAclAction.cs
new file mode 100644
--- /dev/null
+++ b/Engine/Source/Programs/Horde/Horde.Server/Server/AdminAclAction.cs
@@ -0,0 +1,37 @@
+// Copyright Epic Games, Inc. All Rights Reserved.
+
+using Horde.Server.Acls;
+
+namespace Horde.Server.Server
+{
+	/// <summary>
+	/// ACL actions for admin operations
+	/// </summary>
+	public static class AdminAclAction
+	{
+		/// <summary>
+		/// Ability to read any data from the server. Always inherited.
+		/// </summary>
+		public static AclAction AdminRead { get; } = new AclAction("AdminRead");
+
+		/// <summary>
+		/// Ability to write any data to the server.
+		/// </summary>
+		public static AclAction AdminWrite { get; } = new AclAction("AdminWrite");
+
+		/// <summary>
+		/// Ability to impersonate another user
+		/// </summary>
+		public static AclAction Impersonate { get; } = new AclAction("Impersonate");
+
+		/// <summary>
+		/// View estimated costs for particular operations
+		/// </summary>
+		public static AclAction ViewCosts { get; } = new AclAction("ViewCosts");
+
+		/// <summary>
+		/// Issue bearer token for the current user
+		/// </summary>
+		public static AclAction IssueBearerToken { get; } = new AclAction("IssueBearerToken");
+	}
+}
diff --git a/Engine/Source/Programs/Horde/Horde.Server/Server/AdminController.cs b/Engine/Source/Programs/Horde/Horde.Server/Server/AdminController.cs
--- a/Engine/Source/Programs/Horde/Horde.Server/Server/AdminController.cs
+++ b/Engine/Source/Programs/Horde/Horde.Server/Server/AdminController.cs
@@ -75,9 +75,9 @@ namespace Horde.Server.Server
 		[Route("/api/v1/admin/token")]
 		public async Task<ActionResult<string>> GetTokenAsync()
 		{
-			if (!_globalConfig.Value.Authorize(AclAction.IssueBearerToken, User))
+			if (!_globalConfig.Value.Authorize(AdminAclAction.IssueBearerToken, User))
 			{
-				return Forbid(AclAction.IssueBearerToken);
+				return Forbid(AdminAclAction.IssueBearerToken);
 			}
 
 			return await _aclService.IssueBearerTokenAsync(User.Claims, GetDefaultExpiryTime());
@@ -90,9 +90,9 @@ namespace Horde.Server.Server
 		[Route("/api/v1/admin/config")]
 		public ActionResult<object> GetConfig()
 		{
-			if (!_globalConfig.Value.Authorize(AclAction.AdminRead, User))
+			if (!_globalConfig.Value.Authorize(AdminAclAction.AdminRead, User))
 			{
-				return Forbid(AclAction.AdminRead);
+				return Forbid(AdminAclAction.AdminRead);
 			}
 
 			return _globalConfig.Value;
@@ -107,9 +107,9 @@ namespace Horde.Server.Server
 		[Route("/api/v1/admin/roletoken")]
 		public async Task<ActionResult<string>> GetRoleTokenAsync([FromQuery] string roles)
 		{
-			if (!_globalConfig.Value.Authorize(AclAction.AdminWrite, User))
+			if (!_globalConfig.Value.Authorize(AdminAclAction.AdminWrite, User))
 			{
-				return Forbid(AclAction.AdminWrite);
+				return Forbid(AdminAclAction.AdminWrite);
 			}
 
 			List<Claim> claims = new List<Claim>();
@@ -126,9 +126,9 @@ namespace Horde.Server.Server
 		[Route("/api/v1/admin/registrationtoken")]
 		public async Task<ActionResult<string>> GetRegistrationTokenAsync()
 		{
-			if (!_globalConfig.Value.Authorize(AclAction.AdminWrite, User))
+			if (!_globalConfig.Value.Authorize(AdminAclAction.AdminWrite, User))
 			{
-				return Forbid(AclAction.AdminWrite);
+				return Forbid(AdminAclAction.AdminWrite);
 			}
 
 			List<AclClaimConfig> claims = new List<AclClaimConfig>();
@@ -146,9 +146,9 @@ namespace Horde.Server.Server
 		[Route("/api/v1/admin/softwaretoken")]
 		public async Task<ActionResult<string>> GetSoftwareTokenAsync()
 		{
-			if (!_globalConfig.Value.Authorize(AclAction.AdminWrite, User))
+			if (!_globalConfig.Value.Authorize(AdminAclAction.AdminWrite, User))
 			{
-				return Forbid(AclAction.AdminWrite);
+				return Forbid(AdminAclAction.AdminWrite);
 			}
 
 			List<AclClaimConfig> claims = new List<AclClaimConfig>();
@@ -166,9 +166,9 @@ namespace Horde.Server.Server
 		[Route("/api/v1/admin/softwaredownloadtoken")]
 		public async Task<ActionResult<string>> GetSoftwareDownloadTokenAsync()
 		{
-			if (!_globalConfig.Value.Authorize(AclAction.AdminRead, User))
+			if (!_globalConfig.Value.Authorize(AdminAclAction.AdminRead, User))
 			{
-				return Forbid(AclAction.AdminRead);
+				return Forbid(AdminAclAction.AdminRead);
 			}
 
 			List<AclClaimConfig> claims = new List<AclClaimConfig>();
@@ -186,9 +186,9 @@ namespace Horde.Server.Server
 		[Route("/api/v1/admin/configtoken")]
 		public async Task<ActionResult<string>> GetConfigToken()
 		{
-			if (!_globalConfig.Value.Authorize(AclAction.AdminRead, User))
+			if (!_globalConfig.Value.Authorize(AdminAclAction.AdminRead, User))
 			{
-				return Forbid(AclAction.AdminRead);
+				return Forbid(AdminAclAction.AdminRead);
 			}
 
 			List<AclClaimConfig> claims = new List<AclClaimConfig>();
@@ -206,9 +206,9 @@ namespace Horde.Server.Server
 		[Route("/api/v1/admin/chainedjobtoken")]
 		public async Task<ActionResult<string>> GetChainedJobToken()
 		{
-			if (!_globalConfig.Value.Authorize(AclAction.AdminRead, User))
+			if (!_globalConfig.Value.Authorize(AdminAclAction.AdminRead, User))
 			{
-				return Forbid(AclAction.AdminRead);
+				return Forbid(AdminAclAction.AdminRead);
 			}
 
 			List<AclClaimConfig> claims = new List<AclClaimConfig>();
diff --git a/Engine/Source/Programs/Horde/Horde.Server/Server/DebugController.cs b/Engine/Source/Programs/Horde/Horde.Server/Server/DebugController.cs
--- a/Engine/Source/Programs/Horde/Horde.Server/Server/DebugController.cs
+++ b/Engine/Source/Programs/Horde/Horde.Server/Server/DebugController.cs
@@ -281,7 +281,7 @@ namespace Horde.Server.Server
 		[Route("/api/v1/debug/environment")]
 		public ActionResult GetServerEnvVars()
 		{
-			if (!_globalConfig.Value.Authorize(AclAction.AdminRead, User))
+			if (!_globalConfig.Value.Authorize(AdminAclAction.AdminRead, User))
 			{
 				return Forbid();
 			}
@@ -307,7 +307,7 @@ namespace Horde.Server.Server
 		[Route("/api/v1/debug/queue")]
 		public ActionResult<object> GetQueueStatus()
 		{
-			if (!_globalConfig.Value.Authorize(AclAction.AdminRead, User))
+			if (!_globalConfig.Value.Authorize(AdminAclAction.AdminRead, User))
 			{
 				return Forbid();
 			}
@@ -323,7 +323,7 @@ namespace Horde.Server.Server
 		[Route("/api/v1/debug/config")]
 		public ActionResult GetConfig()
 		{
-			if (!_globalConfig.Value.Authorize(AclAction.AdminRead, User))
+			if (!_globalConfig.Value.Authorize(AdminAclAction.AdminRead, User))
 			{
 				return Forbid();
 			}
@@ -349,7 +349,7 @@ namespace Horde.Server.Server
 			[FromQuery] int argCount = 0,
 			[FromQuery] int argLen = 10)
 		{
-			if (!_globalConfig.Value.Authorize(AclAction.AdminRead, User))
+			if (!_globalConfig.Value.Authorize(AdminAclAction.AdminRead, User))
 			{
 				return Forbid();
 			}
@@ -402,7 +402,7 @@ namespace Horde.Server.Server
 		[ProducesResponseType(200, Type = typeof(GetGraphResponse))]
 		public async Task<ActionResult<List<object>>> GetGraphsAsync([FromQuery] int? index = null, [FromQuery] int? count = null, [FromQuery] PropertyFilter? filter = null)
 		{
-			if (!_globalConfig.Value.Authorize(AclAction.AdminRead, User))
+			if (!_globalConfig.Value.Authorize(AdminAclAction.AdminRead, User))
 			{
 				return Forbid();
 			}
@@ -420,7 +420,7 @@ namespace Horde.Server.Server
 		[ProducesResponseType(200, Type = typeof(GetGraphResponse))]
 		public async Task<ActionResult<object>> GetGraphAsync(string graphId, [FromQuery] PropertyFilter? filter = null)
 		{
-			if (!_globalConfig.Value.Authorize(AclAction.AdminRead, User))
+			if (!_globalConfig.Value.Authorize(AdminAclAction.AdminRead, User))
 			{
 				return Forbid();
 			}
@@ -439,7 +439,7 @@ namespace Horde.Server.Server
 		[Route("/api/v1/debug/logs/{LogFileId}")]
 		public async Task<ActionResult<object>> GetLogAsync(LogId logFileId, [FromQuery] PropertyFilter? filter = null)
 		{
-			if (!_globalConfig.Value.Authorize(AclAction.AdminRead, User))
+			if (!_globalConfig.Value.Authorize(AdminAclAction.AdminRead, User))
 			{
 				return Forbid();
 			}
@@ -461,7 +461,7 @@ namespace Horde.Server.Server
 		[Route("/api/v1/debug/collections/{Name}")]
 		public async Task<ActionResult<object>> GetDocumentsAsync(string name, [FromQuery] string? filter = null, [FromQuery] int index = 0, [FromQuery] int count = 10)
 		{
-			if (!_globalConfig.Value.Authorize(AclAction.AdminRead, User))
+			if (!_globalConfig.Value.Authorize(AdminAclAction.AdminRead, User))
 			{
 				return Forbid();
 			}
diff --git a/Engine/Source/Programs/Horde/Horde.Server/Server/GlobalConfig.cs b/Engine/Source/Programs/Horde/Horde.Server/Server/GlobalConfig.cs
--- a/Engine/Source/Programs/Horde/Horde.Server/Server/GlobalConfig.cs
+++ b/Engine/Source/Programs/Horde/Horde.Server/Server/GlobalConfig.cs
@@ -273,7 +273,7 @@ namespace Horde.Server.Server
 			}
 			else
 			{
-				return this.Authorize(AclAction.Impersonate, user);
+				return this.Authorize(AdminAclAction.Impersonate, user);
 			}
 		}
 
diff --git a/Engine/Source/Programs/Horde/Horde.Server/Server/Notices/NoticesController.cs b/Engine/Source/Programs/Horde/Horde.Server/Server/Notices/NoticesController.cs
--- a/Engine/Source/Programs/Horde/Horde.Server/Server/Notices/NoticesController.cs
+++ b/Engine/Source/Programs/Horde/Horde.Server/Server/Notices/NoticesController.cs
@@ -52,7 +52,7 @@ namespace Horde.Server.Server.Notices
 		[HttpPost("/api/v1/notices")]
 		public async Task<ActionResult> AddNoticeAsync(CreateNoticeRequest request)
 		{
-			if (!_globalConfig.Value.Authorize(AclAction.AdminWrite, User))
+			if (!_globalConfig.Value.Authorize(AdminAclAction.AdminWrite, User))
 			{
 				return Forbid();
 			}
@@ -71,7 +71,7 @@ namespace Horde.Server.Server.Notices
 		[HttpPut("/api/v1/notices")]
 		public async Task<ActionResult> UpdateNoticeAsync(UpdateNoticeRequest request)
 		{
-			if (!_globalConfig.Value.Authorize(AclAction.AdminWrite, User))
+			if (!_globalConfig.Value.Authorize(AdminAclAction.AdminWrite, User))
 			{
 				return Forbid();
 			}
@@ -86,7 +86,7 @@ namespace Horde.Server.Server.Notices
 		[HttpDelete("/api/v1/notices/{id}")]
 		public async Task<ActionResult> DeleteNoticeAsync(string id)
 		{
-			if (!_globalConfig.Value.Authorize(AclAction.AdminWrite, User))
+			if (!_globalConfig.Value.Authorize(AdminAclAction.AdminWrite, User))
 			{
 				return Forbid();
 			}
diff --git a/Engine/Source/Programs/Horde/Horde.Server/Server/RpcService.cs b/Engine/Source/Programs/Horde/Horde.Server/Server/RpcService.cs
--- a/Engine/Source/Programs/Horde/Horde.Server/Server/RpcService.cs
+++ b/Engine/Source/Programs/Horde/Horde.Server/Server/RpcService.cs
@@ -266,7 +266,7 @@ namespace Horde.Server.Server
 			IAgent? agent = await _agentService.GetAgentAsync(agentId);
 			if (agent == null)
 			{
-				if (!globalConfig.Authorize(AclAction.CreateAgent, context.GetHttpContext().User))
+				if (!globalConfig.Authorize(AgentAclAction.CreateAgent, context.GetHttpContext().User))
 				{
 					throw new StructuredRpcException(StatusCode.PermissionDenied, "User is not authenticated to create new agents");
 				}
@@ -275,7 +275,7 @@ namespace Horde.Server.Server
 			}
 
 			// Make sure we're allowed to create sessions on this agent
-			if (!globalConfig.Authorize(AclAction.CreateSession, context.GetHttpContext().User))
+			if (!globalConfig.Authorize(SessionAclAction.CreateSession, context.GetHttpContext().User))
 			{
 				throw new StructuredRpcException(StatusCode.PermissionDenied, "User is not authenticated to create session for {AgentId}", request.Name);
 			}
@@ -430,7 +430,7 @@ namespace Horde.Server.Server
 			{
 				throw new StructuredRpcException(StatusCode.NotFound, $"Missing tool {toolId}");
 			}
-			if (!toolConfig.Authorize(AclAction.DownloadTool, context.GetHttpContext().User))
+			if (!toolConfig.Authorize(ToolAclAction.DownloadTool, context.GetHttpContext().User))
 			{
 				throw new StructuredRpcException(StatusCode.NotFound, "Access to software is forbidden");
 			}
diff --git a/Engine/Source/Programs/Horde/Horde.Server/ServerSettings.cs b/Engine/Source/Programs/Horde/Horde.Server/ServerSettings.cs
--- a/Engine/Source/Programs/Horde/Horde.Server/ServerSettings.cs
+++ b/Engine/Source/Programs/Horde/Horde.Server/ServerSettings.cs
@@ -12,9 +12,17 @@ using EpicGames.Horde.Storage;
 using EpicGames.Horde.Storage.Nodes;
 using EpicGames.Perforce;
 using Horde.Server.Acls;
+using Horde.Server.Agents;
 using Horde.Server.Agents.Fleet;
+using Horde.Server.Agents.Pools;
+using Horde.Server.Agents.Sessions;
+using Horde.Server.Agents.Software;
+using Horde.Server.Jobs;
+using Horde.Server.Logs;
+using Horde.Server.Projects;
 using Horde.Server.Server;
 using Horde.Server.Storage.Backends;
+using Horde.Server.Streams;
 using Horde.Server.Telemetry;
 using Horde.Server.Tools;
 using Horde.Server.Utilities;
@@ -371,16 +379,6 @@ namespace Horde.Server
 		/// </summary>
 		public string DatabaseName { get; set; } = "Horde";
 
-		/// <summary>
-		/// The claim type for administrators
-		/// </summary>
-		public string AdminClaimType { get; set; } = HordeClaimTypes.InternalRole;
-
-		/// <summary>
-		/// Value of the claim type for administrators
-		/// </summary>
-		public string AdminClaimValue { get; set; } = "admin";
-
 		/// <summary>
 		/// Optional certificate to trust in order to access the database (eg. AWS public cert for TLS)
 		/// </summary>
@@ -769,20 +767,13 @@ namespace Horde.Server
 		AclConfig GetDefaultAcl()
 		{
 			AclConfig defaultAcl = new AclConfig();
-			defaultAcl.Entries.Add(new AclEntryConfig(new AclClaimConfig(ClaimTypes.Role, "internal:AgentRegistration"), new[] { AclAction.CreateAgent, AclAction.CreateSession }));
-			defaultAcl.Entries.Add(new AclEntryConfig(HordeClaims.AgentRegistrationClaim, new[] { AclAction.CreateAgent, AclAction.CreateSession, AclAction.UpdateAgent, AclAction.DownloadSoftware, AclAction.CreatePool, AclAction.UpdatePool, AclAction.ViewPool, AclAction.DeletePool, AclAction.ListPools, AclAction.ViewStream, AclAction.ViewProject, AclAction.ViewJob, AclAction.ViewCosts }));
-			defaultAcl.Entries.Add(new AclEntryConfig(HordeClaims.AgentRoleClaim, new[] { AclAction.ViewProject, AclAction.ViewStream, AclAction.CreateEvent, AclAction.DownloadSoftware }));
-			defaultAcl.Entries.Add(new AclEntryConfig(HordeClaims.DownloadSoftwareClaim, new[] { AclAction.DownloadSoftware }));
-			defaultAcl.Entries.Add(new AclEntryConfig(HordeClaims.UploadSoftwareClaim, new[] { AclAction.UploadSoftware }));
-			defaultAcl.Entries.Add(new AclEntryConfig(HordeClaims.ConfigureProjectsClaim, new[] { AclAction.CreateProject, AclAction.UpdateProject, AclAction.ViewProject, AclAction.CreateStream, AclAction.UpdateStream, AclAction.ViewStream, AclAction.ChangePermissions }));
-			defaultAcl.Entries.Add(new AclEntryConfig(HordeClaims.StartChainedJobClaim, new[] { AclAction.CreateJob, AclAction.ExecuteJob, AclAction.UpdateJob, AclAction.ViewJob, AclAction.ViewTemplate, AclAction.ViewStream }));
-
-			if (AdminClaimType != null && AdminClaimValue != null)
-			{
-				AclAction[] actions = Enum.GetValues(typeof(AclAction)).OfType<AclAction>().ToArray();
-				defaultAcl.Entries.Add(new AclEntryConfig(new AclClaimConfig(AdminClaimType, AdminClaimValue), actions));
-			}
-
+			defaultAcl.Entries.Add(new AclEntryConfig(new AclClaimConfig(ClaimTypes.Role, "internal:AgentRegistration"), new[] { AgentAclAction.CreateAgent, SessionAclAction.CreateSession }));
+			defaultAcl.Entries.Add(new AclEntryConfig(HordeClaims.AgentRegistrationClaim, new[] { AgentAclAction.CreateAgent, SessionAclAction.CreateSession, AgentAclAction.UpdateAgent, AgentSoftwareAclAction.DownloadSoftware, PoolAclAction.CreatePool, PoolAclAction.UpdatePool, PoolAclAction.ViewPool, PoolAclAction.DeletePool, PoolAclAction.ListPools, StreamAclAction.ViewStream, ProjectAclAction.ViewProject, JobAclAction.ViewJob, AdminAclAction.ViewCosts }));
+			defaultAcl.Entries.Add(new AclEntryConfig(HordeClaims.AgentRoleClaim, new[] { ProjectAclAction.ViewProject, StreamAclAction.ViewStream, LogAclAction.CreateEvent, AgentSoftwareAclAction.DownloadSoftware }));
+			defaultAcl.Entries.Add(new AclEntryConfig(HordeClaims.DownloadSoftwareClaim, new[] { AgentSoftwareAclAction.DownloadSoftware }));
+			defaultAcl.Entries.Add(new AclEntryConfig(HordeClaims.UploadSoftwareClaim, new[] { AgentSoftwareAclAction.UploadSoftware }));
+			defaultAcl.Entries.Add(new AclEntryConfig(HordeClaims.ConfigureProjectsClaim, new[] { ProjectAclAction.CreateProject, ProjectAclAction.UpdateProject, ProjectAclAction.ViewProject, StreamAclAction.CreateStream, StreamAclAction.UpdateStream, StreamAclAction.ViewStream }));
+			defaultAcl.Entries.Add(new AclEntryConfig(HordeClaims.StartChainedJobClaim, new[] { JobAclAction.CreateJob, JobAclAction.ExecuteJob, JobAclAction.UpdateJob, JobAclAction.ViewJob, StreamAclAction.ViewTemplate, StreamAclAction.ViewStream }));
 			return defaultAcl;
 		}
 
diff --git a/Engine/Source/Programs/Horde/Horde.Server/Startup.cs b/Engine/Source/Programs/Horde/Horde.Server/Startup.cs
--- a/Engine/Source/Programs/Horde/Horde.Server/Startup.cs
+++ b/Engine/Source/Programs/Horde/Horde.Server/Startup.cs
@@ -599,11 +599,7 @@ namespace Horde.Server
 			switch (settings.AuthMethod)
 			{
 				case AuthMethod.Anonymous:
-					authBuilder.AddAnonymous(options =>
-					{
-						options.AdminClaimType = settings.AdminClaimType;
-						options.AdminClaimValue = settings.AdminClaimValue;
-					});
+					authBuilder.AddAnonymous(options => { });
 					schemes.Add(AnonymousAuthenticationHandler.AuthenticationScheme);
 					break;
 						
@@ -971,6 +967,15 @@ namespace Horde.Server
 			}
 		}
 
+		public sealed class AclActionBsonSerializer : SerializerBase<AclAction>
+		{
+			/// <inheritdoc/>
+			public override AclAction Deserialize(BsonDeserializationContext context, BsonDeserializationArgs args) => new AclAction(context.Reader.ReadString());
+
+			/// <inheritdoc/>
+			public override void Serialize(BsonSerializationContext context, BsonSerializationArgs args, AclAction value) => context.Writer.WriteString(value.Name);
+		}
+
 		static int s_haveConfiguredMongoDb = 0;
 
 		public static void ConfigureMongoDbClient()
@@ -990,6 +995,7 @@ namespace Horde.Server
 				BsonSerializer.RegisterSerializer(new RefNameBsonSerializer());
 				BsonSerializer.RegisterSerializer(new IoHashBsonSerializer());
 				BsonSerializer.RegisterSerializer(new NamespaceIdBsonSerializer());
+				BsonSerializer.RegisterSerializer(new AclActionBsonSerializer());
 				BsonSerializer.RegisterSerializer(new ConditionSerializer());
 				BsonSerializer.RegisterSerializationProvider(new BsonSerializationProvider());
 				BsonSerializer.RegisterSerializationProvider(new StringIdBsonSerializationProvider());
diff --git a/Engine/Source/Programs/Horde/Horde.Server/Storage/StorageAclAction.cs b/Engine/Source/Programs/Horde/Horde.Server/Storage/StorageAclAction.cs
new file mode 100644
--- /dev/null
+++ b/Engine/Source/Programs/Horde/Horde.Server/Storage/StorageAclAction.cs
@@ -0,0 +1,37 @@
+// Copyright Epic Games, Inc. All Rights Reserved.
+
+using Horde.Server.Acls;
+
+namespace Horde.Server.Storage
+{
+	/// <summary>
+	/// ACL actions which apply to storage
+	/// </summary>
+	public static class StorageAclAction
+	{
+		/// <summary>
+		/// Ability to read blobs from the storage service
+		/// </summary>
+		public static AclAction ReadBlobs = new AclAction("ReadBlobs");
+
+		/// <summary>
+		/// Ability to write blobs to the storage service
+		/// </summary>
+		public static AclAction WriteBlobs = new AclAction("WriteBlobs");
+
+		/// <summary>
+		/// Ability to read refs from the storage service
+		/// </summary>
+		public static AclAction ReadRefs = new AclAction("ReadRefs");
+
+		/// <summary>
+		/// Ability to write refs to the storage service
+		/// </summary>
+		public static AclAction WriteRefs = new AclAction("WriteRefs");
+
+		/// <summary>
+		/// Ability to delete refs
+		/// </summary>
+		public static AclAction DeleteRefs = new AclAction("DeleteRefs");
+	}
+}
diff --git a/Engine/Source/Programs/Horde/Horde.Server/Storage/StorageController.cs b/Engine/Source/Programs/Horde/Horde.Server/Storage/StorageController.cs
--- a/Engine/Source/Programs/Horde/Horde.Server/Storage/StorageController.cs
+++ b/Engine/Source/Programs/Horde/Horde.Server/Storage/StorageController.cs
@@ -193,9 +193,9 @@ namespace Horde.Server.Storage
 			{
 				return NotFound(namespaceId);
 			}
-			if (!namespaceConfig.Authorize(AclAction.WriteBlobs, User) && !HasPathClaim(User, HordeClaimTypes.WriteNamespace, namespaceId, prefix ?? String.Empty))
+			if (!namespaceConfig.Authorize(StorageAclAction.WriteBlobs, User) && !HasPathClaim(User, HordeClaimTypes.WriteNamespace, namespaceId, prefix ?? String.Empty))
 			{
-				return Forbid(AclAction.WriteBlobs, namespaceId);
+				return Forbid(StorageAclAction.WriteBlobs, namespaceId);
 			}
 
 			IStorageClientImpl client = await _storageService.GetClientAsync(namespaceId, cancellationToken);
@@ -272,9 +272,9 @@ namespace Horde.Server.Storage
 			{
 				return NotFound(namespaceId);
 			}
-			if (!namespaceConfig.Authorize(AclAction.ReadBlobs, User) && !HasPathClaim(User, HordeClaimTypes.ReadNamespace, namespaceId, locator.Inner.ToString()))
+			if (!namespaceConfig.Authorize(StorageAclAction.ReadBlobs, User) && !HasPathClaim(User, HordeClaimTypes.ReadNamespace, namespaceId, locator.Inner.ToString()))
 			{
-				return Forbid(AclAction.ReadBlobs, namespaceId);
+				return Forbid(StorageAclAction.ReadBlobs, namespaceId);
 			}
 
 			return await ReadBlobInternalAsync(_storageService, namespaceId, locator, offset, length, cancellationToken);
@@ -337,9 +337,9 @@ namespace Horde.Server.Storage
 			{
 				return NotFound(namespaceId);
 			}
-			if (!namespaceConfig.Authorize(AclAction.ReadBlobs, User))
+			if (!namespaceConfig.Authorize(StorageAclAction.ReadBlobs, User))
 			{
-				return Forbid(AclAction.ReadBlobs, namespaceId);
+				return Forbid(StorageAclAction.ReadBlobs, namespaceId);
 			}
 
 			IStorageClientImpl client = await _storageService.GetClientAsync(namespaceId, cancellationToken);
@@ -374,9 +374,9 @@ namespace Horde.Server.Storage
 			{
 				return NotFound(namespaceId);
 			}
-			if (!namespaceConfig.Authorize(AclAction.WriteRefs, User) && !HasPathClaim(User, HordeClaimTypes.WriteNamespace, namespaceId, refName.ToString()))
+			if (!namespaceConfig.Authorize(StorageAclAction.WriteRefs, User) && !HasPathClaim(User, HordeClaimTypes.WriteNamespace, namespaceId, refName.ToString()))
 			{
-				return Forbid(AclAction.WriteRefs, namespaceId);
+				return Forbid(StorageAclAction.WriteRefs, namespaceId);
 			}
 
 			IStorageClient client = await _storageService.GetClientAsync(namespaceId, cancellationToken);
@@ -401,9 +401,9 @@ namespace Horde.Server.Storage
 			{
 				return NotFound(namespaceId);
 			}
-			if (!namespaceConfig.Authorize(AclAction.ReadRefs, User) && !HasPathClaim(User, HordeClaimTypes.ReadNamespace, namespaceId, refName.ToString()))
+			if (!namespaceConfig.Authorize(StorageAclAction.ReadRefs, User) && !HasPathClaim(User, HordeClaimTypes.ReadNamespace, namespaceId, refName.ToString()))
 			{
-				return Forbid(AclAction.ReadRefs, namespaceId);
+				return Forbid(StorageAclAction.ReadRefs, namespaceId);
 			}
 
 			return await ReadRefInternalAsync(_storageService, namespaceId, refName, cancellationToken);
@@ -484,9 +484,9 @@ namespace Horde.Server.Storage
 			{
 				return NotFound(namespaceId);
 			}
-			if (!namespaceConfig.Authorize(AclAction.ReadBlobs, User))
+			if (!namespaceConfig.Authorize(StorageAclAction.ReadBlobs, User))
 			{
-				return Forbid(AclAction.ReadBlobs, namespaceId);
+				return Forbid(StorageAclAction.ReadBlobs, namespaceId);
 			}
 
 			IStorageClient storageClient = await _storageService.GetClientAsync(namespaceId, cancellationToken);
@@ -566,9 +566,9 @@ namespace Horde.Server.Storage
 			{
 				return NotFound(namespaceId);
 			}
-			if (!namespaceConfig.Authorize(AclAction.ReadBlobs, User))
+			if (!namespaceConfig.Authorize(StorageAclAction.ReadBlobs, User))
 			{
-				return Forbid(AclAction.ReadBlobs, namespaceId);
+				return Forbid(StorageAclAction.ReadBlobs, namespaceId);
 			}
 
 			IStorageClient storageClient = await _storageService.GetClientAsync(namespaceId, cancellationToken);
diff --git a/Engine/Source/Programs/Horde/Horde.Server/Streams/StreamAclAction.cs b/Engine/Source/Programs/Horde/Horde.Server/Streams/StreamAclAction.cs
new file mode 100644
--- /dev/null
+++ b/Engine/Source/Programs/Horde/Horde.Server/Streams/StreamAclAction.cs
@@ -0,0 +1,42 @@
+// Copyright Epic Games, Inc. All Rights Reserved.
+
+using Horde.Server.Acls;
+
+namespace Horde.Server.Streams
+{
+	/// <summary>
+	/// ACL actions that can be performed on streams
+	/// </summary>
+	public static class StreamAclAction
+	{
+		/// <summary>
+		/// Allows the creation of new streams within a project
+		/// </summary>
+		public static AclAction CreateStream { get; } = new AclAction("CreateStream");
+
+		/// <summary>
+		/// Allows updating a stream (agent types, templates, schedules)
+		/// </summary>
+		public static AclAction UpdateStream { get; } = new AclAction("UpdateStream");
+
+		/// <summary>
+		/// Allows deleting a stream
+		/// </summary>
+		public static AclAction DeleteStream { get; } = new AclAction("DeleteStream");
+
+		/// <summary>
+		/// Ability to view a stream
+		/// </summary>
+		public static AclAction ViewStream { get; } = new AclAction("ViewStream");
+
+		/// <summary>
+		/// View changes submitted to a stream. NOTE: this returns responses from the server's Perforce account, which may be a priviledged user.
+		/// </summary>
+		public static AclAction ViewChanges { get; } = new AclAction("ViewChanges");
+
+		/// <summary>
+		/// View template associated with a stream
+		/// </summary>
+		public static AclAction ViewTemplate { get; } = new AclAction("ViewTemplate");
+	}
+}
diff --git a/Engine/Source/Programs/Horde/Horde.Server/Streams/StreamsController.cs b/Engine/Source/Programs/Horde/Horde.Server/Streams/StreamsController.cs
--- a/Engine/Source/Programs/Horde/Horde.Server/Streams/StreamsController.cs
+++ b/Engine/Source/Programs/Horde/Horde.Server/Streams/StreamsController.cs
@@ -74,7 +74,7 @@ namespace Horde.Server.Streams
 				{
 					foreach (StreamConfig streamConfig in projectConfig.Streams)
 					{
-						if (streamConfig.Authorize(AclAction.ViewStream, User))
+						if (streamConfig.Authorize(StreamAclAction.ViewStream, User))
 						{
 							IStream stream = await _streamCollection.GetAsync(streamConfig);
 							GetStreamResponse response = await CreateGetStreamResponseAsync(stream);
@@ -107,7 +107,7 @@ namespace Horde.Server.Streams
 				{
 					foreach (StreamConfig streamConfig in projectConfig.Streams)
 					{
-						if (streamConfig.Authorize(AclAction.ViewStream, User))
+						if (streamConfig.Authorize(StreamAclAction.ViewStream, User))
 						{
 							IStream stream = await _streamCollection.GetAsync(streamConfig);
 							GetStreamResponseV2 response = await CreateGetStreamResponseV2Async(stream, false);
@@ -136,9 +136,9 @@ namespace Horde.Server.Streams
 			{
 				return NotFound(streamId);
 			}
-			if (!streamConfig.Authorize(AclAction.ViewStream, User))
+			if (!streamConfig.Authorize(StreamAclAction.ViewStream, User))
 			{
-				return Forbid(AclAction.ViewStream, streamId);
+				return Forbid(StreamAclAction.ViewStream, streamId);
 			}
 
 			IStream stream = await _streamCollection.GetAsync(streamConfig);
@@ -162,9 +162,9 @@ namespace Horde.Server.Streams
 			{
 				return NotFound(streamId);
 			}
-			if (!streamConfig.Authorize(AclAction.ViewStream, User))
+			if (!streamConfig.Authorize(StreamAclAction.ViewStream, User))
 			{
-				return Forbid(AclAction.ViewStream, streamId);
+				return Forbid(StreamAclAction.ViewStream, streamId);
 			}
 
 			bool includeConfig = !String.Equals(config, streamConfig.Revision, StringComparison.Ordinal);
@@ -189,7 +189,7 @@ namespace Horde.Server.Streams
 				templateScope.Span.SetTag("templateName", pair.Value.Config.Name);
 
 				ITemplateRef templateRef = pair.Value;
-				if (templateRef.Config.Authorize(AclAction.ViewTemplate, User))
+				if (templateRef.Config.Authorize(StreamAclAction.ViewTemplate, User))
 				{
 					List<GetTemplateStepStateResponse>? stepStates = null;
 					if (templateRef.StepStates != null)
@@ -244,9 +244,9 @@ namespace Horde.Server.Streams
 			{
 				return NotFound(streamId);
 			}
-			if (!streamConfig.Authorize(AclAction.ViewChanges, User))
+			if (!streamConfig.Authorize(StreamAclAction.ViewChanges, User))
 			{
-				return Forbid(AclAction.ViewChanges, streamId);
+				return Forbid(StreamAclAction.ViewChanges, streamId);
 			}
 
 			List<CommitTag>? commitTags = null;
@@ -284,9 +284,9 @@ namespace Horde.Server.Streams
 			{
 				return NotFound(streamId);
 			}
-			if (!streamConfig.Authorize(AclAction.ViewChanges, User))
+			if (!streamConfig.Authorize(StreamAclAction.ViewChanges, User))
 			{
-				return Forbid(AclAction.ViewChanges, streamId);
+				return Forbid(StreamAclAction.ViewChanges, streamId);
 			}
 
 			ICommit? changeDetails = await _commitService.GetCollection(streamConfig).GetAsync(changeNumber);
@@ -322,9 +322,9 @@ namespace Horde.Server.Streams
 			{
 				return NotFound(streamId);
 			}
-			if (!streamConfig.Authorize(AclAction.ViewJob, User))
+			if (!streamConfig.Authorize(JobAclAction.ViewJob, User))
 			{
-				return Forbid(AclAction.ViewJob, streamId);
+				return Forbid(JobAclAction.ViewJob, streamId);
 			}
 
 			TemplateId templateIdValue = new TemplateId(templateId);
@@ -356,9 +356,9 @@ namespace Horde.Server.Streams
 			{
 				return NotFound(streamId, templateId);
 			}
-			if (!templateConfig.Authorize(AclAction.ViewTemplate, User))
+			if (!templateConfig.Authorize(StreamAclAction.ViewTemplate, User))
 			{
-				return Forbid(AclAction.ViewTemplate, streamId);
+				return Forbid(StreamAclAction.ViewTemplate, streamId);
 			}
 
 			ITemplate template = await _templateCollection.GetOrAddAsync(templateConfig);
@@ -373,9 +373,9 @@ namespace Horde.Server.Streams
 		[Route("/api/v1/streams/{streamId}/templates/{templateRefId}")]
 		public async Task<ActionResult> UpdateStreamTemplateRefAsync(StreamId streamId, TemplateId templateRefId, [FromBody] UpdateTemplateRefRequest update)
 		{
-			if (!_globalConfig.Value.Authorize(AclAction.AdminWrite, User))
+			if (!_globalConfig.Value.Authorize(AdminAclAction.AdminWrite, User))
 			{
-				return Forbid(AclAction.AdminWrite);
+				return Forbid(AdminAclAction.AdminWrite);
 			}
 
 			StreamConfig? streamConfig;
diff --git a/Engine/Source/Programs/Horde/Horde.Server/Tools/ToolAclAction.cs b/Engine/Source/Programs/Horde/Horde.Server/Tools/ToolAclAction.cs
new file mode 100644
--- /dev/null
+++ b/Engine/Source/Programs/Horde/Horde.Server/Tools/ToolAclAction.cs
@@ -0,0 +1,22 @@
+// Copyright Epic Games, Inc. All Rights Reserved.
+
+using Horde.Server.Acls;
+
+namespace Horde.Server.Tools
+{
+	/// <summary>
+	/// ACL actions which apply to tools
+	/// </summary>
+	public static class ToolAclAction
+	{
+		/// <summary>
+		/// Ability to download a tool
+		/// </summary>
+		public static AclAction DownloadTool { get; } = new AclAction("DownloadTool");
+
+		/// <summary>
+		/// Ability to upload new tool versions
+		/// </summary>
+		public static AclAction UploadTool { get; } = new AclAction("UploadTool");
+	}
+}
diff --git a/Engine/Source/Programs/Horde/Horde.Server/Tools/ToolsController.cs b/Engine/Source/Programs/Horde/Horde.Server/Tools/ToolsController.cs
--- a/Engine/Source/Programs/Horde/Horde.Server/Tools/ToolsController.cs
+++ b/Engine/Source/Programs/Horde/Horde.Server/Tools/ToolsController.cs
@@ -219,9 +219,9 @@ namespace Horde.Server.Tools
 			{
 				return NotFound(id);
 			}
-			if (!tool.Config.Authorize(AclAction.UploadTool, User))
+			if (!tool.Config.Authorize(ToolAclAction.UploadTool, User))
 			{
-				return Forbid(AclAction.UploadTool, id);
+				return Forbid(ToolAclAction.UploadTool, id);
 			}
 
 			using (Stream stream = file.OpenReadStream())
@@ -248,9 +248,9 @@ namespace Horde.Server.Tools
 			{
 				return NotFound(id);
 			}
-			if (!tool.Config.Authorize(AclAction.UploadTool, User))
+			if (!tool.Config.Authorize(ToolAclAction.UploadTool, User))
 			{
-				return Forbid(AclAction.UploadTool, id);
+				return Forbid(ToolAclAction.UploadTool, id);
 			}
 
 			if (request.State != null)
@@ -342,7 +342,7 @@ namespace Horde.Server.Tools
 			}
 			if (!AuthorizeDownload(tool.Config))
 			{
-				return Forbid(AclAction.DownloadTool, id);
+				return Forbid(ToolAclAction.DownloadTool, id);
 			}
 
 			if (action == GetToolAction.Info)
@@ -385,7 +385,7 @@ namespace Horde.Server.Tools
 			}
 			if (!AuthorizeDownload(tool.Config))
 			{
-				return Forbid(AclAction.DownloadTool, id);
+				return Forbid(ToolAclAction.DownloadTool, id);
 			}
 
 			IToolDeployment? deployment = tool.GetCurrentDeployment(phase, _clock.UtcNow);
@@ -412,7 +412,7 @@ namespace Horde.Server.Tools
 			}
 			if (!AuthorizeDownload(tool.Config))
 			{
-				return Forbid(AclAction.DownloadTool, id);
+				return Forbid(ToolAclAction.DownloadTool, id);
 			}
 
 			IToolDeployment? deployment = tool.Deployments.FirstOrDefault(x => x.Id == deploymentId);
@@ -482,7 +482,7 @@ namespace Horde.Server.Tools
 			}
 			if (!AuthorizeDownload(tool.Config))
 			{
-				return Forbid(AclAction.DownloadTool, id);
+				return Forbid(ToolAclAction.DownloadTool, id);
 			}
 
 			if (!locator.BlobId.WithinFolder(tool.Id.Id.Text))
@@ -502,7 +502,7 @@ namespace Horde.Server.Tools
 				{
 					return false;
 				}
-				if (!toolConfig.Authorize(AclAction.DownloadTool, User))
+				if (!toolConfig.Authorize(ToolAclAction.DownloadTool, User))
 				{
 					return false;
 				}
diff --git a/Engine/Source/Programs/Horde/Horde.Server/Utilities/HordeClaimTypes.cs b/Engine/Source/Programs/Horde/Horde.Server/Utilities/HordeClaimTypes.cs
--- a/Engine/Source/Programs/Horde/Horde.Server/Utilities/HordeClaimTypes.cs
+++ b/Engine/Source/Programs/Horde/Horde.Server/Utilities/HordeClaimTypes.cs
@@ -21,11 +21,6 @@ namespace Horde.Server.Utilities
 		/// </summary>
 		public const string Role = Prefix + "role";
 
-		/// <summary>
-		/// Claim for a well known internal role (eg. admin)
-		/// </summary>
-		public const string InternalRole = Prefix + "internal-role";
-
 		/// <summary>
 		/// Claim for a particular session
 		/// </summary>
diff --git a/Engine/Source/Programs/Horde/Horde.Server/Utilities/HordeClaims.cs b/Engine/Source/Programs/Horde/Horde.Server/Utilities/HordeClaims.cs
--- a/Engine/Source/Programs/Horde/Horde.Server/Utilities/HordeClaims.cs
+++ b/Engine/Source/Programs/Horde/Horde.Server/Utilities/HordeClaims.cs
@@ -11,6 +11,11 @@ namespace Horde.Server.Utilities
 	/// </summary>
 	public static class HordeClaims
 	{
+		/// <summary>
+		/// Claim which gives unfettered access to all Horde endpoints
+		/// </summary>
+		public static AclClaimConfig AdminClaim { get; } = new AclClaimConfig(HordeClaimTypes.Role, "admin");
+
 		/// <summary>
 		/// Name of the role that can be used to administer agents
 		/// </summary>
