Horde: Move AclScopeName into a separate file.

diff --git a/Engine/Source/Programs/Horde/Horde.Server/Acls/AclScopeName.cs b/Engine/Source/Programs/Horde/Horde.Server/Acls/AclScopeName.cs
new file mode 100644
--- /dev/null
+++ b/Engine/Source/Programs/Horde/Horde.Server/Acls/AclScopeName.cs
@@ -0,0 +1,56 @@
+// Copyright Epic Games, Inc. All Rights Reserved.
+
+using System;
+using System.Diagnostics;
+using System.Text.Json;
+using System.Text.Json.Serialization;
+using MongoDB.Bson.Serialization.Attributes;
+
+namespace Horde.Server.Acls
+{
+	/// <summary>
+	/// Name of an ACL scope
+	/// </summary>
+	[DebuggerDisplay("{Name}")]
+	[JsonConverter(typeof(AclScopeNameJsonConverter))]
+	public record struct AclScopeName(string Text)
+	{
+		/// <summary>
+		/// The root scope name
+		/// </summary>
+		public static AclScopeName Root { get; } = new AclScopeName("horde");
+
+		/// <summary>
+		/// Append another name to this scope
+		/// </summary>
+		/// <param name="name">Name to append</param>
+		/// <returns>New scope name</returns>
+		public AclScopeName Append(string name) => new AclScopeName($"{Text}/{name}");
+
+		/// <summary>
+		/// Append another name to this scope
+		/// </summary>
+		/// <param name="type">Type of the scope</param>
+		/// <param name="name">Name to append</param>
+		/// <returns>New scope name</returns>
+		public AclScopeName Append(string type, string name) => new AclScopeName($"{Text}/{type}:{name}");
+
+		/// <inheritdoc/>
+		public bool Equals(AclScopeName other) => Text.Equals(other.Text, StringComparison.Ordinal);
+
+		/// <inheritdoc/>
+		public override int GetHashCode() => Text.GetHashCode(StringComparison.Ordinal);
+	}
+
+	/// <summary>
+	/// Serializes <see cref="AclScopeName"/> objects to JSON
+	/// </summary>
+	class AclScopeNameJsonConverter : JsonConverter<AclScopeName>
+	{
+		/// <inheritdoc/>
+		public override AclScopeName Read(ref Utf8JsonReader reader, Type typeToConvert, JsonSerializerOptions options) => new AclScopeName(reader.GetString() ?? String.Empty);
+
+		/// <inheritdoc/>
+		public override void Write(Utf8JsonWriter writer, AclScopeName value, JsonSerializerOptions options) => writer.WriteStringValue(value.Text);
+	}
+}
diff --git a/Engine/Source/Programs/Horde/Horde.Server/Acls/IAclScope.cs b/Engine/Source/Programs/Horde/Horde.Server/Acls/IAclScope.cs
--- a/Engine/Source/Programs/Horde/Horde.Server/Acls/IAclScope.cs
+++ b/Engine/Source/Programs/Horde/Horde.Server/Acls/IAclScope.cs
@@ -32,65 +32,6 @@ namespace Horde.Server.Acls
 		AclConfig? Acl { get; }
 	}
 
-	/// <summary>
-	/// Name of an ACL scope
-	/// </summary>
-	[DebuggerDisplay("{Name}")]
-	[JsonConverter(typeof(AclScopeNameJsonConverter))]
-	[BsonSerializer(typeof(AclScopeNameBsonSerializer))]
-	public record struct AclScopeName(string Text)
-	{
-		/// <summary>
-		/// The root scope name
-		/// </summary>
-		public static AclScopeName Root { get; } = new AclScopeName("horde");
-
-		/// <summary>
-		/// Append another name to this scope
-		/// </summary>
-		/// <param name="name">Name to append</param>
-		/// <returns>New scope name</returns>
-		public AclScopeName Append(string name) => new AclScopeName($"{Text}/{name}");
-
-		/// <summary>
-		/// Append another name to this scope
-		/// </summary>
-		/// <param name="type">Type of the scope</param>
-		/// <param name="name">Name to append</param>
-		/// <returns>New scope name</returns>
-		public AclScopeName Append(string type, string name) => new AclScopeName($"{Text}/{type}:{name}");
-
-		/// <inheritdoc/>
-		public bool Equals(AclScopeName other) => Text.Equals(other.Text, StringComparison.Ordinal);
-
-		/// <inheritdoc/>
-		public override int GetHashCode() => Text.GetHashCode(StringComparison.Ordinal);
-	}
-
-	/// <summary>
-	/// Serializes <see cref="AclScopeName"/> objects to JSON
-	/// </summary>
-	class AclScopeNameJsonConverter : JsonConverter<AclScopeName>
-	{
-		/// <inheritdoc/>
-		public override AclScopeName Read(ref Utf8JsonReader reader, Type typeToConvert, JsonSerializerOptions options) => new AclScopeName(reader.GetString() ?? String.Empty);
-
-		/// <inheritdoc/>
-		public override void Write(Utf8JsonWriter writer, AclScopeName value, JsonSerializerOptions options) => writer.WriteStringValue(value.Text);
-	}
-
-	/// <summary>
-	/// Serializes <see cref="AclScopeName"/> objects to JSON
-	/// </summary>
-	class AclScopeNameBsonSerializer : SerializerBase<AclScopeName> 
-	{
-		/// <inheritdoc/>
-		public override AclScopeName Deserialize(BsonDeserializationContext context, BsonDeserializationArgs args) => new AclScopeName(context.Reader.ReadString());
-
-		/// <inheritdoc/>
-		public override void Serialize(BsonSerializationContext context, BsonSerializationArgs args, AclScopeName value) => context.Writer.WriteString(value.Text);
-	}
-
 	/// <summary>
 	/// Extension methods for ACL scopes
 	/// </summary>
diff --git a/Engine/Source/Programs/Horde/Horde.Server/Startup.cs b/Engine/Source/Programs/Horde/Horde.Server/Startup.cs
--- a/Engine/Source/Programs/Horde/Horde.Server/Startup.cs
+++ b/Engine/Source/Programs/Horde/Horde.Server/Startup.cs
@@ -976,6 +976,15 @@ namespace Horde.Server
 			public override void Serialize(BsonSerializationContext context, BsonSerializationArgs args, AclAction value) => context.Writer.WriteString(value.Name);
 		}
 
+		public sealed class AclScopeNameBsonSerializer : SerializerBase<AclScopeName>
+		{
+			/// <inheritdoc/>
+			public override AclScopeName Deserialize(BsonDeserializationContext context, BsonDeserializationArgs args) => new AclScopeName(context.Reader.ReadString());
+
+			/// <inheritdoc/>
+			public override void Serialize(BsonSerializationContext context, BsonSerializationArgs args, AclScopeName value) => context.Writer.WriteString(value.Text);
+		}
+
 		static int s_haveConfiguredMongoDb = 0;
 
 		public static void ConfigureMongoDbClient()
@@ -996,6 +1005,7 @@ namespace Horde.Server
 				BsonSerializer.RegisterSerializer(new IoHashBsonSerializer());
 				BsonSerializer.RegisterSerializer(new NamespaceIdBsonSerializer());
 				BsonSerializer.RegisterSerializer(new AclActionBsonSerializer());
+				BsonSerializer.RegisterSerializer(new AclScopeNameBsonSerializer());
 				BsonSerializer.RegisterSerializer(new ConditionSerializer());
 				BsonSerializer.RegisterSerializationProvider(new BsonSerializationProvider());
 				BsonSerializer.RegisterSerializationProvider(new StringIdBsonSerializationProvider());
