Horde: Add a service which manages IStorageBackend instances (IStorageBackendProvider). The objects it returns are IDisposable wrappers around an underlying ref-counted IStorageBackend implementation, allowing different clients to manage lifetime of them using standard patterns.

diff --git a/Engine/Source/Programs/Horde/Horde.Server/Startup.cs b/Engine/Source/Programs/Horde/Horde.Server/Startup.cs
--- a/Engine/Source/Programs/Horde/Horde.Server/Startup.cs
+++ b/Engine/Source/Programs/Horde/Horde.Server/Startup.cs
@@ -519,6 +519,7 @@ namespace Horde.Server
 			}
 
 			// Storage providers
+			services.AddSingleton<IStorageBackendProvider, StorageBackendProvider>();
 			services.AddSingleton(sp => CreateStorageBackend(sp, settings.LogStorage).ForType<PersistentLogStorage>());
 			services.AddSingleton(sp => CreateStorageBackend(sp, settings.ArtifactStorage).ForType<ArtifactCollectionV1>());
 
diff --git a/Engine/Source/Programs/Horde/Horde.Server/Storage/Backends/FileStorageBackend.cs b/Engine/Source/Programs/Horde/Horde.Server/Storage/Backends/FileStorageBackend.cs
--- a/Engine/Source/Programs/Horde/Horde.Server/Storage/Backends/FileStorageBackend.cs
+++ b/Engine/Source/Programs/Horde/Horde.Server/Storage/Backends/FileStorageBackend.cs
@@ -27,7 +27,7 @@ namespace Horde.Server.Storage.Backends
 	/// <summary>
 	/// Storage backend using the filesystem
 	/// </summary>
-	public class FileSystemStorageBackend : IStorageBackend
+	public sealed class FileSystemStorageBackend : IStorageBackend
 	{
 		/// <summary>
 		/// Base directory for log files
@@ -55,6 +55,11 @@ namespace Horde.Server.Storage.Backends
 			DirectoryReference.CreateDirectory(_baseDir);
 		}
 
+		/// <inheritdoc/>
+		public void Dispose()
+		{
+		}
+
 		/// <inheritdoc/>
 		public Task<Stream?> TryReadAsync(string path, CancellationToken cancellationToken)
 		{
diff --git a/Engine/Source/Programs/Horde/Horde.Server/Storage/Backends/MemoryStorageBackend.cs b/Engine/Source/Programs/Horde/Horde.Server/Storage/Backends/MemoryStorageBackend.cs
--- a/Engine/Source/Programs/Horde/Horde.Server/Storage/Backends/MemoryStorageBackend.cs
+++ b/Engine/Source/Programs/Horde/Horde.Server/Storage/Backends/MemoryStorageBackend.cs
@@ -12,13 +12,18 @@ namespace Horde.Server.Storage.Backends
 	/// <summary>
 	/// In-memory implementation of ILogFileStorage
 	/// </summary>
-	public class MemoryStorageBackend : IStorageBackend
+	public sealed class MemoryStorageBackend : IStorageBackend
 	{
 		/// <summary>
 		/// Data storage
 		/// </summary>
 		readonly ConcurrentDictionary<string, byte[]> _pathToData = new ConcurrentDictionary<string, byte[]>();
 
+		/// <inheritdoc/>
+		public void Dispose()
+		{
+		}
+
 		/// <inheritdoc/>
 		public Task<Stream?> TryReadAsync(string path, CancellationToken cancellationToken)
 		{
diff --git a/Engine/Source/Programs/Horde/Horde.Server/Storage/IStorageBackend.cs b/Engine/Source/Programs/Horde/Horde.Server/Storage/IStorageBackend.cs
--- a/Engine/Source/Programs/Horde/Horde.Server/Storage/IStorageBackend.cs
+++ b/Engine/Source/Programs/Horde/Horde.Server/Storage/IStorageBackend.cs
@@ -13,7 +13,7 @@ namespace Horde.Server.Storage
 	/// <summary>
 	/// Interface for a traditional location-addressed storage provider
 	/// </summary>
-	public interface IStorageBackend
+	public interface IStorageBackend : IDisposable
 	{
 		/// <summary>
 		/// Attempts to open a read stream for the given path.
@@ -104,7 +104,7 @@ namespace Horde.Server.Storage
 		/// Wrapper for <see cref="IStorageBackend"/>
 		/// </summary>
 		/// <typeparam name="T"></typeparam>
-		class TypedStorageBackend<T> : IStorageBackend<T>
+		sealed class TypedStorageBackend<T> : IStorageBackend<T>
 		{
 			readonly IStorageBackend _inner;
 
@@ -114,6 +114,9 @@ namespace Horde.Server.Storage
 			/// <param name="inner"></param>
 			public TypedStorageBackend(IStorageBackend inner) => _inner = inner;
 
+			/// <inheritdoc/>
+			public void Dispose() { }
+
 			/// <inheritdoc/>
 			public Task<Stream?> TryReadAsync(string path, CancellationToken cancellationToken) => _inner.TryReadAsync(path, cancellationToken);
 
diff --git a/Engine/Source/Programs/Horde/Horde.Server/Storage/IStorageBackendProvider.cs b/Engine/Source/Programs/Horde/Horde.Server/Storage/IStorageBackendProvider.cs
new file mode 100644
--- /dev/null
+++ b/Engine/Source/Programs/Horde/Horde.Server/Storage/IStorageBackendProvider.cs
@@ -0,0 +1,17 @@
+// Copyright Epic Games, Inc. All Rights Reserved.
+
+namespace Horde.Server.Storage
+{
+	/// <summary>
+	/// Service which resolves storage backend instances
+	/// </summary>
+	public interface IStorageBackendProvider
+	{
+		/// <summary>
+		/// Creates a new storage backend with the given identifier
+		/// </summary>
+		/// <param name="backendConfig">Configuration object for the storage backend</param>
+		/// <returns>New backend instance. This instance should be disposed when the client has finished with it.</returns>
+		IStorageBackend CreateBackend(BackendConfig backendConfig);
+	}
+}
diff --git a/Engine/Source/Programs/Horde/Horde.Server/Storage/StorageBackendProvider.cs b/Engine/Source/Programs/Horde/Horde.Server/Storage/StorageBackendProvider.cs
new file mode 100644
--- /dev/null
+++ b/Engine/Source/Programs/Horde/Horde.Server/Storage/StorageBackendProvider.cs
@@ -0,0 +1,166 @@
+// Copyright Epic Games, Inc. All Rights Reserved.
+
+using System;
+using System.Collections.Generic;
+using System.IO;
+using System.Text.Json;
+using System.Threading;
+using System.Threading.Tasks;
+using EpicGames.Core;
+using Horde.Server.Storage.Backends;
+using Microsoft.Extensions.Configuration;
+using Microsoft.Extensions.DependencyInjection;
+using Microsoft.Extensions.Logging;
+
+namespace Horde.Server.Storage
+{
+	/// <summary>
+	/// Default implementation of <see cref="IStorageBackendProvider"/>
+	/// </summary>
+	class StorageBackendProvider : IStorageBackendProvider
+	{
+		class RefCountedBackend : IDisposable
+		{
+			public IoHash Hash { get; }
+			public IStorageBackend Backend { get; }
+
+			public int _refCount = 1;
+
+			public RefCountedBackend(IoHash hash, IStorageBackend backend)
+			{
+				Hash = hash;
+				Backend = backend;
+			}
+
+			public void Dispose()
+			{
+				Backend.Dispose();
+			}
+		}
+
+		class BackendWrapper : IStorageBackend
+		{
+			readonly StorageBackendProvider _owner;
+			RefCountedBackend _refCountedBackend;
+			IStorageBackend _backend;
+
+			public BackendWrapper(StorageBackendProvider owner, RefCountedBackend refCountedBackend)
+			{
+				_owner = owner;
+				_refCountedBackend = refCountedBackend;
+				_backend = _refCountedBackend.Backend;
+			}
+
+			public void Dispose()
+			{
+				if (_refCountedBackend != null)
+				{
+					_owner.ReleaseBackend(_refCountedBackend);
+					_refCountedBackend = null!;
+
+					_backend = null!;
+				}
+			}
+
+			#region IStorageBackend Implementation
+
+			/// <inheritdoc/>
+			public Task<bool> ExistsAsync(string path, CancellationToken cancellationToken = default) => _backend.ExistsAsync(path, cancellationToken);
+
+			/// <inheritdoc/>
+			public Task DeleteAsync(string path, CancellationToken cancellationToken = default) => _backend.DeleteAsync(path, cancellationToken);
+
+			/// <inheritdoc/>
+			public IAsyncEnumerable<string> EnumerateAsync(CancellationToken cancellationToken = default) => _backend.EnumerateAsync(cancellationToken);
+
+			/// <inheritdoc/>
+			public Task<Stream?> TryReadAsync(string path, CancellationToken cancellationToken = default) => _backend.TryReadAsync(path, cancellationToken);
+
+			/// <inheritdoc/>
+			public Task<Stream?> TryReadAsync(string path, int offset, int length, CancellationToken cancellationToken = default) => _backend.TryReadAsync(path, offset, length, cancellationToken);
+
+			/// <inheritdoc/>
+			public Task WriteAsync(string path, Stream stream, CancellationToken cancellationToken = default) => _backend.WriteAsync(path, stream, cancellationToken);
+
+			#endregion
+		}
+
+		readonly IServiceProvider _serviceProvider;
+		readonly object _lockObject = new object();
+		readonly Dictionary<IoHash, RefCountedBackend> _backends = new Dictionary<IoHash, RefCountedBackend>();
+
+		/// <summary>
+		/// Constructor
+		/// </summary>
+		public StorageBackendProvider(IServiceProvider serviceProvider)
+		{
+			_serviceProvider = serviceProvider;
+		}
+
+		/// <inheritdoc/>
+		public IStorageBackend CreateBackend(BackendConfig config)
+		{
+			// Compute the new hash of the configuration data
+			IoHash hash;
+			using (MemoryStream stream = new MemoryStream())
+			{
+				JsonSerializerOptions options = new JsonSerializerOptions();
+				Startup.ConfigureJsonSerializer(options);
+
+				JsonSerializer.Serialize(stream, config, options: options);
+
+				hash = IoHash.Compute(stream.ToArray());
+			}
+
+			// See if we've got an existing backend we can use
+			RefCountedBackend? refCountedBackend;
+			lock (_lockObject)
+			{
+				if (_backends.TryGetValue(hash, out refCountedBackend))
+				{
+					refCountedBackend._refCount++;
+				}
+				else
+				{
+					IStorageBackend newBackend = CreateStorageBackend(config);
+					refCountedBackend = new RefCountedBackend(hash, newBackend);
+					_backends.Add(hash, refCountedBackend);
+				}
+			}
+
+			return new BackendWrapper(this, refCountedBackend);
+		}
+
+		void ReleaseBackend(RefCountedBackend backend)
+		{
+			lock (_lockObject)
+			{
+				if (--backend._refCount == 0)
+				{
+					_backends.Remove(backend.Hash);
+					backend.Backend.Dispose();
+				}
+			}
+		}
+
+		/// <summary>
+		/// Creates a storage backend with the given configuration
+		/// </summary>
+		/// <param name="config">Configuration for the backend</param>
+		/// <returns>New storage backend instance</returns>
+		IStorageBackend CreateStorageBackend(BackendConfig config)
+		{
+			switch (config.Type ?? StorageBackendType.FileSystem)
+			{
+				case StorageBackendType.FileSystem:
+					return new FileSystemStorageBackend(config);
+				case StorageBackendType.Aws:
+					return new AwsStorageBackend(_serviceProvider.GetRequiredService<IConfiguration>(), config, _serviceProvider.GetRequiredService<ILogger<AwsStorageBackend>>());
+				case StorageBackendType.Memory:
+					return new MemoryStorageBackend();
+				default:
+					throw new NotImplementedException();
+			}
+		}
+	}
+}
diff --git a/Engine/Source/Programs/Horde/Horde.Server/Storage/StorageService.cs b/Engine/Source/Programs/Horde/Horde.Server/Storage/StorageService.cs
--- a/Engine/Source/Programs/Horde/Horde.Server/Storage/StorageService.cs
+++ b/Engine/Source/Programs/Horde/Horde.Server/Storage/StorageService.cs
@@ -2,26 +2,18 @@
 
 using System;
 using System.Buffers;
-using System.Buffers.Binary;
 using System.Collections;
 using System.Collections.Generic;
-using System.Composition;
 using System.IO;
 using System.Linq;
-using System.Reflection;
 using System.Runtime.CompilerServices;
-using System.Runtime.InteropServices;
-using System.Security.Claims;
-using System.Text.Json;
 using System.Threading;
 using System.Threading.Tasks;
 using EpicGames.Core;
 using EpicGames.Horde.Storage;
 using EpicGames.Redis;
 using EpicGames.Redis.Utility;
-using Horde.Server.Acls;
 using Horde.Server.Server;
-using Horde.Server.Storage.Backends;
 using Horde.Server.Utilities;
 using HordeCommon;
 using Microsoft.Extensions.Caching.Memory;
@@ -260,42 +252,20 @@ namespace Horde.Server.Storage
 		}
 
 
-		class BackendInfo
-		{
-			public IStorageBackend Backend { get; }
-
-			int _refCount = 1;
-
-			public BackendInfo(IStorageBackend backend)
-			{
-				Backend = backend;
-			}
-
-			public void AddRef() => Interlocked.Increment(ref _refCount);
-
-			public void Release()
-			{
-				if (Interlocked.Decrement(ref _refCount) == 0 && Backend is IDisposable disposable)
-				{
-					disposable.Dispose();
-				}
-			}
-		}
-
 		class NamespaceInfo : IDisposable
 		{
 			public NamespaceConfig Config { get; }
 			public StorageClient Client { get; }
-			public BackendInfo Backend { get; }
+			public IStorageBackend Backend { get; }
 
-			public NamespaceInfo(NamespaceConfig config, StorageClient client, BackendInfo backend)
+			public NamespaceInfo(NamespaceConfig config, StorageClient client, IStorageBackend backend)
 			{
 				Config = config;
 				Client = client;
 				Backend = backend;
 			}
 
-			public void Dispose() => Backend.Release();
+			public void Dispose() => Backend.Dispose();
 		}
 
 		class State : IDisposable
@@ -455,7 +425,7 @@ namespace Horde.Server.Storage
 		readonly RedisService _redisService;
 		readonly IClock _clock;
 		readonly IMemoryCache _cache;
-		readonly IServiceProvider _serviceProvider;
+		readonly IStorageBackendProvider _storageBackendProvider;
 		readonly IOptionsMonitor<GlobalConfig> _globalConfig;
 		readonly ILogger _logger;
 
@@ -470,19 +440,18 @@ namespace Horde.Server.Storage
 
 		State? _lastState;
 		string? _lastConfigRevision;
-		readonly Dictionary<IoHash, BackendInfo> _backends = new Dictionary<IoHash, BackendInfo>();
 
 		readonly AsyncCachedValue<State> _cachedState;
 
 		/// <summary>
 		/// Constructor
 		/// </summary>
-		public StorageService(MongoService mongoService, RedisService redisService, IClock clock, IMemoryCache cache, IServiceProvider serviceProvider, IOptionsMonitor<GlobalConfig> globalConfig, ILogger<StorageService> logger)
+		public StorageService(MongoService mongoService, RedisService redisService, IClock clock, IMemoryCache cache, IStorageBackendProvider storageBackendProvider, IOptionsMonitor<GlobalConfig> globalConfig, ILogger<StorageService> logger)
 		{
 			_redisService = redisService;
 			_clock = clock;
 			_cache = cache;
-			_serviceProvider = serviceProvider;
+			_storageBackendProvider = storageBackendProvider;
 			_cachedState = new AsyncCachedValue<State>(() => Task.FromResult(GetNextState()), TimeSpan.FromMinutes(1.0));
 			_globalConfig = globalConfig;
 			_logger = logger;
@@ -605,8 +574,8 @@ namespace Horde.Server.Storage
 				{
 					foreach (NamespaceConfig namespaceConfig in storageConfig.Namespaces)
 					{
-						BackendInfo backend = FindOrAddStorageBackend(namespaceConfig.BackendConfig);
-						StorageClient client = new StorageClient(this, namespaceConfig, backend.Backend);
+						IStorageBackend backend = _storageBackendProvider.CreateBackend(namespaceConfig.BackendConfig);
+						StorageClient client = new StorageClient(this, namespaceConfig, backend);
 						nextState.Namespaces.Add(namespaceConfig.Id, new NamespaceInfo(namespaceConfig, client, backend));
 					}
 				}
@@ -624,58 +593,6 @@ namespace Horde.Server.Storage
 			return _lastState;
 		}
 
-		/// <summary>
-		/// Creates a storage backend with the given configuration
-		/// </summary>
-		/// <param name="config">Configuration for the backend</param>
-		/// <returns>New storage backend instance</returns>
-		BackendInfo FindOrAddStorageBackend(BackendConfig config)
-		{
-			// Compute the new hash of the configuration data
-			IoHash hash;
-			using (MemoryStream stream = new MemoryStream())
-			{
-				JsonSerializerOptions options = new JsonSerializerOptions();
-				Startup.ConfigureJsonSerializer(options);
-
-				JsonSerializer.Serialize(stream, config, options: options);
-
-				hash = IoHash.Compute(stream.ToArray());
-			}
-
-			// See if we've got an existing backend we can use
-			if (_backends.TryGetValue(hash, out BackendInfo? existingBackendInfo))
-			{
-				existingBackendInfo.AddRef();
-				return existingBackendInfo;
-			}
-
-			// Create a new backend
-			BackendInfo backendInfo = new BackendInfo(CreateStorageBackend(config));
-			_backends.Add(hash, backendInfo);
-			return backendInfo;
-		}
-
-		/// <summary>
-		/// Creates a storage backend with the given configuration
-		/// </summary>
-		/// <param name="config">Configuration for the backend</param>
-		/// <returns>New storage backend instance</returns>
-		IStorageBackend CreateStorageBackend(BackendConfig config)
-		{
-			switch (config.Type ?? StorageBackendType.FileSystem)
-			{
-				case StorageBackendType.FileSystem:
-					return new FileSystemStorageBackend(config);
-				case StorageBackendType.Aws:
-					return new AwsStorageBackend(_serviceProvider.GetRequiredService<IConfiguration>(), config, _serviceProvider.GetRequiredService<ILogger<AwsStorageBackend>>());
-				case StorageBackendType.Memory:
-					return new MemoryStorageBackend();
-				default:
-					throw new NotImplementedException();
-			}
-		}
-
 		#endregion
 
 		#region Blobs
