Horde: Merge support for read/write redirects into the default IStorageBackend interface. Implementations that don't support redirects can just return default values.

diff --git a/Engine/Source/Programs/Horde/Horde.Server.Tests/TestSetup.cs b/Engine/Source/Programs/Horde/Horde.Server.Tests/TestSetup.cs
--- a/Engine/Source/Programs/Horde/Horde.Server.Tests/TestSetup.cs
+++ b/Engine/Source/Programs/Horde/Horde.Server.Tests/TestSetup.cs
@@ -261,6 +261,7 @@ namespace Horde.Server.Tests
 			services.AddSingleton<ConformTaskSource>();
 			services.AddSingleton<ICommitService, CommitService>();
 
+			services.AddSingleton<IStorageBackendProvider, StorageBackendProvider>();
 			services.AddSingleton<IStorageBackend<PersistentLogStorage>>(sp => new MemoryStorageBackend().ForType<PersistentLogStorage>());
 			services.AddSingleton<IStorageBackend<ArtifactCollectionV1>>(sp => new MemoryStorageBackend().ForType<ArtifactCollectionV1>());
 
diff --git a/Engine/Source/Programs/Horde/Horde.Server/Storage/Backends/AwsStorageBackend.cs b/Engine/Source/Programs/Horde/Horde.Server/Storage/Backends/AwsStorageBackend.cs
--- a/Engine/Source/Programs/Horde/Horde.Server/Storage/Backends/AwsStorageBackend.cs
+++ b/Engine/Source/Programs/Horde/Horde.Server/Storage/Backends/AwsStorageBackend.cs
@@ -103,7 +103,7 @@ namespace Horde.Server.Storage.Backends
 	/// <summary>
 	/// FileStorage implementation using an s3 bucket
 	/// </summary>
-	public sealed class AwsStorageBackend : IStorageBackendWithRedirects, IDisposable
+	public sealed class AwsStorageBackend : IStorageBackend, IDisposable
 	{
 		/// <summary>
 		/// S3 Client
@@ -130,6 +130,9 @@ namespace Horde.Server.Storage.Backends
 		/// </summary>
 		private readonly ILogger _logger;
 
+		/// <inheritdoc/>
+		public bool SupportsRedirects => true;
+
 		/// <summary>
 		/// Constructor
 		/// </summary>
@@ -300,10 +303,10 @@ namespace Horde.Server.Storage.Backends
 		}
 
 		/// <inheritdoc/>
-		public ValueTask<Uri?> GetReadRedirectAsync(string path, CancellationToken cancellationToken = default) => new ValueTask<Uri?>(GetPresignedUrl(path, HttpVerb.GET));
+		public ValueTask<Uri?> TryGetReadRedirectAsync(string path, CancellationToken cancellationToken = default) => new ValueTask<Uri?>(GetPresignedUrl(path, HttpVerb.GET));
 
 		/// <inheritdoc/>
-		public ValueTask<Uri?> GetWriteRedirectAsync(string path, CancellationToken cancellationToken = default) => new ValueTask<Uri?>(GetPresignedUrl(path, HttpVerb.PUT));
+		public ValueTask<Uri?> TryGetWriteRedirectAsync(string path, CancellationToken cancellationToken = default) => new ValueTask<Uri?>(GetPresignedUrl(path, HttpVerb.PUT));
 
 		/// <summary>
 		/// Helper method to generate a presigned URL for a request
diff --git a/Engine/Source/Programs/Horde/Horde.Server/Storage/Backends/FileStorageBackend.cs b/Engine/Source/Programs/Horde/Horde.Server/Storage/Backends/FileStorageBackend.cs
--- a/Engine/Source/Programs/Horde/Horde.Server/Storage/Backends/FileStorageBackend.cs
+++ b/Engine/Source/Programs/Horde/Horde.Server/Storage/Backends/FileStorageBackend.cs
@@ -44,6 +44,9 @@ namespace Horde.Server.Storage.Backends
 		/// </summary>
 		private int _uniqueId;
 
+		/// <inheritdoc/>
+		public bool SupportsRedirects => false;
+
 		/// <summary>
 		/// Constructor
 		/// </summary>
@@ -195,5 +198,11 @@ namespace Horde.Server.Storage.Backends
 				}
 			}
 		}
+
+		/// <inheritdoc/>
+		public ValueTask<Uri?> TryGetReadRedirectAsync(string path, CancellationToken cancellationToken = default) => default;
+
+		/// <inheritdoc/>
+		public ValueTask<Uri?> TryGetWriteRedirectAsync(string path, CancellationToken cancellationToken = default) => default;
 	}
 }
diff --git a/Engine/Source/Programs/Horde/Horde.Server/Storage/Backends/MemoryStorageBackend.cs b/Engine/Source/Programs/Horde/Horde.Server/Storage/Backends/MemoryStorageBackend.cs
--- a/Engine/Source/Programs/Horde/Horde.Server/Storage/Backends/MemoryStorageBackend.cs
+++ b/Engine/Source/Programs/Horde/Horde.Server/Storage/Backends/MemoryStorageBackend.cs
@@ -1,5 +1,6 @@
 // Copyright Epic Games, Inc. All Rights Reserved.
 
+using System;
 using System.Collections.Concurrent;
 using System.Collections.Generic;
 using System.IO;
@@ -19,6 +20,9 @@ namespace Horde.Server.Storage.Backends
 		/// </summary>
 		readonly ConcurrentDictionary<string, byte[]> _pathToData = new ConcurrentDictionary<string, byte[]>();
 
+		/// <inheritdoc/>
+		public bool SupportsRedirects => false;
+
 		/// <inheritdoc/>
 		public void Dispose()
 		{
@@ -82,5 +86,11 @@ namespace Horde.Server.Storage.Backends
 				await Task.Yield();
 			}
 		}
+
+		/// <inheritdoc/>
+		public ValueTask<Uri?> TryGetReadRedirectAsync(string path, CancellationToken cancellationToken = default) => default;
+
+		/// <inheritdoc/>
+		public ValueTask<Uri?> TryGetWriteRedirectAsync(string path, CancellationToken cancellationToken = default) => default;
 	}
 }
diff --git a/Engine/Source/Programs/Horde/Horde.Server/Storage/IStorageBackend.cs b/Engine/Source/Programs/Horde/Horde.Server/Storage/IStorageBackend.cs
--- a/Engine/Source/Programs/Horde/Horde.Server/Storage/IStorageBackend.cs
+++ b/Engine/Source/Programs/Horde/Horde.Server/Storage/IStorageBackend.cs
@@ -15,6 +15,11 @@ namespace Horde.Server.Storage
 	/// </summary>
 	public interface IStorageBackend : IDisposable
 	{
+		/// <summary>
+		/// Whether this storage backend supports redirects or not
+		/// </summary>
+		bool SupportsRedirects { get; }
+
 		/// <summary>
 		/// Attempts to open a read stream for the given path.
 		/// </summary>
@@ -63,28 +68,22 @@ namespace Horde.Server.Storage
 		/// <param name="cancellationToken">Cancellation token for the operation</param>
 		/// <returns>Sequence of object paths</returns>
 		IAsyncEnumerable<string> EnumerateAsync(CancellationToken cancellationToken = default);
-	}
 
-	/// <summary>
-	/// Storage backend that supports redirecting HTTP traffic (eg. using S3 presigned urls)
-	/// </summary>
-	public interface IStorageBackendWithRedirects : IStorageBackend
-	{
 		/// <summary>
-		/// Gets a redirect for a read request
+		/// Gets a HTTP redirect for a read request
 		/// </summary>
 		/// <param name="path">Path to read from</param>
 		/// <param name="cancellationToken">Cancellation token for the operation</param>
 		/// <returns>Path to upload the data to</returns>
-		ValueTask<Uri?> GetReadRedirectAsync(string path, CancellationToken cancellationToken = default);
+		ValueTask<Uri?> TryGetReadRedirectAsync(string path, CancellationToken cancellationToken = default);
 
 		/// <summary>
-		/// Gets a redirect for a write request
+		/// Gets a HTTP redirect for a write request
 		/// </summary>
 		/// <param name="path">Path to write to</param>
 		/// <param name="cancellationToken">Cancellation token for the operation</param>
 		/// <returns>Path to upload the data to</returns>
-		ValueTask<Uri?> GetWriteRedirectAsync(string path, CancellationToken cancellationToken = default);
+		ValueTask<Uri?> TryGetWriteRedirectAsync(string path, CancellationToken cancellationToken = default);
 	}
 
 	/// <summary>
@@ -108,6 +107,9 @@ namespace Horde.Server.Storage
 		{
 			readonly IStorageBackend _inner;
 
+			/// <inheritdoc/>
+			public bool SupportsRedirects => _inner.SupportsRedirects;
+
 			/// <summary>
 			/// Constructor
 			/// </summary>
@@ -115,7 +117,7 @@ namespace Horde.Server.Storage
 			public TypedStorageBackend(IStorageBackend inner) => _inner = inner;
 
 			/// <inheritdoc/>
-			public void Dispose() { }
+			public void Dispose() => _inner.Dispose();
 
 			/// <inheritdoc/>
 			public Task<Stream?> TryReadAsync(string path, CancellationToken cancellationToken) => _inner.TryReadAsync(path, cancellationToken);
@@ -134,6 +136,12 @@ namespace Horde.Server.Storage
 
 			/// <inheritdoc/>
 			public IAsyncEnumerable<string> EnumerateAsync(CancellationToken cancellationToken = default) => _inner.EnumerateAsync(cancellationToken);
+
+			/// <inheritdoc/>
+			public ValueTask<Uri?> TryGetReadRedirectAsync(string path, CancellationToken cancellationToken = default) => _inner.TryGetReadRedirectAsync(path, cancellationToken);
+
+			/// <inheritdoc/>
+			public ValueTask<Uri?> TryGetWriteRedirectAsync(string path, CancellationToken cancellationToken = default) => _inner.TryGetWriteRedirectAsync(path, cancellationToken);
 		}
 
 		/// <summary>
diff --git a/Engine/Source/Programs/Horde/Horde.Server/Storage/StorageBackendProvider.cs b/Engine/Source/Programs/Horde/Horde.Server/Storage/StorageBackendProvider.cs
--- a/Engine/Source/Programs/Horde/Horde.Server/Storage/StorageBackendProvider.cs
+++ b/Engine/Source/Programs/Horde/Horde.Server/Storage/StorageBackendProvider.cs
@@ -64,6 +64,9 @@ namespace Horde.Server.Storage
 
 			#region IStorageBackend Implementation
 
+			/// <inheritdoc/>
+			public bool SupportsRedirects => _backend.SupportsRedirects;
+
 			/// <inheritdoc/>
 			public Task<bool> ExistsAsync(string path, CancellationToken cancellationToken = default) => _backend.ExistsAsync(path, cancellationToken);
 
@@ -82,6 +85,12 @@ namespace Horde.Server.Storage
 			/// <inheritdoc/>
 			public Task WriteAsync(string path, Stream stream, CancellationToken cancellationToken = default) => _backend.WriteAsync(path, stream, cancellationToken);
 
+			/// <inheritdoc/>
+			public ValueTask<Uri?> TryGetReadRedirectAsync(string path, CancellationToken cancellationToken = default) => _backend.TryGetReadRedirectAsync(path, cancellationToken);
+
+			/// <inheritdoc/>
+			public ValueTask<Uri?> TryGetWriteRedirectAsync(string path, CancellationToken cancellationToken = default) => _backend.TryGetWriteRedirectAsync(path, cancellationToken);
+
 			#endregion
 		}
 
diff --git a/Engine/Source/Programs/Horde/Horde.Server/Storage/StorageService.cs b/Engine/Source/Programs/Horde/Horde.Server/Storage/StorageService.cs
--- a/Engine/Source/Programs/Horde/Horde.Server/Storage/StorageService.cs
+++ b/Engine/Source/Programs/Horde/Horde.Server/Storage/StorageService.cs
@@ -103,7 +103,7 @@ namespace Horde.Server.Storage
 				_outer = outer;
 				Config = config;
 				Backend = backend;
-				SupportsRedirects = (backend is IStorageBackendWithRedirects) && !config.EnableAliases;
+				SupportsRedirects = backend.SupportsRedirects && !config.EnableAliases;
 
 				_prefix = config.Prefix;
 				if (_prefix.Length > 0 && !_prefix.EndsWith("/", StringComparison.Ordinal))
@@ -131,15 +131,7 @@ namespace Horde.Server.Storage
 			}
 
 			/// <inheritdoc/>
-			public ValueTask<Uri?> GetReadRedirectAsync(BlobLocator locator, CancellationToken cancellationToken = default)
-			{
-				IStorageBackendWithRedirects? redirectBackend = Backend as IStorageBackendWithRedirects;
-				if (redirectBackend == null)
-				{
-					return new ValueTask<Uri?>((Uri?)null);
-				}
-				return redirectBackend.GetReadRedirectAsync(GetBlobPath(locator.BlobId), cancellationToken);
-			}
+			public ValueTask<Uri?> GetReadRedirectAsync(BlobLocator locator, CancellationToken cancellationToken = default) => Backend.TryGetReadRedirectAsync(GetBlobPath(locator.BlobId), cancellationToken);
 
 			/// <inheritdoc/>
 			public override async Task<Stream> ReadBlobRangeAsync(BlobLocator locator, int offset, int length, CancellationToken cancellationToken = default)
@@ -204,8 +196,7 @@ namespace Horde.Server.Storage
 			/// <inheritdoc/>
 			public async ValueTask<(BlobLocator, Uri)?> GetWriteRedirectAsync(Utf8String prefix = default, CancellationToken cancellationToken = default)
 			{
-				IStorageBackendWithRedirects? redirectBackend = Backend as IStorageBackendWithRedirects;
-				if (redirectBackend == null)
+				if (!Backend.SupportsRedirects)
 				{
 					return null;
 				}
@@ -213,7 +204,7 @@ namespace Horde.Server.Storage
 				BlobLocator locator = await _outer.AddBlobAsync(NamespaceId, prefix, null, cancellationToken);
 				string path = GetBlobPath(locator.BlobId);
 
-				Uri? url = await redirectBackend.GetWriteRedirectAsync(path, cancellationToken);
+				Uri? url = await Backend.TryGetWriteRedirectAsync(path, cancellationToken);
 				if (url == null)
 				{
 					return null;
@@ -251,7 +242,6 @@ namespace Horde.Server.Storage
 			#endregion
 		}
 
-
 		class NamespaceInfo : IDisposable
 		{
 			public NamespaceConfig Config { get; }
