Horde: Fix mapping of admin claims to internal Horde roles

diff --git a/Engine/Source/Programs/Horde/Horde.Server/Authentication/HordeOpenIdConnectHandler.cs b/Engine/Source/Programs/Horde/Horde.Server/Authentication/HordeOpenIdConnectHandler.cs
--- a/Engine/Source/Programs/Horde/Horde.Server/Authentication/HordeOpenIdConnectHandler.cs
+++ b/Engine/Source/Programs/Horde/Horde.Server/Authentication/HordeOpenIdConnectHandler.cs
@@ -74,6 +74,14 @@ namespace Horde.Server.Authentication
 					identity.AddClaim(new Claim(ClaimTypes.Role, groupsElement[idx].ToString()!));
 				}
 			}
+
+			if (!String.IsNullOrEmpty(settings.AdminClaimType) && !String.IsNullOrEmpty(settings.AdminClaimValue))
+			{
+				if (identity.HasClaim(settings.AdminClaimType, settings.AdminClaimValue))
+				{
+					identity.AddClaim(HordeClaims.AdminClaim.ToClaim());
+				}
+			}
 		}
 
 		protected override async Task<HandleRequestResult> HandleRemoteAuthenticateAsync()
diff --git a/Engine/Source/Programs/Horde/Horde.Server/Authentication/OktaHandler.cs b/Engine/Source/Programs/Horde/Horde.Server/Authentication/OktaHandler.cs
--- a/Engine/Source/Programs/Horde/Horde.Server/Authentication/OktaHandler.cs
+++ b/Engine/Source/Programs/Horde/Horde.Server/Authentication/OktaHandler.cs
@@ -34,7 +34,7 @@ namespace Horde.Server.Authentication
 			_userCollection = userCollection;
 		}
 
-		public static void AddUserInfoClaims(JsonElement userInfo, ClaimsIdentity identity)
+		public static void AddUserInfoClaims(JsonElement userInfo, ServerSettings settings, ClaimsIdentity identity)
 		{
 			JsonElement userElement;
 			if (userInfo.TryGetProperty("preferred_username", out userElement))
@@ -67,6 +67,14 @@ namespace Horde.Server.Authentication
 			{
 				identity.AddClaim(new Claim(ClaimTypes.Email, emailElement.ToString()!));
 			}
+
+			if (!String.IsNullOrEmpty(settings.AdminClaimType) && !String.IsNullOrEmpty(settings.AdminClaimValue))
+			{
+				if (identity.HasClaim(settings.AdminClaimType, settings.AdminClaimValue))
+				{
+					identity.AddClaim(HordeClaims.AdminClaim.ToClaim());
+				}
+			}
 		}
 
 		protected override async Task<HandleRequestResult> HandleRemoteAuthenticateAsync()
@@ -101,18 +109,21 @@ namespace Horde.Server.Authentication
 	{
 		class MapRolesClaimAction : ClaimAction
 		{
-			public MapRolesClaimAction()
+			readonly ServerSettings _settings;
+
+			public MapRolesClaimAction(ServerSettings settings)
 				: base(ClaimTypes.Role, ClaimTypes.Role)
 			{
+				_settings = settings;
 			}
 
 			public override void Run(JsonElement userData, ClaimsIdentity identity, string issuer)
 			{
-				OktaHandler.AddUserInfoClaims(userData, identity);
+				OktaHandler.AddUserInfoClaims(userData, _settings, identity);
 			}
 		}
 
-		static void ApplyDefaultOktaOptions(OpenIdConnectOptions options, Action<OpenIdConnectOptions> handler)
+		static void ApplyDefaultOktaOptions(OpenIdConnectOptions options, ServerSettings settings, Action<OpenIdConnectOptions> handler)
 		{
 			options.Scope.Add("profile");
 			options.Scope.Add("groups");
@@ -121,15 +132,15 @@ namespace Horde.Server.Authentication
 			options.GetClaimsFromUserInfoEndpoint = true;
 			options.SaveTokens = true;
 			options.TokenValidationParameters.NameClaimType = "name";
-			options.ClaimActions.Add(new MapRolesClaimAction());
+			options.ClaimActions.Add(new MapRolesClaimAction(settings));
 
 			handler(options);
 		}
 
-		public static void AddOkta(this AuthenticationBuilder builder, string authenticationScheme, string displayName, Action<OpenIdConnectOptions> handler)
+		public static void AddOkta(this AuthenticationBuilder builder, ServerSettings settings, string authenticationScheme, string displayName, Action<OpenIdConnectOptions> handler)
 		{
 			builder.Services.TryAddEnumerable(ServiceDescriptor.Singleton<IPostConfigureOptions<OpenIdConnectOptions>, OpenIdConnectPostConfigureOptions>());
-			builder.AddRemoteScheme<OpenIdConnectOptions, OktaHandler>(authenticationScheme, displayName, options => ApplyDefaultOktaOptions(options, handler));
+			builder.AddRemoteScheme<OpenIdConnectOptions, OktaHandler>(authenticationScheme, displayName, options => ApplyDefaultOktaOptions(options, settings, handler));
 		}
 	}
 }
diff --git a/Engine/Source/Programs/Horde/Horde.Server/ServerSettings.cs b/Engine/Source/Programs/Horde/Horde.Server/ServerSettings.cs
--- a/Engine/Source/Programs/Horde/Horde.Server/ServerSettings.cs
+++ b/Engine/Source/Programs/Horde/Horde.Server/ServerSettings.cs
@@ -465,6 +465,16 @@ namespace Horde.Server
 		/// </summary>
 		public int JwtExpiryTimeHours { get; set; } = 8;
 
+		/// <summary>
+		/// The claim type for administrators
+		/// </summary>
+		public string? AdminClaimType { get; set; }
+
+		/// <summary>
+		/// Value of the claim type for administrators
+		/// </summary>
+		public string? AdminClaimValue { get; set; }
+
 		/// <summary>
 		/// Whether to enable Cors, generally for development purposes
 		/// </summary>
diff --git a/Engine/Source/Programs/Horde/Horde.Server/Startup.cs b/Engine/Source/Programs/Horde/Horde.Server/Startup.cs
--- a/Engine/Source/Programs/Horde/Horde.Server/Startup.cs
+++ b/Engine/Source/Programs/Horde/Horde.Server/Startup.cs
@@ -605,7 +605,7 @@ namespace Horde.Server
 					break;
 						
 				case AuthMethod.Okta:
-					authBuilder.AddOkta(OktaDefaults.AuthenticationScheme, OpenIdConnectDefaults.DisplayName, options =>
+					authBuilder.AddOkta(settings, OktaDefaults.AuthenticationScheme, OpenIdConnectDefaults.DisplayName, options =>
 						{
 							options.Authority = settings.OidcAuthority;
 							options.ClientId = settings.OidcClientId;
