Horde: Fix issue with agent/pools not updating on save
#skipci
#rnx

diff --git a/Engine/Source/Programs/Horde/HordeDashboard/src/backend/AgentStore.ts b/Engine/Source/Programs/Horde/HordeDashboard/src/backend/AgentStore.ts
--- a/Engine/Source/Programs/Horde/HordeDashboard/src/backend/AgentStore.ts
+++ b/Engine/Source/Programs/Horde/HordeDashboard/src/backend/AgentStore.ts
@@ -38,11 +38,12 @@ export class AgentStore {
     @action 
     private _setPools(data: PoolData[]) {
         this._pools = data;
+        this.poolsUpdated++;
     }
 
     @action
     private _setAgents(data: AgentData[]) {
-        if(data.length !== 0) {
+        if (data.length !== 0) {            
             const toSet: AgentData[] = [...this._agents];
             data.forEach(updatedAgent => {
                 const newUpdateIdx = toSet.findIndex(existingAgent => existingAgent.id === updatedAgent.id);
@@ -54,6 +55,7 @@ export class AgentStore {
                 }
             });
             this._agents = toSet;
+            this.agentsUpdated++;
         }
 	}
 
