Horde: Support adding environment variables for all agent types in a stream, and expanding the $(RootDir) variable within their values.

diff --git a/Engine/Source/Programs/Horde/Horde.Agent/Execution/JobExecutor.cs b/Engine/Source/Programs/Horde/Horde.Agent/Execution/JobExecutor.cs
--- a/Engine/Source/Programs/Horde/Horde.Agent/Execution/JobExecutor.cs
+++ b/Engine/Source/Programs/Horde/Horde.Agent/Execution/JobExecutor.cs
@@ -1247,18 +1247,28 @@ namespace Horde.Agent.Execution
 		{
 			List<DirectoryReference> dirs = new();
 			dirs.Add(workspaceDir);
-			if (sharedStorageDir != null) dirs.Add(sharedStorageDir);
-			if (envVars.TryGetValue("UE_SDKS_ROOT", out string? autoSdkDirPath)) dirs.Add(new DirectoryReference(autoSdkDirPath));
+			if (sharedStorageDir != null)
+			{
+				dirs.Add(sharedStorageDir);
+			}
+			if (envVars.TryGetValue("UE_SDKS_ROOT", out string? autoSdkDirPath))
+			{
+				dirs.Add(new DirectoryReference(autoSdkDirPath));
+			}
 			return dirs;
 		}
-		
+
 		async Task<int> ExecuteCommandAsync(BeginStepResponse step, DirectoryReference workspaceDir, DirectoryReference? sharedStorageDir, string fileName, string arguments, ILogger jobLogger, CancellationToken cancellationToken)
 		{
+			// Method for expanding environment variable properties related to this step
+			Dictionary<string, string> properties = new Dictionary<string, string>(StringComparer.OrdinalIgnoreCase);
+			properties.Add("RootDir", workspaceDir.FullName);
+
 			// Combine all the supplied environment variables together
 			Dictionary<string, string> newEnvVars = new Dictionary<string, string>(_envVars, StringComparer.Ordinal);
 			foreach (KeyValuePair<string, string> envVar in step.EnvVars)
 			{
-				newEnvVars[envVar.Key] = envVar.Value;
+				newEnvVars[envVar.Key] = StringUtils.ExpandProperties(envVar.Value, properties);
 			}
 			foreach (KeyValuePair<string, string> envVar in step.Credentials)
 			{
diff --git a/Engine/Source/Programs/Horde/Horde.Server/Streams/StreamConfig.cs b/Engine/Source/Programs/Horde/Horde.Server/Streams/StreamConfig.cs
--- a/Engine/Source/Programs/Horde/Horde.Server/Streams/StreamConfig.cs
+++ b/Engine/Source/Programs/Horde/Horde.Server/Streams/StreamConfig.cs
@@ -176,6 +176,11 @@ namespace Horde.Server.Streams
 		/// </summary>
 		public List<TabConfig> Tabs { get; set; } = new List<TabConfig>();
 
+		/// <summary>
+		/// Global environment variables for all agents in this stream
+		/// </summary>
+		public Dictionary<string, string>? Environment { get; set; }
+
 		/// <summary>
 		/// Map of agent name to type
 		/// </summary>
@@ -248,6 +253,18 @@ namespace Horde.Server.Streams
 			{
 				template.PostLoad(this);
 			}
+
+			if (Environment != null && Environment.Count > 0)
+			{
+				foreach (AgentConfig agentConfig in AgentTypes.Values)
+				{
+					agentConfig.Environment ??= new Dictionary<string, string>(StringComparer.OrdinalIgnoreCase);
+					foreach ((string key, string value) in Environment)
+					{
+						agentConfig.Environment.TryAdd(key, value);
+					}
+				}
+			}
 		}
 
 		/// <summary>
