Horde: Include batch environment variables in property expansion.

diff --git a/Engine/Source/Programs/Horde/Horde.Agent/Execution/JobExecutor.cs b/Engine/Source/Programs/Horde/Horde.Agent/Execution/JobExecutor.cs
--- a/Engine/Source/Programs/Horde/Horde.Agent/Execution/JobExecutor.cs
+++ b/Engine/Source/Programs/Horde/Horde.Agent/Execution/JobExecutor.cs
@@ -1265,7 +1265,11 @@ namespace Horde.Agent.Execution
 			properties.Add("RootDir", workspaceDir.FullName);
 
 			// Combine all the supplied environment variables together
-			Dictionary<string, string> newEnvVars = new Dictionary<string, string>(_envVars, StringComparer.Ordinal);
+			Dictionary<string, string> newEnvVars = new Dictionary<string, string>(StringComparer.Ordinal);
+			foreach (KeyValuePair<string, string> envVar in _batch.Environment)
+			{
+				newEnvVars[envVar.Key] = StringUtils.ExpandProperties(envVar.Value, properties);
+			}
 			foreach (KeyValuePair<string, string> envVar in step.EnvVars)
 			{
 				newEnvVars[envVar.Key] = StringUtils.ExpandProperties(envVar.Value, properties);
