Horde: Increase default telemetry report interval for agents

diff --git a/Engine/Source/Programs/Horde/Horde.Agent/AgentSettings.cs b/Engine/Source/Programs/Horde/Horde.Agent/AgentSettings.cs
--- a/Engine/Source/Programs/Horde/Horde.Agent/AgentSettings.cs
+++ b/Engine/Source/Programs/Horde/Horde.Agent/AgentSettings.cs
@@ -194,7 +194,7 @@ namespace Horde.Agent
 		/// <summary>
 		/// How often to report telemetry events to server in milliseconds
 		/// </summary>
-		public int TelemetryReportInterval { get; set; } = 10 * 1000;
+		public int TelemetryReportInterval { get; set; } = 30 * 1000;
 
 		/// <summary>
 		/// Key/value properties in addition to those set internally by the agent
