Horde: More artifact browsertweaks
#rnx
#skipci

diff --git a/Engine/Source/Programs/Horde/HordeDashboard/src/components/jobDetailsV2/JobDetailArtifactsV3.tsx b/Engine/Source/Programs/Horde/HordeDashboard/src/components/jobDetailsV2/JobDetailArtifactsV3.tsx
--- a/Engine/Source/Programs/Horde/HordeDashboard/src/components/jobDetailsV2/JobDetailArtifactsV3.tsx
+++ b/Engine/Source/Programs/Horde/HordeDashboard/src/components/jobDetailsV2/JobDetailArtifactsV3.tsx
@@ -83,8 +83,12 @@ class ArtifactsHandler {
       a = artifacts.find(a => a.type === this.context)!;
       this.artifact = a;
 
+      this.loading = true;
+      this.updateReady();
+
       this.browse = await backend.getBrowseArtifacts(a.id);
 
+      this.loading = false;
       this.updateReady();
 
    }
@@ -95,16 +99,23 @@ class ArtifactsHandler {
 
    }
 
-   async browseTo(path: string) {
+   async browseTo(path: string, push = true) {
 
       if (!this.artifact) {
          return;
       }
 
+      this.loading = true;
+      this.updateReady();
       this.browse = await backend.getBrowseArtifacts(this.artifact.id, path);
 
       this.path = path;
 
+      if (push) {
+         this.history.push(path);
+      }
+
+      this.loading = false;
       this.updateReady();
 
    }
@@ -169,7 +180,9 @@ class ArtifactsHandler {
       this.artifact = undefined;
       this.artifacts = undefined;
       this.contexts = undefined;
+      this.history = [];
       this.stepId = "";
+      this.loading = false;
    }
 
 
@@ -191,6 +204,10 @@ class ArtifactsHandler {
    readonly context: ArtifactContextType;
 
    stepId: string;
+
+   history: string[] = [];
+
+   loading = false;
 }
 
 const styles = mergeStyleSets({
@@ -218,6 +235,41 @@ const styles = mergeStyleSets({
    }
 });
 
+const BrowseHistory: React.FC<{ handler: ArtifactsHandler }> = observer(({ handler }) => {
+
+   // subscribe
+   if (handler.updated) { };
+
+   let backDisabled = !handler.history.length;
+   let upDisabled = !handler.path;
+
+
+   return <Stack style={{ paddingRight: 12 }}>
+      <Stack horizontal tokens={{ childrenGap: 8 }}>
+         <Stack>
+            <IconButton disabled={backDisabled} style={{ fontSize: 14, paddingTop: 1 }} iconProps={{ iconName: 'ArrowLeft' }} onClick={() => {
+               if (handler.history.length === 1) {
+                  handler.history = [];
+                  handler.browseTo("", false);
+               } else {
+                  handler.history.pop();
+                  handler.browseTo(handler.history[handler.history.length - 1], false);
+               }
+            }} />
+         </Stack>
+         <Stack>
+            <IconButton disabled={upDisabled} style={{ fontSize: 14, paddingTop: 1 }} iconProps={{ iconName: 'ArrowUp' }} onClick={() => {
+               const nbrowse = handler.path?.split("/")
+               nbrowse?.pop();
+               if (nbrowse) {
+                  handler.browseTo(nbrowse.join("/"));
+               }
+            }} />
+         </Stack>
+      </Stack>
+   </Stack>
+
+});
 
 const BrowseBreadCrumbs: React.FC<{ handler: ArtifactsHandler }> = observer(({ handler }) => {
 
@@ -235,9 +287,15 @@ const BrowseBreadCrumbs: React.FC<{ handler: ArtifactsHandler }> = observer(({ h
    }
 
    if (!handler.path) {
-      return <Stack className="horde-no-darktheme">
-         <Text style={{ fontSize: fontSize, fontWeight: 600, paddingLeft: 2, paddingRight: 4 }}>{rootName} /</Text>
+      return <Stack horizontal verticalAlign="center" styles={{ root: { height: 40 } }}>
+         <Stack>
+            <BrowseHistory handler={handler} />
+         </Stack>
+         <Stack>
+            <Text style={{ fontSize: fontSize, fontWeight: 600, paddingLeft: 2, paddingRight: 4 }}>{rootName} /</Text>
+         </Stack>
       </Stack>
+
    }
 
    const ppath = handler.path.split("/");
@@ -260,11 +318,18 @@ const BrowseBreadCrumbs: React.FC<{ handler: ArtifactsHandler }> = observer(({ h
 
    elements.unshift(<Stack style={{ cursor: "pointer" }} onClick={() => handler.browseTo("")}>
       <Stack style={{ paddingLeft: 2, paddingRight: 4 }}>
-         <Text style={{ fontSize: fontSize, color:  (dashboard.darktheme ? "#55B7FF" : "#0078D4"), fontWeight: 600 }}>{rootName} /</Text>
+         <Text style={{ fontSize: fontSize, color: (dashboard.darktheme ? "#55B7FF" : "#0078D4"), fontWeight: 600 }}>{rootName} /</Text>
       </Stack>
    </Stack>);
 
-   return <Stack horizontal>{elements}</Stack>
+   return <Stack horizontal verticalAlign="center" styles={{ root: { height: 40 } }}>
+      <Stack>
+         <BrowseHistory handler={handler} />
+      </Stack>
+      <Stack horizontal wrap style={{ width: 800 }}>
+         {elements}
+      </Stack>
+   </Stack>
 
 })
 
@@ -318,7 +383,7 @@ const DownloadButton: React.FC<{ handler: ArtifactsHandler }> = observer(({ hand
    }
 
    return <Stack>
-      <PrimaryButton styles={{ root: { fontFamily: 'Horde Open Sans SemiBold !important' } }} disabled={!selection.size} onClick={async () => {
+      <PrimaryButton styles={{ root: { fontFamily: 'Horde Open Sans SemiBold !important' } }} disabled={!selection.size || downloading} onClick={async () => {
 
          const selection = handler.currentSelection.items;
 
@@ -380,7 +445,7 @@ const DownloadButton: React.FC<{ handler: ArtifactsHandler }> = observer(({ hand
             setDownloading(false);
          }
 
-      }}>{buttonText}{!!downloading && <Spinner style={{ paddingLeft: 12 }} size={SpinnerSize.medium} />}</PrimaryButton>
+      }}>{buttonText}</PrimaryButton>
    </Stack>
 
 });
@@ -400,9 +465,12 @@ const JobDetailArtifactsInner: React.FC<{ handler: ArtifactsHandler }> = observe
 
    const items: BrowserItem[] = [];
 
+   // use the up arrow instead
+   /*
    if (handler.path?.length) {
       items.push({ key: "navigate up", text: "..", type: BrowserType.NavigateUp });
    }
+   */
 
    browse.directories?.forEach(d => {
 
@@ -470,7 +538,7 @@ const JobDetailArtifactsInner: React.FC<{ handler: ArtifactsHandler }> = observe
 
          const pathElements = path.map((t, index) => {
             const last = index == (path.length - 1);
-            let color = last ? (dashboard.darktheme ? "#FFFFFF" :"#605E5C") : undefined;
+            let color = last ? (dashboard.darktheme ? "#FFFFFF" : "#605E5C") : undefined;
             const font = last ? undefined : "Horde Open Sans Light";
             const sep = last ? undefined : "/"
             return <Text styles={{ root: { fontFamily: font } }} style={{ color: color }}>{t}{sep}</Text>
@@ -517,7 +585,10 @@ const JobDetailArtifactsInner: React.FC<{ handler: ArtifactsHandler }> = observe
          <Stack grow />
          <DownloadButton handler={handler} />
       </Stack >
-      <Stack style={{ height: 492 + 160, position: "relative" }}>
+      {handler.loading && <Stack>
+         <Spinner styles={{ root: { opacity: 0, animation: "hordeFadeIn 1s ease-in-out 2s forwards" } }} size={SpinnerSize.large} />
+      </Stack>}
+      {!handler.loading && <Stack style={{ height: 492 + 160, position: "relative" }}>
          <ScrollablePane style={{ height: 492 + 160 }}>
             <SelectionZone selection={handler.selection!}>
                <DetailsList
@@ -537,7 +608,7 @@ const JobDetailArtifactsInner: React.FC<{ handler: ArtifactsHandler }> = observe
             </SelectionZone>
 
          </ScrollablePane>
-      </Stack>
+      </Stack>}
    </Stack >
 
    //
