Horde: Fix missing icons
#skipci
#rnx

diff --git a/Engine/Source/Programs/Horde/HordeDashboard/src/index.css b/Engine/Source/Programs/Horde/HordeDashboard/src/index.css
--- a/Engine/Source/Programs/Horde/HordeDashboard/src/index.css
+++ b/Engine/Source/Programs/Horde/HordeDashboard/src/index.css
@@ -2,22 +2,24 @@ body {
   margin: 0;
   font-family: "Horde Open Sans Regular";
   -webkit-font-smoothing: antialiased;
-  -moz-osx-font-smoothing: grayscale; 
+  -moz-osx-font-smoothing: grayscale;
   overflow: hidden;
 }
 
-a, a:visited, a:active {
+a,
+a:visited,
+a:active {
   text-decoration: none;
-  color:rgb(0, 120, 212);
+  color: rgb(0, 120, 212);
 }
 
 :focus {
-  outline:none;
+  outline: none;
 }
 
 a:hover {
   text-decoration: none;
-  color:rgb(0, 120, 212);
+  color: rgb(0, 120, 212);
 }
 
 code {
@@ -25,12 +27,15 @@ code {
     monospace;
 }
 
-.log-link, .log-link:visited, .log-link:hover {  
+.log-link,
+.log-link:visited,
+.log-link:hover {
   text-decoration: none;
-  color:#000000;
+  color: #000000;
 }
 
-.horde-link, .horde-link:visited {
+.horde-link,
+.horde-link:visited {
   font-family: "Horde Open Sans Regular";
   text-decoration: none;
   color: rgb(0, 120, 212);
@@ -43,44 +48,50 @@ code {
 }
 
 .job-item:hover {
-  background:rgb(243, 242, 241);
+  background: rgb(243, 242, 241);
 }
 
 .job-item.selected {
-  background:rgb(243, 242, 241);
+  background: rgb(243, 242, 241);
 }
 
 
-.view-log-link, .view-log-link:visited, .view-log-link:hover {
+.view-log-link,
+.view-log-link:visited,
+.view-log-link:hover {
   text-decoration: none;
   user-select: none;
   cursor: pointer;
 }
 
-.view-log-link-disabled, .view-log-link-disabled:visited, .view-log-link-disabled:hover {
+.view-log-link-disabled,
+.view-log-link-disabled:visited,
+.view-log-link-disabled:hover {
   text-decoration: none;
   color: #8f8f8f;
   user-select: none;
 }
 
-.cl-callout-button,.cl-callout-button:hover  {
+.cl-callout-button,
+.cl-callout-button:hover {
   text-decoration: none;
   white-space: nowrap;
   background: #035ca1;
   color: rgb(255, 255, 255);
   user-select: none;
-  font-family: "Horde Open Sans SemiBold";  
+  font-family: "Horde Open Sans SemiBold";
   font-size: 10px;
   border-width: 0;
 }
 
-.cl-callout-button-user,.cl-callout-button-user:hover  {
+.cl-callout-button-user,
+.cl-callout-button-user:hover {
   text-decoration: none;
   white-space: nowrap;
-  background:#0288ee;
+  background: #0288ee;
   color: rgb(255, 255, 255);
   user-select: none;
-  font-family: "Horde Open Sans SemiBold";  
+  font-family: "Horde Open Sans SemiBold";
   font-size: 10px;
   border-width: 0;
 }
@@ -90,10 +101,12 @@ code {
     background-color: #CC0000;
     border: 1px solid #CC0000;
   }
+
   50% {
     background-color: #FF0000;
     border: 1px solid #FF0000;
   }
+
   100% {
     background-color: #CC0000;
     border: 1px solid #CC0000;
@@ -101,3 +114,12 @@ code {
 
 }
 
+@keyframes hordeFadeIn {
+  from {
+    opacity: 0;
+  }
+
+  to {
+    opacity: 1;
+  }
+}
\ No newline at end of file
diff --git a/Engine/Source/Programs/Horde/HordeDashboard/src/styles/Styles.tsx b/Engine/Source/Programs/Horde/HordeDashboard/src/styles/Styles.tsx
--- a/Engine/Source/Programs/Horde/HordeDashboard/src/styles/Styles.tsx
+++ b/Engine/Source/Programs/Horde/HordeDashboard/src/styles/Styles.tsx
@@ -95,6 +95,8 @@ registerIcons({
       "Cross": IconContents.CROSS,
       "ChevronUp": IconContents.CHEVRON_UP,
       "ChevronDown": IconContents.CHEVRON_DOWN,
+      "ChevronLeft": IconContents.CHEVRON_LEFT,
+      "ChevronRight": IconContents.CHEVRON_RIGHT,
       "Refresh": IconContents.REFRESH,
       "Warning": IconContents.WARNING_SIGN,
       "LayoutCircle": IconContents.LAYOUT_CIRCLE,
@@ -126,7 +128,10 @@ registerIcons({
       "Document": IconContents.DOCUMENT,
       "Folder": IconContents.FOLDER_CLOSE,
       "CloudDownload": IconContents.CLOUD_DOWNLOAD,
-      "Eye" : IconContents.EYE_OPEN
+      "Eye": IconContents.EYE_OPEN,
+      "ArrowLeft": IconContents.ARROW_LEFT,
+      "ArrowRight": IconContents.ARROW_RIGHT,
+      "ArrowUp" : IconContents.ARROW_UP
    }
 });
 
