Horde: Add a spinner for download
#skipci
#rnx

diff --git a/Engine/Source/Programs/Horde/HordeDashboard/src/components/jobDetailsV2/JobDetailArtifactsV3.tsx b/Engine/Source/Programs/Horde/HordeDashboard/src/components/jobDetailsV2/JobDetailArtifactsV3.tsx
--- a/Engine/Source/Programs/Horde/HordeDashboard/src/components/jobDetailsV2/JobDetailArtifactsV3.tsx
+++ b/Engine/Source/Programs/Horde/HordeDashboard/src/components/jobDetailsV2/JobDetailArtifactsV3.tsx
@@ -382,7 +382,8 @@ const DownloadButton: React.FC<{ handler: ArtifactsHandler }> = observer(({ hand
       buttonText = `Download (${sizeText})`;
    }
 
-   return <Stack>
+   return <Stack horizontal verticalAlign="center" tokens={{ childrenGap: 8 }}>
+      {downloading && <Stack><Spinner size={SpinnerSize.large} /></Stack>}
       <PrimaryButton styles={{ root: { fontFamily: 'Horde Open Sans SemiBold !important' } }} disabled={!selection.size || downloading} onClick={async () => {
 
          const selection = handler.currentSelection.items;
