Horde: Fix button overflow
#skipci
#rnx

diff --git a/Engine/Source/Programs/Horde/HordeDashboard/src/components/jobDetailsV2/JobDetailArtifactsV3.tsx b/Engine/Source/Programs/Horde/HordeDashboard/src/components/jobDetailsV2/JobDetailArtifactsV3.tsx
--- a/Engine/Source/Programs/Horde/HordeDashboard/src/components/jobDetailsV2/JobDetailArtifactsV3.tsx
+++ b/Engine/Source/Programs/Horde/HordeDashboard/src/components/jobDetailsV2/JobDetailArtifactsV3.tsx
@@ -382,6 +382,10 @@ const DownloadButton: React.FC<{ handler: ArtifactsHandler }> = observer(({ hand
       buttonText = `Download (${sizeText})`;
    }
 
+   if (downloading) {
+      buttonText = "Downloading";
+   }
+
    return <Stack horizontal verticalAlign="center" tokens={{ childrenGap: 8 }}>
       {downloading && <Stack><Spinner size={SpinnerSize.large} /></Stack>}
       <PrimaryButton styles={{ root: { fontFamily: 'Horde Open Sans SemiBold !important' } }} disabled={!selection.size || downloading} onClick={async () => {
