Horde: Supply -NoP4 argument to UAT when have-table-less syncs are used

diff --git a/Engine/Source/Programs/Horde/Horde.Agent/Execution/PerforceExecutor.cs b/Engine/Source/Programs/Horde/Horde.Agent/Execution/PerforceExecutor.cs
--- a/Engine/Source/Programs/Horde/Horde.Agent/Execution/PerforceExecutor.cs
+++ b/Engine/Source/Programs/Horde/Horde.Agent/Execution/PerforceExecutor.cs
@@ -214,13 +214,15 @@ namespace Horde.Agent.Execution
 		protected override async Task<bool> SetupAsync(BeginStepResponse step, ILogger logger, CancellationToken cancellationToken)
 		{
 			PerforceLogger perforceLogger = CreatePerforceLogger(logger);
-			return await SetupAsync(step, _workspace.WorkspaceDir, _sharedStorageDir, null, perforceLogger, cancellationToken);
+			bool useP4 = WorkspaceInfo.ShouldUseHaveTable(_workspaceInfo.Method);
+			return await SetupAsync(step, _workspace.WorkspaceDir, _sharedStorageDir, useP4, perforceLogger, cancellationToken);
 		}
 
 		protected override async Task<bool> ExecuteAsync(BeginStepResponse step, ILogger logger, CancellationToken cancellationToken)
 		{
 			PerforceLogger perforceLogger = CreatePerforceLogger(logger);
-			return await ExecuteAsync(step, _workspace.WorkspaceDir, _sharedStorageDir, null, perforceLogger, cancellationToken);
+			bool useP4 = WorkspaceInfo.ShouldUseHaveTable(_workspaceInfo.Method);
+			return await ExecuteAsync(step, _workspace.WorkspaceDir, _sharedStorageDir, useP4, perforceLogger, cancellationToken);
 		}
 
 		public override async Task FinalizeAsync(ILogger logger, CancellationToken cancellationToken)
