Horde: Speculative fix for user home job refresh issue
#rnx
#skipci

diff --git a/Engine/Source/Programs/Horde/HordeDashboard/src/components/UserHome.tsx b/Engine/Source/Programs/Horde/HordeDashboard/src/components/UserHome.tsx
--- a/Engine/Source/Programs/Horde/HordeDashboard/src/components/UserHome.tsx
+++ b/Engine/Source/Programs/Horde/HordeDashboard/src/components/UserHome.tsx
@@ -324,7 +324,7 @@ const HealthPanel: React.FC = observer(() => {
             <Text variant="mediumPlus" styles={{ root: { fontFamily: "Horde Open Sans SemiBold" } }}>Issues</Text>
             <Stack styles={{ root: { paddingLeft: 4, paddingRight: 0, paddingTop: 8, paddingBottom: 4 } }}>
                <div style={{ overflowY: 'auto', overflowX: 'hidden', maxHeight: "380px" }} data-is-scrollable={true}>
-                  {!items && <Stack style={{ paddingBottom: 12 }} horizontal tokens={{ childrenGap: 12 }}><Text variant="medium">Loading issues</Text><Spinner size={SpinnerSize.medium} /></Stack>}
+                  {!items && <Stack style={{ paddingBottom: 12 }} horizontal tokens={{ childrenGap: 12 }}><Text variant="medium">Loading issues</Text><Spinner size={SpinnerSize.small} /></Stack>}
                   {!!items && !items.length && <Stack style={{ paddingBottom: 12 }}><Text variant="medium">No issues found</Text></Stack>}
                   {!!items && !!items.length && <DetailsList
                      compact={true}
@@ -348,7 +348,7 @@ const HealthPanel: React.FC = observer(() => {
 
 /// Jobs ==========================================================================================
 
-const jobsRefreshTime = 3000;
+const jobsRefreshTime = 10000;
 class UserJobsHandler {
 
    constructor() {
@@ -538,17 +538,17 @@ class UserJobsHandler {
 
 }
 
-const jobHandler = new UserJobsHandler();
-
 const JobsPanel: React.FC<{ includeOtherPreflights: boolean }> = observer(({ includeOtherPreflights }) => {
 
+   const [jobHandler, setJobHandler] = useState(new UserJobsHandler())
+
    const { projectStore } = useBackend();
 
    useEffect(() => {
 
       dashboard.startPolling();
 
-      return () => {
+      return () => {         
          dashboard.stopPolling();
          jobHandler.clear();
       };
@@ -1032,6 +1032,7 @@ const JobsPanel: React.FC<{ includeOtherPreflights: boolean }> = observer(({ inc
                   </Stack>
                </Stack>
                <Stack styles={{ root: { paddingLeft: 4, paddingRight: 0, paddingTop: 8, paddingBottom: 4 } }}>
+                  {!!dashboard.pinnedJobsIds.length && !jobItems.length && <Stack><Spinner size={SpinnerSize.large}/></Stack>}
                   {!!jobItems.length && <DetailsList
                      styles={{ root: { paddingLeft: 8, paddingRight: 8, marginBottom: 18 } }}
                      compact={true}
