Horde: Add initial support for OpenTelemetry in server

diff --git a/Engine/Source/Programs/Horde/Horde.Server.Tests/TestSetup.cs b/Engine/Source/Programs/Horde/Horde.Server.Tests/TestSetup.cs
--- a/Engine/Source/Programs/Horde/Horde.Server.Tests/TestSetup.cs
+++ b/Engine/Source/Programs/Horde/Horde.Server.Tests/TestSetup.cs
@@ -64,6 +64,7 @@ using Moq;
 using Horde.Server.Telemetry;
 using Horde.Server.Artifacts;
 using Horde.Server.Dashboard;
+using OpenTelemetry.Trace;
 
 namespace Horde.Server.Tests
 {
@@ -187,6 +188,7 @@ namespace Horde.Server.Tests
 			services.AddSingleton(typeof(IAuditLogFactory<>), typeof(AuditLogFactory<>));
 			services.AddSingleton<IAuditLog<AgentId>>(sp => sp.GetRequiredService<IAuditLogFactory<AgentId>>().Create("Agents.Log", "AgentId"));
 			services.AddSingleton<ITelemetrySink, NullTelemetrySink>();
+			services.AddSingleton<OpenTelemetry.Trace.Tracer>(sp => TracerProvider.Default.GetTracer("TestTracer"));
 
 			services.AddSingleton<ConfigService>();
 			services.AddSingleton<IOptionsFactory<GlobalConfig>>(sp => sp.GetRequiredService<ConfigService>());
@@ -366,14 +368,14 @@ namespace Horde.Server.Tests
 
 			string msg = "Unable to patch Datadog agent writer! Tests will still work, but shutdown will block for +20 seconds.";
 			
-			FieldInfo? agentWriterField = Tracer.Instance.GetType().GetField("_agentWriter", BindingFlags.NonPublic | BindingFlags.Instance);
+			FieldInfo? agentWriterField = Datadog.Trace.Tracer.Instance.GetType().GetField("_agentWriter", BindingFlags.NonPublic | BindingFlags.Instance);
 			if (agentWriterField == null)
 			{
 				Console.Error.WriteLine(msg);
 				return;
 			}
 			
-			object? agentWriterInstance = agentWriterField.GetValue(Tracer.Instance);
+			object? agentWriterInstance = agentWriterField.GetValue(Datadog.Trace.Tracer.Instance);
 			if (agentWriterInstance == null)
 			{
 				Console.Error.WriteLine(msg);
diff --git a/Engine/Source/Programs/Horde/Horde.Server/Agents/AgentService.cs b/Engine/Source/Programs/Horde/Horde.Server/Agents/AgentService.cs
--- a/Engine/Source/Programs/Horde/Horde.Server/Agents/AgentService.cs
+++ b/Engine/Source/Programs/Horde/Horde.Server/Agents/AgentService.cs
@@ -23,6 +23,7 @@ using HordeCommon;
 using Microsoft.Extensions.Hosting;
 using Microsoft.Extensions.Logging;
 using MongoDB.Driver;
+using OpenTelemetry.Trace;
 using OpenTracing;
 using OpenTracing.Util;
 using StackExchange.Redis;
@@ -76,6 +77,7 @@ namespace Horde.Server.Agents
 		readonly ITaskSource[] _taskSources;
 		readonly IHostApplicationLifetime _applicationLifetime;
 		readonly IClock _clock;
+		readonly Tracer _tracer;
 		readonly ILogger _logger;
 		readonly ITicker _ticker;
 		
@@ -97,7 +99,7 @@ namespace Horde.Server.Agents
 		/// <summary>
 		/// Constructor
 		/// </summary>
-		public AgentService(IAgentCollection agents, ILeaseCollection leases, ISessionCollection sessions, AclService aclService, IDowntimeService downtimeService, IPoolCollection poolCollection, IDogStatsd dogStatsd, IEnumerable<ITaskSource> taskSources, RedisService redisService, IHostApplicationLifetime applicationLifetime, IClock clock, ILogger<AgentService> logger)
+		public AgentService(IAgentCollection agents, ILeaseCollection leases, ISessionCollection sessions, AclService aclService, IDowntimeService downtimeService, IPoolCollection poolCollection, IDogStatsd dogStatsd, IEnumerable<ITaskSource> taskSources, RedisService redisService, IHostApplicationLifetime applicationLifetime, IClock clock, Tracer tracer, ILogger<AgentService> logger)
 		{
 			Agents = agents;
 			_leases = leases;
@@ -112,6 +114,7 @@ namespace Horde.Server.Agents
 			_redisService = redisService;
 			_clock = clock;
 			_ticker = clock.AddTicker<AgentService>(TimeSpan.FromSeconds(30.0), TickAsync, logger);
+			_tracer = tracer;
 			_logger = logger;
 		}
 
@@ -823,6 +826,14 @@ namespace Horde.Server.Agents
 			scope.Span.SetTag("FinishTime", finishTime?.ToString());
 			scope.Span.SetTag("Index", index);
 			scope.Span.SetTag("Count", count);
+			
+			using TelemetrySpan span = _tracer.StartActiveSpan($"{nameof(AgentService)}.{nameof(FindLeasesAsync)}-ot");
+			span.SetAttribute("AgentId", agentId?.ToString());
+			span.SetAttribute("SessionId", sessionId?.ToString());
+			span.SetAttribute("StartTime", startTime?.ToString());
+			span.SetAttribute("FinishTime", finishTime?.ToString());
+			span.SetAttribute("Index", index);
+			span.SetAttribute("Count", count);
 			return _leases.FindLeasesAsync(agentId, sessionId, startTime, finishTime, index, count);
 		}
 
diff --git a/Engine/Source/Programs/Horde/Horde.Server/Horde.Server.csproj b/Engine/Source/Programs/Horde/Horde.Server/Horde.Server.csproj
--- a/Engine/Source/Programs/Horde/Horde.Server/Horde.Server.csproj
+++ b/Engine/Source/Programs/Horde/Horde.Server/Horde.Server.csproj
@@ -46,6 +46,14 @@
     <PackageReference Include="MongoDB.Driver" Version="2.18.0" />
     <PackageReference Include="OpenTracing" Version="0.12.1" />
     <PackageReference Include="OpenTracing.Contrib.Grpc" Version="0.4.0" />
+    <PackageReference Include="OpenTelemetry" Version="1.4.0" />
+    <PackageReference Include="OpenTelemetry.Exporter.Console" Version="1.4.0" />
+    <PackageReference Include="OpenTelemetry.Exporter.OpenTelemetryProtocol" Version="1.4.0" />
+    <PackageReference Include="OpenTelemetry.Extensions.Hosting" Version="1.4.0" />
+    <PackageReference Include="OpenTelemetry.Instrumentation.AspNetCore" Version="1.0.0-rc9.14" />
+    <PackageReference Include="OpenTelemetry.Instrumentation.GrpcNetClient" Version="1.0.0-rc9.14" />
+    <PackageReference Include="OpenTelemetry.Instrumentation.Http" Version="1.0.0-rc9.14" />
+    <PackageReference Include="OpenTelemetry.Instrumentation.StackExchangeRedis" Version="1.0.0-rc9.8" />
     <PackageReference Include="Polly" Version="7.2.3" />
     <PackageReference Include="Polly.Extensions.Http" Version="3.0.0" />
     <PackageReference Include="Serilog.AspNetCore" Version="6.0.1" />
diff --git a/Engine/Source/Programs/Horde/Horde.Server/ServerSettings.cs b/Engine/Source/Programs/Horde/Horde.Server/ServerSettings.cs
--- a/Engine/Source/Programs/Horde/Horde.Server/ServerSettings.cs
+++ b/Engine/Source/Programs/Horde/Horde.Server/ServerSettings.cs
@@ -318,6 +318,63 @@ namespace Horde.Server
 		}
 	}
 
+	/// <summary>
+	/// OpenTelemetry configuration for collection and sending of traces and metrics.
+	/// </summary>
+	public class OpenTelemetrySettings
+	{
+		/// <summary>
+		/// Whether OpenTelemetry exporting is enabled
+		/// </summary>
+		public bool Enabled { get; set; } = false;
+		
+		/// <summary>
+		/// Service name
+		/// </summary>
+		public string ServiceName { get; set; } = "HordeServer";
+		
+		/// <summary>
+		/// Service namespace
+		/// </summary>
+		public string ServiceNamespace { get; set; } = "Horde";
+		
+		/// <summary>
+		/// Service version
+		/// </summary>
+		public string? ServiceVersion { get; set; }
+		
+		/// <summary>
+		/// Extra attributes to set
+		/// </summary>
+		public Dictionary<string, string> Attributes { get; set; } = new();
+
+		/// <summary>
+		/// Whether to enable the console exporter (for debugging purposes)
+		/// </summary>
+		public bool EnableConsoleExporter { get; set; } = false;
+		
+		/// <summary>
+		/// Protocol exporters (key is a unique and arbitrary name) 
+		/// </summary>
+		public Dictionary<string, OpenTelemetryProtocolExporterSettings> ProtocolExporters { get; set; } = new();
+	}
+
+	/// <summary>
+	/// Configuration for an OpenTelemetry exporter
+	/// </summary>
+	public class OpenTelemetryProtocolExporterSettings
+	{
+		/// <summary>
+		/// Endpoint URL. Usually differs depending on protocol used.
+		/// </summary>
+		public Uri? Endpoint { get; set; }
+
+		/// <summary>
+		/// Protocol for the exporter ('grpc' or 'httpprotobuf')
+		/// </summary>
+		public string Protocol { get; set; } = "grpc";
+	}
+
 	/// <summary>
 	/// Global settings for the application
 	/// </summary>
@@ -753,6 +810,11 @@ namespace Horde.Server
 		/// Tools bundled along with the server. Data for each tool can be produced using the 'bundle create' command, and should be stored in the /tools/{id} directory.
 		/// </summary>
 		public List<BundledToolConfig> BundledTools { get; set; } = new List<BundledToolConfig>();
+		
+		/// <summary>
+		/// Options for OpenTelemetry
+		/// </summary>
+		public OpenTelemetrySettings OpenTelemetry { get; set; } = new OpenTelemetrySettings();
 
 		/// <summary>
 		/// Default pre-baked ACL for authentication of well-known roles
diff --git a/Engine/Source/Programs/Horde/Horde.Server/Startup.cs b/Engine/Source/Programs/Horde/Horde.Server/Startup.cs
--- a/Engine/Source/Programs/Horde/Horde.Server/Startup.cs
+++ b/Engine/Source/Programs/Horde/Horde.Server/Startup.cs
@@ -97,6 +97,9 @@ using Horde.Server.Notifications.Sinks;
 using StatusCode = Grpc.Core.StatusCode;
 using Horde.Server.Artifacts;
 using Horde.Server.Compute;
+using OpenTelemetry.Exporter;
+using OpenTelemetry.Resources;
+using OpenTelemetry.Trace;
 
 namespace Horde.Server
 {
@@ -328,20 +331,23 @@ namespace Horde.Server
 		// This method gets called *multiple times* by the runtime. Use this method to add services to the container.
 		public void ConfigureServices(IServiceCollection services)
 		{
-			// IOptionsMonitor pattern for live updating of configuration settings
-			services.Configure<ServerSettings>(serverSettings =>
+			void BindServerSettings(ServerSettings settings)
 			{
 				IConfigurationSection hordeSection = Configuration.GetSection("Horde");
-				hordeSection.Bind(serverSettings);
-				BindTelemetrySettings(serverSettings, hordeSection.GetSection("Telemetry"));
-			});
-
-			// Settings used for configuring services
-			IConfigurationSection configSection = Configuration.GetSection("Horde");
-			ServerSettings settings = new ServerSettings();
-			configSection.Bind(settings);
+				hordeSection.Bind(settings);
+				BindTelemetrySettings(settings, hordeSection.GetSection("Telemetry"));
+			}
+			
+			// IOptionsMonitor pattern for live updating of configuration settings
+			services.Configure<ServerSettings>(BindServerSettings);
+			
+			// Bind the settings again for local variable access in this method
+			ServerSettings settings = new();
+			BindServerSettings(settings);
 			settings.Validate();
 
+			ConfigureOpenTelemetry(services, settings.OpenTelemetry);
+
 			if (settings.GlobalThreadPoolMinSize != null)
 			{
 				// Min thread pool size is set to combat timeouts seen with the Redis client.
@@ -1014,6 +1020,49 @@ namespace Horde.Server
 			}
 		}
 
+		private static void ConfigureOpenTelemetry(IServiceCollection services, OpenTelemetrySettings settings)
+		{
+			// Always configure OpenTelemetry.Trace.Tracer class as it's used in the codebase even when OpenTelemetry is not configured
+			services.AddSingleton(TracerProvider.Default.GetTracer(settings.ServiceName));
+			
+			if (!settings.Enabled)
+			{
+				return;
+			}
+			
+			List<KeyValuePair<string,object>> attributes = settings.Attributes.Select(x => new KeyValuePair<string, object>(x.Key, x.Value)).ToList();
+			services.AddOpenTelemetry()
+				.WithTracing(builder =>
+				{
+					builder
+						.AddSource(settings.ServiceName)
+						.AddHttpClientInstrumentation()
+						.AddAspNetCoreInstrumentation()
+						.AddGrpcClientInstrumentation()
+						//.AddRedisInstrumentation()
+						.SetResourceBuilder(ResourceBuilder.CreateDefault()
+							.AddService(settings.ServiceName, serviceNamespace: settings.ServiceNamespace, serviceVersion: settings.ServiceVersion)
+							.AddAttributes(attributes)
+							.AddTelemetrySdk()
+							.AddEnvironmentVariableDetector());
+
+					if (settings.EnableConsoleExporter)
+					{
+						builder.AddConsoleExporter();
+					}
+					
+					foreach ((string name, OpenTelemetryProtocolExporterSettings exporterSettings) in settings.ProtocolExporters)
+					{
+						builder.AddOtlpExporter(name, exporter =>
+						{
+							exporter.Endpoint = exporterSettings.Endpoint;
+							exporter.Protocol = Enum.TryParse(exporterSettings.Protocol, true, out OtlpExportProtocol protocol) ? protocol : OtlpExportProtocol.Grpc;
+						});
+					}
+				});
+
+		}
+
 		private static void OnAddHealthChecks(IServiceCollection services)
 		{
 			services.AddHealthChecks()
