Horde: Expand properties in _envVars rather than batch.EnvVars.

diff --git a/Engine/Source/Programs/Horde/Horde.Agent/Execution/JobExecutor.cs b/Engine/Source/Programs/Horde/Horde.Agent/Execution/JobExecutor.cs
--- a/Engine/Source/Programs/Horde/Horde.Agent/Execution/JobExecutor.cs
+++ b/Engine/Source/Programs/Horde/Horde.Agent/Execution/JobExecutor.cs
@@ -1266,7 +1266,7 @@ namespace Horde.Agent.Execution
 
 			// Combine all the supplied environment variables together
 			Dictionary<string, string> newEnvVars = new Dictionary<string, string>(StringComparer.Ordinal);
-			foreach (KeyValuePair<string, string> envVar in _batch.Environment)
+			foreach (KeyValuePair<string, string> envVar in _envVars)
 			{
 				newEnvVars[envVar.Key] = StringUtils.ExpandProperties(envVar.Value, properties);
 			}
