Horde: Output a warning if a child process does not terminate gracefully due to pipes remaining open.

diff --git a/Engine/Source/Programs/Horde/Horde.Agent/Execution/JobExecutor.cs b/Engine/Source/Programs/Horde/Horde.Agent/Execution/JobExecutor.cs
--- a/Engine/Source/Programs/Horde/Horde.Agent/Execution/JobExecutor.cs
+++ b/Engine/Source/Programs/Horde/Horde.Agent/Execution/JobExecutor.cs
@@ -1158,15 +1158,24 @@ namespace Horde.Agent.Execution
 				{
 					using (ManagedProcess process = new ManagedProcess(processGroup, fileName, arguments, null, newEnvironment, null, ProcessPriorityClass.Normal))
 					{
-						async Task WaitForExitOrCancelAsync()
+						bool hasExited = false;
+						try
 						{
-							await process.WaitForExitAsync(cancellationToken);
-							cancellationSource.CancelAfter(TimeSpan.FromSeconds(60.0));
-						}
+							async Task WaitForExitOrCancelAsync()
+							{
+								await process.WaitForExitAsync(cancellationToken);
+								hasExited = true;
+								cancellationSource.CancelAfter(TimeSpan.FromSeconds(60.0));
+							}
 
-						Task cancelTask = WaitForExitOrCancelAsync();
-						await process.CopyToAsync((buffer, offset, length) => filter.WriteData(buffer.AsMemory(offset, length)), 4096, cancellationSource.Token);
-						await cancelTask;
+							Task cancelTask = WaitForExitOrCancelAsync();
+							await process.CopyToAsync((buffer, offset, length) => filter.WriteData(buffer.AsMemory(offset, length)), 4096, cancellationSource.Token);
+							await cancelTask;
+						}
+						catch (OperationCanceledException) when (hasExited)
+						{
+							logger.LogWarning("Process exited without closing output pipes; they may have been inherited by a child process that is still running.");
+						}
 						return process.ExitCode;
 					}
 				}
