Horde: Fix a warning about objects not being disposed in tests.

diff --git a/Engine/Source/Programs/Horde/Horde.Server.Tests/LogFileServiceTests.cs b/Engine/Source/Programs/Horde/Horde.Server.Tests/LogFileServiceTests.cs
--- a/Engine/Source/Programs/Horde/Horde.Server.Tests/LogFileServiceTests.cs
+++ b/Engine/Source/Programs/Horde/Horde.Server.Tests/LogFileServiceTests.cs
@@ -32,6 +32,7 @@ namespace Horde.Server.Tests
 		private readonly FakeClock _clock;
         private readonly LogFileService _logFileService;
 		private readonly ILoggerFactory _loggerFactory;
+		private readonly IStorageBackend _logStorageBackend;
 		private readonly ILogStorage _logStorage;
 
 		public LogFileServiceTest()
@@ -42,7 +43,8 @@ namespace Horde.Server.Tests
             ILogger<LogFileService> logger = _loggerFactory.CreateLogger<LogFileService>();
 
 			ILogBuilder logBuilder = new RedisLogBuilder(GetRedisServiceSingleton().ConnectionPool, NullLogger.Instance);
-			_logStorage = new PersistentLogStorage(new MemoryStorageBackend().ForType<PersistentLogStorage>(), NullLogger<PersistentLogStorage>.Instance);
+			_logStorageBackend = new MemoryStorageBackend();
+			_logStorage = new PersistentLogStorage(_logStorageBackend.ForType<PersistentLogStorage>(), NullLogger<PersistentLogStorage>.Instance);
 			_clock = new FakeClock();
 			TestOptions<ServerSettings> settingsOpts = new (new ServerSettings());
 			_logFileService = new LogFileService(logFileCollection, null!, logBuilder, _logStorage, _clock, null!, null!, settingsOpts, logger);
@@ -53,6 +55,7 @@ namespace Horde.Server.Tests
 			base.Dispose(disposing);
 
 			_logFileService.Dispose();
+			_logStorageBackend.Dispose();
 			_logStorage.Dispose();
 			_loggerFactory.Dispose();
 		}
