Horde: Change up arrow for Ben on artifacts browser
#rnx
#skipci

diff --git a/Engine/Source/Programs/Horde/HordeDashboard/src/components/jobDetailsV2/JobDetailArtifactsV3.tsx b/Engine/Source/Programs/Horde/HordeDashboard/src/components/jobDetailsV2/JobDetailArtifactsV3.tsx
--- a/Engine/Source/Programs/Horde/HordeDashboard/src/components/jobDetailsV2/JobDetailArtifactsV3.tsx
+++ b/Engine/Source/Programs/Horde/HordeDashboard/src/components/jobDetailsV2/JobDetailArtifactsV3.tsx
@@ -257,15 +257,6 @@ const BrowseHistory: React.FC<{ handler: ArtifactsHandler }> = observer(({ handl
                }
             }} />
          </Stack>
-         <Stack>
-            <IconButton disabled={upDisabled} style={{ fontSize: 14, paddingTop: 1 }} iconProps={{ iconName: 'ArrowUp' }} onClick={() => {
-               const nbrowse = handler.path?.split("/")
-               nbrowse?.pop();
-               if (nbrowse) {
-                  handler.browseTo(nbrowse.join("/"));
-               }
-            }} />
-         </Stack>
       </Stack>
    </Stack>
 
@@ -471,11 +462,9 @@ const JobDetailArtifactsInner: React.FC<{ handler: ArtifactsHandler }> = observe
    const items: BrowserItem[] = [];
 
    // use the up arrow instead
-   /*
    if (handler.path?.length) {
       items.push({ key: "navigate up", text: "..", type: BrowserType.NavigateUp });
    }
-   */
 
    browse.directories?.forEach(d => {
 
@@ -570,8 +559,13 @@ const JobDetailArtifactsInner: React.FC<{ handler: ArtifactsHandler }> = observe
                   {pathElements}
                </Stack>
             </Stack>}
-            {item.type === BrowserType.NavigateUp && <Stack data-selection-disabled verticalFill verticalAlign="center">
-               <Text>..</Text>
+            {item.type === BrowserType.NavigateUp && <Stack data-selection-disabled verticalFill horizontal verticalAlign="center" tokens={{ childrenGap: 9 }}>
+               <Stack>
+                  <FontIcon style={{ paddingTop: 1, fontSize: 15 }} iconName="ArrowUp" />
+               </Stack>
+               <Stack>
+                  <Text>..</Text>
+               </Stack>
             </Stack>
             }
          </Stack>
