Horde: Convert JobService to use OpenTelemetry-based tracing

diff --git a/Engine/Source/Programs/Horde/Horde.Server/Agents/AgentService.cs b/Engine/Source/Programs/Horde/Horde.Server/Agents/AgentService.cs
--- a/Engine/Source/Programs/Horde/Horde.Server/Agents/AgentService.cs
+++ b/Engine/Source/Programs/Horde/Horde.Server/Agents/AgentService.cs
@@ -819,15 +819,7 @@ namespace Horde.Server.Agents
 		/// <returns>List of leases matching the given criteria</returns>
 		public Task<List<ILease>> FindLeasesAsync(AgentId? agentId, SessionId? sessionId, DateTime? startTime, DateTime? finishTime, int index, int count)
 		{
-			using IScope scope = GlobalTracer.Instance.BuildSpan($"{nameof(AgentService)}.{nameof(FindLeasesAsync)}").StartActive();
-			scope.Span.SetTag("AgentId", agentId?.ToString());
-			scope.Span.SetTag("SessionId", sessionId?.ToString());
-			scope.Span.SetTag("StartTime", startTime?.ToString());
-			scope.Span.SetTag("FinishTime", finishTime?.ToString());
-			scope.Span.SetTag("Index", index);
-			scope.Span.SetTag("Count", count);
-			
-			using TelemetrySpan span = _tracer.StartActiveSpan($"{nameof(AgentService)}.{nameof(FindLeasesAsync)}-ot");
+			using TelemetrySpan span = _tracer.StartActiveSpan($"{nameof(AgentService)}.{nameof(FindLeasesAsync)}");
 			span.SetAttribute("AgentId", agentId?.ToString());
 			span.SetAttribute("SessionId", sessionId?.ToString());
 			span.SetAttribute("StartTime", startTime?.ToString());
diff --git a/Engine/Source/Programs/Horde/Horde.Server/Jobs/JobService.cs b/Engine/Source/Programs/Horde/Horde.Server/Jobs/JobService.cs
--- a/Engine/Source/Programs/Horde/Horde.Server/Jobs/JobService.cs
+++ b/Engine/Source/Programs/Horde/Horde.Server/Jobs/JobService.cs
@@ -30,6 +30,7 @@ using HordeCommon;
 using Microsoft.Extensions.Logging;
 using MongoDB.Bson;
 using MongoDB.Driver;
+using OpenTelemetry.Trace;
 using OpenTracing;
 using OpenTracing.Util;
 
@@ -84,6 +85,7 @@ namespace Horde.Server.Jobs
 		readonly ITemplateCollection _templateCollection;
 		readonly IssueService? _issueService;
 		readonly IPerforceService _perforceService;
+		readonly Tracer _tracer;
 		readonly ILogger _logger;
 
 		/// <summary>
@@ -116,7 +118,7 @@ namespace Horde.Server.Jobs
 		/// <summary>
 		/// Constructor
 		/// </summary>
-		public JobService(IJobCollection jobs, IGraphCollection graphs, IAgentCollection agents, IJobStepRefCollection jobStepRefs, IJobTimingCollection jobTimings, IUserCollection userCollection, INotificationTriggerCollection triggerCollection, JobTaskSource jobTaskSource, IStreamCollection streamCollection, ITemplateCollection templateCollection, IssueService issueService, IPerforceService perforceService, ILogger<JobService> logger)
+		public JobService(IJobCollection jobs, IGraphCollection graphs, IAgentCollection agents, IJobStepRefCollection jobStepRefs, IJobTimingCollection jobTimings, IUserCollection userCollection, INotificationTriggerCollection triggerCollection, JobTaskSource jobTaskSource, IStreamCollection streamCollection, ITemplateCollection templateCollection, IssueService issueService, IPerforceService perforceService, Tracer tracer, ILogger<JobService> logger)
 		{
 			_jobs = jobs;
 			_graphs = graphs;
@@ -130,6 +132,7 @@ namespace Horde.Server.Jobs
 			_templateCollection = templateCollection;
 			_issueService = issueService;
 			_perforceService = perforceService;
+			_tracer = tracer;
 			_logger = logger;
 
 			_jobTaskSource.OnJobScheduled += (pool, numAgentsOnline, job, graph, batchId) =>
@@ -160,35 +163,35 @@ namespace Horde.Server.Jobs
 		/// <returns>Unique id representing the job</returns>
 		public async Task<IJob> CreateJobAsync(JobId? jobId, StreamConfig streamConfig, TemplateId templateRefId, ContentHash templateHash, IGraph graph, string name, int change, int codeChange, CreateJobOptions options)
 		{
-			using IScope traceScope = GlobalTracer.Instance.BuildSpan("JobService.CreateJobAsync").StartActive();
-			traceScope.Span.SetTag("JobId", jobId?.ToString());
-			traceScope.Span.SetTag("Stream", streamConfig.Name);
-			traceScope.Span.SetTag("TemplateRefId", templateRefId);
-			traceScope.Span.SetTag("TemplateHash", templateHash);
-			traceScope.Span.SetTag("GraphId", graph.Id);
-			traceScope.Span.SetTag("Name", name);
-			traceScope.Span.SetTag("Change", change);
-			traceScope.Span.SetTag("CodeChange", codeChange);
-			traceScope.Span.SetTag("PreflightChange", options.PreflightChange);
-			traceScope.Span.SetTag("ClonedPreflightChange", options.ClonedPreflightChange);
-			traceScope.Span.SetTag("StartedByUserId", options.StartedByUserId.ToString());
-			traceScope.Span.SetTag("Priority", options.Priority.ToString());
-			traceScope.Span.SetTag("ShowUgsBadges", options.ShowUgsBadges);
-			traceScope.Span.SetTag("NotificationChannel", options.NotificationChannel ?? "null");
-			traceScope.Span.SetTag("NotificationChannelFilter", options.NotificationChannelFilter ?? "null");
-			traceScope.Span.SetTag("Arguments", String.Join(',', options.Arguments));
+			using TelemetrySpan span = _tracer.StartActiveSpan($"{nameof(JobService)}.{nameof(CreateJobAsync)}");
+			span.SetAttribute("JobId", jobId?.ToString());
+			span.SetAttribute("Stream", streamConfig.Name);
+			span.SetAttribute("TemplateRefId", templateRefId);
+			span.SetAttribute("TemplateHash", templateHash);
+			span.SetAttribute("GraphId", graph.Id);
+			span.SetAttribute("Name", name);
+			span.SetAttribute("Change", change);
+			span.SetAttribute("CodeChange", codeChange);
+			span.SetAttribute("PreflightChange", options.PreflightChange);
+			span.SetAttribute("ClonedPreflightChange", options.ClonedPreflightChange);
+			span.SetAttribute("StartedByUserId", options.StartedByUserId.ToString());
+			span.SetAttribute("Priority", options.Priority.ToString());
+			span.SetAttribute("ShowUgsBadges", options.ShowUgsBadges);
+			span.SetAttribute("NotificationChannel", options.NotificationChannel ?? "null");
+			span.SetAttribute("NotificationChannelFilter", options.NotificationChannelFilter ?? "null");
+			span.SetAttribute("Arguments", String.Join(',', options.Arguments));
 
 			if (options.AutoSubmit != null)
 			{
-				traceScope.Span.SetTag("AutoSubmit", options.AutoSubmit.Value);
+				span.SetAttribute("AutoSubmit", options.AutoSubmit.Value);
 			}
 			if (options.UpdateIssues != null)
 			{
-				traceScope.Span.SetTag("UpdateIssues", options.UpdateIssues.Value);
+				span.SetAttribute("UpdateIssues", options.UpdateIssues.Value);
 			}
 			if (options.JobTriggers != null)
 			{
-				traceScope.Span.SetTag("JobTriggers.Count", options.JobTriggers.Count);
+				span.SetAttribute("JobTriggers.Count", options.JobTriggers.Count);
 			}
 
 			JobId jobIdValue = jobId ?? JobId.GenerateNewId();
@@ -234,9 +237,9 @@ namespace Horde.Server.Jobs
 
 		private async Task AbortAnyDuplicateJobs(IJob newJob)
 		{
-			using IScope scope = GlobalTracer.Instance.BuildSpan("JobService.AbortAnyDuplicateJobs").StartActive();
-			scope.Span.SetTag("JobId", newJob.Id.ToString());
-			scope.Span.SetTag("JobName", newJob.Name);
+			using TelemetrySpan span = _tracer.StartActiveSpan($"{nameof(JobService)}.{nameof(AbortAnyDuplicateJobs)}");
+			span.SetAttribute("JobId", newJob.Id.ToString());
+			span.SetAttribute("JobName", newJob.Name);
 			
 			List<IJob> jobsToAbort = new List<IJob>();
 			if (newJob.PreflightChange > 0)
@@ -291,9 +294,9 @@ namespace Horde.Server.Jobs
 		/// <param name="job">The job to delete</param>
 		public async Task<bool> DeleteJobAsync(IJob job)
 		{
-			using IScope traceScope = GlobalTracer.Instance.BuildSpan("JobService.DeleteJobAsync").StartActive();
-			traceScope.Span.SetTag("JobId", job.Id.ToString());
-			traceScope.Span.SetTag("JobName", job.Name);
+			using TelemetrySpan span = _tracer.StartActiveSpan($"{nameof(JobService)}.{nameof(DeleteJobAsync)}");
+			span.SetAttribute("JobId", job.Id.ToString());
+			span.SetAttribute("JobName", job.Name);
 
 			using IDisposable scope = _logger.BeginScope("DeleteJobAsync({JobId})", job.Id);
 
@@ -350,9 +353,9 @@ namespace Horde.Server.Jobs
 		/// <param name="labelIdxToTriggerId">New trigger ID for a label in the job</param>
 		public async Task<IJob?> UpdateJobAsync(IJob job, string? name = null, Priority? priority = null, bool? autoSubmit = null, UserId? abortedByUserId = null, ObjectId? onCompleteTriggerId = null, List<Report>? reports = null, List<string>? arguments = null, KeyValuePair<int, ObjectId>? labelIdxToTriggerId = null)
 		{
-			using IScope traceScope = GlobalTracer.Instance.BuildSpan("JobService.UpdateJobAsync").StartActive();
-			traceScope.Span.SetTag("JobId", job.Id.ToString());
-			traceScope.Span.SetTag("Name", name);
+			using TelemetrySpan span = _tracer.StartActiveSpan($"{nameof(JobService)}.{nameof(UpdateJobAsync)}");
+			span.SetAttribute("JobId", job.Id.ToString());
+			span.SetAttribute("Name", name);
 			
 			using IDisposable scope = _logger.BeginScope("UpdateJobAsync({JobId})", job.Id);
 			for(IJob? newJob = job; newJob != null; newJob = await GetJobAsync(job.Id))
@@ -459,9 +462,9 @@ namespace Horde.Server.Jobs
 		/// <returns></returns>
 		async Task CancelLeaseAsync(AgentId agentId, LeaseId leaseId)
 		{
-			using IScope scope = GlobalTracer.Instance.BuildSpan("JobService.CancelLeaseAsync").StartActive();
-			scope.Span.SetTag("AgentId", agentId.ToString());
-			scope.Span.SetTag("LeaseId", leaseId.ToString());
+			using TelemetrySpan span = _tracer.StartActiveSpan($"{nameof(JobService)}.{nameof(CancelLeaseAsync)}");
+			span.SetAttribute("AgentId", agentId.ToString());
+			span.SetAttribute("LeaseId", leaseId.ToString());
 			
 			for (; ; )
 			{
@@ -543,25 +546,25 @@ namespace Horde.Server.Jobs
 		/// <returns>List of jobs matching the given criteria</returns>
 		public async Task<List<IJob>> FindJobsAsync(JobId[]? jobIds = null, StreamId? streamId = null, string? name = null, TemplateId[]? templates = null, int? minChange = null, int? maxChange = null, int? preflightChange = null, bool? preflightOnly = null, UserId? preflightStartedByUser = null, UserId? startedByUser = null, DateTimeOffset ? minCreateTime = null, DateTimeOffset? maxCreateTime = null, string? target = null, JobStepState[]? state = null, JobStepOutcome[]? outcome = null, DateTimeOffset? modifiedBefore = null, DateTimeOffset? modifiedAfter = null, int? index = null, int? count = null, bool consistentRead = true, bool? excludeUserJobs = null)
 		{
-			using IScope scope = GlobalTracer.Instance.BuildSpan("JobService.FindJobsAsync").StartActive();
-			scope.Span.SetTag("JobIds", (jobIds == null)? null : String.Join(',', jobIds));
-			scope.Span.SetTag("StreamId", streamId);
-			scope.Span.SetTag("Name", name);
-			scope.Span.SetTag("Templates", templates);
-			scope.Span.SetTag("MinChange", minChange);
-			scope.Span.SetTag("MaxChange", maxChange);
-			scope.Span.SetTag("PreflightChange", preflightChange);
-			scope.Span.SetTag("PreflightStartedByUser", preflightStartedByUser?.ToString());
-			scope.Span.SetTag("StartedByUser", startedByUser?.ToString());
-			scope.Span.SetTag("MinCreateTime", minCreateTime);
-			scope.Span.SetTag("MaxCreateTime", maxCreateTime);
-			scope.Span.SetTag("Target", target);
-			scope.Span.SetTag("State", state?.ToString());
-			scope.Span.SetTag("Outcome", outcome?.ToString());
-			scope.Span.SetTag("ModifiedBefore", modifiedBefore);
-			scope.Span.SetTag("ModifiedAfter", modifiedAfter);
-			scope.Span.SetTag("Index", index);
-			scope.Span.SetTag("Count", count);
+			using TelemetrySpan span = _tracer.StartActiveSpan($"{nameof(JobService)}.{nameof(FindJobsAsync)}");
+			span.SetAttribute("JobIds", (jobIds == null)? null : String.Join(',', jobIds));
+			span.SetAttribute("StreamId", streamId);
+			span.SetAttribute("Name", name);
+			span.SetAttribute("Templates", templates);
+			span.SetAttribute("MinChange", minChange);
+			span.SetAttribute("MaxChange", maxChange);
+			span.SetAttribute("PreflightChange", preflightChange);
+			span.SetAttribute("PreflightStartedByUser", preflightStartedByUser?.ToString());
+			span.SetAttribute("StartedByUser", startedByUser?.ToString());
+			span.SetAttribute("MinCreateTime", minCreateTime);
+			span.SetAttribute("MaxCreateTime", maxCreateTime);
+			span.SetAttribute("Target", target);
+			span.SetAttribute("State", state?.ToString());
+			span.SetAttribute("Outcome", outcome?.ToString());
+			span.SetAttribute("ModifiedBefore", modifiedBefore);
+			span.SetAttribute("ModifiedAfter", modifiedAfter);
+			span.SetAttribute("Index", index);
+			span.SetAttribute("Count", count);
 			
 			if (target == null && (state == null || state.Length == 0) && (outcome == null || outcome.Length == 0))
 			{
@@ -628,14 +631,14 @@ namespace Horde.Server.Jobs
 		/// <returns>List of jobs matching the given criteria</returns>
 		public async Task<List<IJob>> FindJobsByStreamWithTemplatesAsync(StreamId streamId, TemplateId[] templates, UserId? preflightStartedByUser = null, DateTimeOffset? maxCreateTime = null, DateTimeOffset? modifiedAfter = null, int? index = null, int? count = null, bool consistentRead = true)
 		{
-			using IScope scope = GlobalTracer.Instance.BuildSpan("JobService.FindJobsByStreamWithTemplatesAsync").StartActive();
-			scope.Span.SetTag("StreamId", streamId);
-			scope.Span.SetTag("Templates", templates);
-			scope.Span.SetTag("PreflightStartedByUser", preflightStartedByUser?.ToString());
-			scope.Span.SetTag("MaxCreateTime", maxCreateTime);
-			scope.Span.SetTag("ModifiedAfter", modifiedAfter);
-			scope.Span.SetTag("Index", index);
-			scope.Span.SetTag("Count", count);
+			using TelemetrySpan span = _tracer.StartActiveSpan($"{nameof(JobService)}.{nameof(FindJobsByStreamWithTemplatesAsync)}");
+			span.SetAttribute("StreamId", streamId);
+			span.SetAttribute("Templates", templates);
+			span.SetAttribute("PreflightStartedByUser", preflightStartedByUser?.ToString());
+			span.SetAttribute("MaxCreateTime", maxCreateTime);
+			span.SetAttribute("ModifiedAfter", modifiedAfter);
+			span.SetAttribute("Index", index);
+			span.SetAttribute("Count", count);
 			
 			return await _jobs.FindLatestByStreamWithTemplatesAsync(streamId, templates, preflightStartedByUser, maxCreateTime, modifiedAfter, index, count, consistentRead);
 		}
@@ -648,9 +651,9 @@ namespace Horde.Server.Jobs
 		/// <returns>True if the groups were updated to the given list. False if another write happened first.</returns>
 		public async Task<IJob?> TryUpdateGraphAsync(IJob job, IGraph newGraph)
 		{
-			using IScope traceScope = GlobalTracer.Instance.BuildSpan("JobService.TryUpdateGraphAsync").StartActive();
-			traceScope.Span.SetTag("Job", job.Id.ToString());
-			traceScope.Span.SetTag("NewGraph", newGraph.Id);
+			using TelemetrySpan span = _tracer.StartActiveSpan($"{nameof(JobService)}.{nameof(TryUpdateGraphAsync)}");
+			span.SetAttribute("Job", job.Id.ToString());
+			span.SetAttribute("NewGraph", newGraph.Id);
 
 			using IDisposable scope = _logger.BeginScope("TryUpdateGraphAsync({JobId})", job.Id);
 
@@ -672,8 +675,8 @@ namespace Horde.Server.Jobs
 		/// <returns>Timing info for the given job</returns>
 		public async Task<IJobTiming> GetJobTimingAsync(IJob job)
 		{
-			using IScope traceScope = GlobalTracer.Instance.BuildSpan("JobService.GetJobTimingAsync").StartActive();
-			traceScope.Span.SetTag("Job", job.Id.ToString());
+			using TelemetrySpan span = _tracer.StartActiveSpan($"{nameof(JobService)}.{nameof(GetJobTimingAsync)}");
+			span.SetAttribute("Job", job.Id.ToString());
 
 			using IDisposable scope = _logger.BeginScope("GetJobTimingAsync({JobId})", job.Id);
 
@@ -745,11 +748,11 @@ namespace Horde.Server.Jobs
 		/// <returns>Expected duration for the given step</returns>
 		async Task<JobStepTimingData> GetStepTimingInfo(StreamId streamId, TemplateId templateId, string nodeName, int? change)
 		{
-			using IScope traceScope = GlobalTracer.Instance.BuildSpan("JobService.GetStepTimingInfo").StartActive();
-			traceScope.Span.SetTag("StreamId", streamId);
-			traceScope.Span.SetTag("TemplateId", templateId);
-			traceScope.Span.SetTag("NodeName", nodeName);
-			traceScope.Span.SetTag("Change", change);
+			using TelemetrySpan span = _tracer.StartActiveSpan($"{nameof(JobService)}.{nameof(GetStepTimingInfo)}");
+			span.SetAttribute("StreamId", streamId);
+			span.SetAttribute("TemplateId", templateId);
+			span.SetAttribute("NodeName", nodeName);
+			span.SetAttribute("Change", change);
 			
 			// Find all the steps matching the given criteria
 			List<IJobStepRef> steps = await _jobStepRefs.GetStepsForNodeAsync(streamId, templateId, nodeName, change, false, 10);
@@ -793,11 +796,11 @@ namespace Horde.Server.Jobs
 		/// <returns>True if the job was updated, false if it was deleted</returns>
 		public async Task<IJob?> UpdateBatchAsync(IJob job, SubResourceId batchId, StreamConfig streamConfig, LogId? newLogId = null, JobStepBatchState? newState = null)
 		{
-			using IScope traceScope = GlobalTracer.Instance.BuildSpan("JobService.UpdateBatchAsync").StartActive();
-			traceScope.Span.SetTag("Job", job.Id.ToString());
-			traceScope.Span.SetTag("BatchId", batchId);
-			traceScope.Span.SetTag("NewLogId", newLogId?.ToString());
-			traceScope.Span.SetTag("NewState", newState.ToString());
+			using TelemetrySpan span = _tracer.StartActiveSpan($"{nameof(JobService)}.{nameof(UpdateBatchAsync)}");
+			span.SetAttribute("Job", job.Id.ToString());
+			span.SetAttribute("BatchId", batchId);
+			span.SetAttribute("NewLogId", newLogId?.ToString());
+			span.SetAttribute("NewState", newState.ToString());
 
 			using IDisposable scope = _logger.BeginScope("UpdateBatchAsync({JobId})", job.Id);
 
@@ -928,10 +931,10 @@ namespace Horde.Server.Jobs
 		/// <returns>True if the job was updated, false if it was deleted in the meantime</returns>
 		public async Task<IJob?> UpdateStepAsync(IJob job, SubResourceId batchId, SubResourceId stepId, StreamConfig streamConfig, JobStepState newState = JobStepState.Unspecified, JobStepOutcome newOutcome = JobStepOutcome.Unspecified, JobStepError? newError = null, bool? newAbortRequested = null, UserId? newAbortByUserId = null, LogId? newLogId = null, ObjectId? newNotificationTriggerId = null, UserId? newRetryByUserId = null, Priority? newPriority = null, List<Report>? newReports = null, Dictionary<string, string?>? newProperties = null)
 		{
-			using IScope traceScope = GlobalTracer.Instance.BuildSpan("JobService.UpdateStepAsync").StartActive();
-			traceScope.Span.SetTag("Job", job.Id.ToString());
-			traceScope.Span.SetTag("BatchId", batchId);
-			traceScope.Span.SetTag("StepId", stepId);
+			using TelemetrySpan span = _tracer.StartActiveSpan($"{nameof(JobService)}.{nameof(UpdateStepAsync)}");
+			span.SetAttribute("Job", job.Id.ToString());
+			span.SetAttribute("BatchId", batchId);
+			span.SetAttribute("StepId", stepId);
 			
 			using IDisposable scope = _logger.BeginScope("UpdateStepAsync({JobId}:{BatchId}:{StepId})", job.Id, batchId, stepId);
 			for (; ;)
@@ -973,10 +976,10 @@ namespace Horde.Server.Jobs
 		/// <returns>True if the job was updated, false if it was deleted in the meantime</returns>
 		public async Task<IJob?> TryUpdateStepAsync(IJob job, SubResourceId batchId, SubResourceId stepId, StreamConfig streamConfig, JobStepState newState = JobStepState.Unspecified, JobStepOutcome newOutcome = JobStepOutcome.Unspecified, JobStepError? newError = null, bool? newAbortRequested = null, UserId? newAbortByUserId = null, LogId? newLogId = null, ObjectId? newTriggerId = null, UserId? newRetryByUserId = null, Priority? newPriority = null, List<Report>? newReports = null, Dictionary<string, string?>? newProperties = null)
 		{
-			using IScope traceScope = GlobalTracer.Instance.BuildSpan("JobService.TryUpdateStepAsync").StartActive();
-			traceScope.Span.SetTag("Job", job.Id.ToString());
-			traceScope.Span.SetTag("BatchId", batchId);
-			traceScope.Span.SetTag("StepId", stepId);
+			using TelemetrySpan span = _tracer.StartActiveSpan($"{nameof(JobService)}.{nameof(TryUpdateStepAsync)}");
+			span.SetAttribute("Job", job.Id.ToString());
+			span.SetAttribute("BatchId", batchId);
+			span.SetAttribute("StepId", stepId);
 
 			using IDisposable scope = _logger.BeginScope("TryUpdateStepAsync({JobId}:{BatchId}:{StepId})", job.Id, batchId, stepId);
 
@@ -1008,10 +1011,10 @@ namespace Horde.Server.Jobs
 			{
 				job = newJob;
 
-				using IScope ddScope = GlobalTracer.Instance.BuildSpan("TryUpdateStepAsync").StartActive();
-				ddScope.Span.SetTag("JobId", job.Id.ToString());
-				ddScope.Span.SetTag("BatchId", batchId.ToString());
-				ddScope.Span.SetTag("StepId", stepId.ToString());
+				using TelemetrySpan updateSpan = _tracer.StartActiveSpan($"{nameof(JobService)}.{nameof(TryUpdateStepAsync)}.Update");
+				updateSpan.SetAttribute("JobId", job.Id.ToString());
+				updateSpan.SetAttribute("BatchId", batchId.ToString());
+				updateSpan.SetAttribute("StepId", stepId.ToString());
 				
 				if (oldState != newState || oldOutcome != newOutcome)
 				{
@@ -1020,7 +1023,7 @@ namespace Horde.Server.Jobs
 					// Send any updates for modified badges
 					if (oldLabelStates != null)
 					{
-						using IScope _ = GlobalTracer.Instance.BuildSpan("Send badge updates").StartActive();
+						using TelemetrySpan _ = _tracer.StartActiveSpan($"{nameof(JobService)}.{nameof(TryUpdateStepAsync)}.SendBadgeUpdates");
 						IReadOnlyList<(LabelState, LabelOutcome)> newLabelStates = job.GetLabelStates(graph);
 						OnLabelUpdate?.Invoke(job, oldLabelStates, newLabelStates);
 						await _jobTaskSource.UpdateUgsBadges(job, graph, oldLabelStates, newLabelStates);
@@ -1046,14 +1049,14 @@ namespace Horde.Server.Jobs
 					// Notify subscribers
 					if (newState == JobStepState.Completed)
 					{
-						using IScope _ = GlobalTracer.Instance.BuildSpan("Notify subscribers").StartActive();
+						using TelemetrySpan _ = _tracer.StartActiveSpan($"{nameof(JobService)}.{nameof(TryUpdateStepAsync)}.NotifySubscribers");
 						OnJobStepComplete?.Invoke(job, graph, batchId, stepId);
 					}
 
 					// Create any downstream jobs
 					if (newState == JobStepState.Completed && newOutcome != JobStepOutcome.Failure)
 					{
-						using IScope _ = GlobalTracer.Instance.BuildSpan("Create downstream jobs").StartActive();
+						using TelemetrySpan _ = _tracer.StartActiveSpan($"{nameof(JobService)}.{nameof(TryUpdateStepAsync)}.CreateDownstreamJobs");
 						for (int idx = 0; idx < job.ChainedJobs.Count; idx++)
 						{
 							IChainedJob jobTrigger = job.ChainedJobs[idx];
@@ -1071,7 +1074,7 @@ namespace Horde.Server.Jobs
 					// Update the jobstep ref if it completed
 					if (newState == JobStepState.Running || newState == JobStepState.Completed || newState == JobStepState.Aborted)
 					{
-						using IScope _ = GlobalTracer.Instance.BuildSpan("Update job step ref").StartActive();
+						using TelemetrySpan _ = _tracer.StartActiveSpan($"{nameof(JobService)}.{nameof(TryUpdateStepAsync)}.UpdateJobStepRef");
 						if (job.TryGetBatch(batchId, out IJobStepBatch? batch) && batch.TryGetStep(stepId, out IJobStep? step) && step.StartTimeUtc != null)
 						{
 							await _jobStepRefs.UpdateAsync(job, batch, step, graph);
@@ -1083,7 +1086,7 @@ namespace Horde.Server.Jobs
 					{
 						if (_issueService != null)
 						{
-							using IScope _ = GlobalTracer.Instance.BuildSpan("Update issues (V2)").StartActive();
+							using TelemetrySpan _ = _tracer.StartActiveSpan($"{nameof(JobService)}.{nameof(TryUpdateStepAsync)}.UpdateIssuesV2");
 							try
 							{
 								await _issueService.UpdateCompleteStep(job, graph, batchId, stepId);
@@ -1096,7 +1099,7 @@ namespace Horde.Server.Jobs
 					}
 				}
 
-				using (IScope dispatchScope = GlobalTracer.Instance.BuildSpan("Update queued jobs").StartActive())
+				using (TelemetrySpan _ = _tracer.StartActiveSpan($"{nameof(JobService)}.{nameof(TryUpdateStepAsync)}.UpdateQueuedJobs"))
 				{
 					// Notify the dispatch service that the job has changed
 					_jobTaskSource.UpdateQueuedJob(job, graph);
@@ -1117,9 +1120,9 @@ namespace Horde.Server.Jobs
 		/// <returns></returns>
 		private async Task<IJob> AutoSubmitChangeAsync(StreamConfig streamConfig, IJob job, IGraph graph)
 		{
-			using IScope traceScope = GlobalTracer.Instance.BuildSpan("JobService.AutoSubmitChangeAsync").StartActive();
-			traceScope.Span.SetTag("Job", job.Id.ToString());
-			traceScope.Span.SetTag("Graph", graph.Id);
+			using TelemetrySpan span = _tracer.StartActiveSpan($"{nameof(JobService)}.{nameof(AutoSubmitChangeAsync)}");
+			span.SetAttribute("Job", job.Id.ToString());
+			span.SetAttribute("Graph", graph.Id);
 			
 			int? change;
 			string message;
@@ -1250,12 +1253,12 @@ namespace Horde.Server.Jobs
 		{
 			CancellationToken cancellationToken = CancellationToken.None;
 
-			using IScope traceScope = GlobalTracer.Instance.BuildSpan("JobService.FireJobTriggerAsync").StartActive();
-			traceScope.Span.SetTag("Job", job.Id.ToString());
-			traceScope.Span.SetTag("Graph", graph.Id);
-			traceScope.Span.SetTag("JobTrigger.Target", jobTrigger.Target);
-			traceScope.Span.SetTag("JobTrigger.JobId", jobTrigger.JobId?.ToString());
-			traceScope.Span.SetTag("JobTrigger.TemplateRefId", jobTrigger.TemplateRefId);
+			using TelemetrySpan span = _tracer.StartActiveSpan($"{nameof(JobService)}.{nameof(FireJobTriggerAsync)}");
+			span.SetAttribute("Job", job.Id.ToString());
+			span.SetAttribute("Graph", graph.Id);
+			span.SetAttribute("JobTrigger.Target", jobTrigger.Target);
+			span.SetAttribute("JobTrigger.JobId", jobTrigger.JobId?.ToString());
+			span.SetAttribute("JobTrigger.TemplateRefId", jobTrigger.TemplateRefId);
 			
 			for (; ; )
 			{
diff --git a/Engine/Source/Programs/Horde/Horde.Server/Utilities/TracingExtensions.cs b/Engine/Source/Programs/Horde/Horde.Server/Utilities/TracingExtensions.cs
--- a/Engine/Source/Programs/Horde/Horde.Server/Utilities/TracingExtensions.cs
+++ b/Engine/Source/Programs/Horde/Horde.Server/Utilities/TracingExtensions.cs
@@ -2,8 +2,11 @@
 
 using System;
 using System.Collections.Generic;
+using System.Linq;
 using EpicGames.Core;
 using Horde.Server.Jobs;
+using Horde.Server.Streams;
+using OpenTelemetry.Trace;
 using OpenTracing;
 
 namespace Horde.Server.Utilities
@@ -51,4 +54,57 @@ namespace Horde.Server.Utilities
 			return span;
 		}
 	}
+	
+	/// <summary>
+	/// Extensions to handle Horde specific data types in the OpenTelemetry library
+	/// </summary>
+	public static class OpenTelemetrySpanExtensions
+	{
+		/// <summary>Set a key:value tag on the span</summary>
+		/// <returns>This span instance, for chaining</returns>
+		public static TelemetrySpan SetAttribute(this TelemetrySpan span, string key, ContentHash value)
+		{
+			span.SetAttribute(key, value.ToString());
+			return span;
+		}
+		
+		/// <summary>Set a key:value tag on the span</summary>
+		/// <returns>This span instance, for chaining</returns>
+		public static TelemetrySpan SetAttribute(this TelemetrySpan span, string key, SubResourceId value)
+		{
+			span.SetAttribute(key, value.ToString());
+			return span;
+		}
+		
+		/// <summary>Set a key:value tag on the span</summary>
+		/// <returns>This span instance, for chaining</returns>
+		public static TelemetrySpan SetAttribute(this TelemetrySpan span, string key, int? value)
+		{
+			if (value != null)
+			{
+				span.SetAttribute(key, value.Value);
+			}
+			return span;
+		}
+		
+		/// <summary>Set a key:value tag on the span</summary>
+		/// <returns>This span instance, for chaining</returns>
+		public static TelemetrySpan SetAttribute(this TelemetrySpan span, string key, DateTimeOffset? value)
+		{
+			if (value != null)
+			{
+				span.SetAttribute(key, value.ToString());
+			}
+			return span;
+		}
+		
+		/// <inheritdoc cref="ISpan.SetTag(System.String, System.String)"/>
+		public static TelemetrySpan SetAttribute(this TelemetrySpan span, string key, StreamId? value) => span.SetAttribute(key, value?.ToString());
+		
+		/// <inheritdoc cref="ISpan.SetTag(System.String, System.String)"/>
+		public static TelemetrySpan SetAttribute(this TelemetrySpan span, string key, TemplateId? value) => span.SetAttribute(key, value?.ToString());
+		
+		/// <inheritdoc cref="ISpan.SetTag(System.String, System.String)"/>
+		public static TelemetrySpan SetAttribute(this TelemetrySpan span, string key, TemplateId[]? values) => span.SetAttribute(key, values != null ? String.Join(',', values.Select(x => x.Id.ToString())) : null);
+	}
 }
\ No newline at end of file
