Horde: Always raise error dialog on put and delete actions
#rnx
#skipci

diff --git a/Engine/Source/Programs/Horde/HordeDashboard/src/backend/Fetch.ts b/Engine/Source/Programs/Horde/HordeDashboard/src/backend/Fetch.ts
--- a/Engine/Source/Programs/Horde/HordeDashboard/src/backend/Fetch.ts
+++ b/Engine/Source/Programs/Horde/HordeDashboard/src/backend/Fetch.ts
@@ -129,7 +129,7 @@ export class Fetch {
                     reason: reason,
                     mode: "PUT",
                     url: url
-                });
+                }, true);
 
                 reject(reason);
             });
@@ -154,7 +154,7 @@ export class Fetch {
                     reason: reason,
                     mode: "DELETE",
                     url: url
-                });
+                }, true);
 
                 reject(reason);
             });
