Horde: Refactor artifacts browser and work it into log view, also optimize log view to delay loading step history
#rnx
#skipci

diff --git a/Engine/Source/Programs/Horde/HordeDashboard/src/backend/JobDetails.ts b/Engine/Source/Programs/Horde/HordeDashboard/src/backend/JobDetails.ts
--- a/Engine/Source/Programs/Horde/HordeDashboard/src/backend/JobDetails.ts
+++ b/Engine/Source/Programs/Horde/HordeDashboard/src/backend/JobDetails.ts
@@ -1,13 +1,13 @@
 // Copyright Epic Games, Inc. All Rights Reserved.
 
+import { getTheme, mergeStyleSets, mergeStyles } from '@fluentui/react/lib/Styling';
 import { action, makeObservable, observable } from 'mobx';
-import { getTheme, mergeStyles, mergeStyleSets } from '@fluentui/react/lib/Styling';
+import moment from 'moment';
 import backend from '.';
-import { getBatchInitElapsed, getNiceTime, getStepElapsed, getStepETA, getStepFinishTime, getStepTimingDelta } from '../base/utilities/timeUtils';
+import { getBatchInitElapsed, getNiceTime, getStepETA, getStepElapsed, getStepFinishTime, getStepTimingDelta } from '../base/utilities/timeUtils';
 import { getBatchText } from '../components/JobDetailCommon';
-import { AgentData, ArtifactData, BatchData, EventData, GetGroupResponse, GetJobStepRefResponse, GetJobTimingResponse, GetLabelResponse, GetLabelStateResponse, GetLabelTimingInfoResponse, GroupData, IssueData, JobData, JobState, JobStepBatchState, JobStepError, JobStepOutcome, JobStepState, LabelState, NodeData, ReportPlacement, StepData, StreamData, TestData } from './Api';
+import { AgentData, ArtifactData, BatchData, EventData, GetGroupResponse, GetJobTimingResponse, GetLabelResponse, GetLabelStateResponse, GetLabelTimingInfoResponse, GroupData, IssueData, JobData, JobState, JobStepBatchState, JobStepError, JobStepOutcome, JobStepState, LabelState, NodeData, ReportPlacement, StepData, StreamData, TestData } from './Api';
 import { projectStore } from './ProjectStore';
-import moment from 'moment';
 
 const theme = getTheme();
 
@@ -95,8 +95,6 @@ export class JobDetails {
 
             const requests = [];
 
-            let historyIdx = -1;
-
             if (this.stepId !== stepId || this.labelIdx !== labelIdx) {
 
                 this.stepId = stepId;
@@ -104,13 +102,6 @@ export class JobDetails {
                 if (!this.suppressIssues) {
                     requests.push(this.getIssues());
                 }
-
-                this.history = undefined;
-
-                if (stepId) {
-                    historyIdx = requests.length;
-                    requests.push(backend.getJobStepHistory(this.jobdata!.streamId, this.getStepName(stepId, false), 1024, this.jobdata!.templateId!));
-                }
             }
 
             if (!this.isLogView) {
@@ -124,9 +115,6 @@ export class JobDetails {
             if (requests.length) {
                 await Promise.all(requests as any).then(async (values) => {
 
-                    if (historyIdx !== -1) {
-                        this.history = values[historyIdx] as GetJobStepRefResponse[];
-                    }
                 }).catch(reason => {
                     console.error(reason);
                 });
@@ -561,7 +549,6 @@ export class JobDetails {
         this.batches = [];
         this.events = [];
         this.issues = [];
-        this.history = undefined;
         this.labels = [];
         this.fatalError = undefined;
         this.state = JobStepState.Waiting;
@@ -876,15 +863,8 @@ export class JobDetails {
                 }
             }
 
-            let stepName = "";
             const batch = jobdata.batches?.find(b => b.steps.find(s => s.id === this.stepId));
             const stepNode = batch?.steps.find(s => s.id === this.stepId);
-            const groups = this.jobdata?.graphRef?.groups;
-            if (groups && stepNode && batch) {
-                stepName = groups[batch.groupIdx]?.nodes[stepNode.nodeIdx]?.name;
-            }
-
-            requests.push(backend.getJobArtifacts(this.id!));
 
             if (!this.isLogView) {
                 requests.push(backend.getJobTestData(this.id!));
@@ -905,10 +885,10 @@ export class JobDetails {
                 requests.push(this.getLogEvents(logId));
             }
 
-            let historyIdx = -1;
-            if (stepName) {
-                historyIdx = requests.length;
-                requests.push(backend.getJobStepHistory(jobdata.streamId, stepName, 1024, jobdata.templateId!));
+            let artifactsIdx = -1;
+            if (!jobdata.useArtifactsV2 && !this.artifacts?.length) {
+                artifactsIdx = requests.length;
+                requests.push(backend.getJobArtifacts(this.id!));
             }
 
             Promise.all(requests as any).then(async (values) => {
@@ -917,14 +897,12 @@ export class JobDetails {
                     return;
                 }
 
-                this.artifacts = values[0] as any;
-
                 if (!this.isLogView) {
                     this.testdata = values[1] as any;
                 }
 
-                if (historyIdx !== -1) {
-                    this.history = values[historyIdx] as any;
+                if (artifactsIdx !== -1) {
+                    this.artifacts = values[artifactsIdx] as any;
                 }
 
                 this.process();
@@ -997,7 +975,6 @@ export class JobDetails {
     artifacts: ArtifactData[] = []
     testdata: TestData[] = []
     issues: IssueData[] = [];
-    history?: GetJobStepRefResponse[] = undefined;
     private timing?: GetJobTimingResponse;
     agents: Map<string, AgentData> = new Map();
 
diff --git a/Engine/Source/Programs/Horde/HordeDashboard/src/components/LogSource.tsx b/Engine/Source/Programs/Horde/HordeDashboard/src/components/LogSource.tsx
--- a/Engine/Source/Programs/Horde/HordeDashboard/src/components/LogSource.tsx
+++ b/Engine/Source/Programs/Horde/HordeDashboard/src/components/LogSource.tsx
@@ -3,7 +3,7 @@
 import { action, makeObservable, observable } from "mobx";
 import moment, { Moment } from 'moment-timezone';
 import backend from '../backend';
-import { AgentData, BatchData, EventSeverity, GetLogEventResponse, IssueData, LeaseData, LogData, NodeData, StepData, StreamData } from "../backend/Api";
+import { AgentData, BatchData, EventSeverity, GetArtifactResponseV2, GetLogEventResponse, IssueData, LeaseData, LogData, NodeData, StepData, StreamData } from "../backend/Api";
 import { getBatchSummaryMarkdown, getStepSummaryMarkdown, JobDetails } from "../backend/JobDetails";
 import { getLeaseElapsed, getStepPercent } from '../base/utilities/timeUtils';
 import { BreadcrumbItem } from './Breadcrumbs';
@@ -391,7 +391,7 @@ export class JobLogSource extends LogSource {
 
    }
 
-   private detailsUpdated() {
+   private async detailsUpdated() {
 
       const details = this.jobDetails;
 
@@ -405,6 +405,18 @@ export class JobLogSource extends LogSource {
       const active = details.getLogActive(this.logData!.id);
 
       if (this.active !== active) {
+         
+         const step = details.getSteps().find(s => s.logId === this.logId);
+         if (step) {
+            const key = `job:${details.id!}/step:${step.id}`;
+            try {
+               const v = await backend.getJobArtifactsV2(undefined, [key]);
+               this.artifactsV2 = v.artifacts;
+            } catch (err) {
+               console.error(err);
+            }                   
+         }
+
          this.setActive(active);
       }
    }
@@ -532,6 +544,8 @@ export class JobLogSource extends LogSource {
    step?: StepData;
    node?: NodeData;
 
+   artifactsV2?: GetArtifactResponseV2[];
+
    jobDetails: JobDetails = new JobDetails(undefined, undefined, undefined, true);
 }
 
diff --git a/Engine/Source/Programs/Horde/HordeDashboard/src/components/LogView.tsx b/Engine/Source/Programs/Horde/HordeDashboard/src/components/LogView.tsx
--- a/Engine/Source/Programs/Horde/HordeDashboard/src/components/LogView.tsx
+++ b/Engine/Source/Programs/Horde/HordeDashboard/src/components/LogView.tsx
@@ -1,13 +1,13 @@
 // Copyright Epic Games, Inc. All Rights Reserved.
 
-import { DefaultButton, DetailsList, DetailsListLayoutMode, DetailsRow, DirectionalHint, Dropdown, FocusZone, FocusZoneDirection, IColumn, Icon, IconButton, IContextualMenuProps, IDetailsListProps, ITextField, List, Modal, ProgressIndicator, ScrollToMode, Selection, SelectionMode, SelectionZone, Separator, Spinner, SpinnerSize, Stack, Text, TextField } from '@fluentui/react';
+import { DefaultButton, DetailsList, DetailsListLayoutMode, DetailsRow, DirectionalHint, Dropdown, FocusZone, FocusZoneDirection, IColumn, Icon, IconButton, IContextualMenuItem, IContextualMenuProps, IDetailsListProps, ITextField, List, Modal, ProgressIndicator, ScrollToMode, Selection, SelectionMode, SelectionZone, Separator, Spinner, SpinnerSize, Stack, Text, TextField } from '@fluentui/react';
 import { action, makeObservable, observable } from 'mobx';
 import { observer } from 'mobx-react-lite';
 import moment from 'moment-timezone';
 import React, { useEffect, useState } from 'react';
 import { Link, useNavigate, useLocation, useParams } from 'react-router-dom';
 import backend from '../backend';
-import { ArtifactData, EventSeverity, GetChangeSummaryResponse, GetJobStepRefResponse, GetLogEventResponse, LogLevel } from '../backend/Api';
+import { ArtifactContextType, ArtifactData, EventSeverity, GetChangeSummaryResponse, GetJobStepRefResponse, GetLogEventResponse, LogLevel } from '../backend/Api';
 import { CommitCache } from '../backend/CommitCache';
 import dashboard from '../backend/Dashboard';
 import { JobDetails } from '../backend/JobDetails';
@@ -27,6 +27,7 @@ import { lineRenderStyleNormal, lineRenderStyleSmall, logMetricNormal, logMetric
 import { PrintException } from './PrintException';
 import { StepRefStatusIcon } from './StatusIcon';
 import { TopNav } from './TopNav';
+import { JobArtifactsModal } from './artifacts/ArtifactsModal';
 
 class LogHandler {
 
@@ -191,8 +192,28 @@ const StepHistoryModal: React.FC<{ jobDetails: JobDetails, stepId: string | unde
    const navigate = useNavigate();
    const location = useLocation();
    const [commitState, setCommitState] = useState<{ target?: ChangeContextMenuTarget, commit?: GetChangeSummaryResponse, rangeCL?: number }>({});
+   const [stepHistory, setStepHistory] = useState<GetJobStepRefResponse[] | undefined>(undefined);
 
-   const stepHistory = jobDetails.history;
+   const jobData = jobDetails.jobdata;
+
+   if (!jobData || !jobData.streamId || !jobData.templateId) {
+      return null;
+   }
+
+   if (stepHistory === undefined) {
+      backend.getJobStepHistory(jobData.streamId, jobDetails.getStepName(stepId, false), 1024, jobData!.templateId!).then(r => {
+         setStepHistory(r);
+      })
+
+      return <Modal isOpen={true} styles={{ main: { padding: 8, width: 1084, height: '624px', backgroundColor: '#FFFFFF' } }} className={hordeClasses.modal} onDismiss={() => { onClose() }}>
+         <Stack style={{paddingTop: 24}} horizontalAlign='center' tokens={{childrenGap: 18}}>
+            <Stack>
+               <Text variant='mediumPlus'>Loading Step History</Text>
+            </Stack>
+            <Spinner size={SpinnerSize.large} />
+         </Stack>
+      </Modal>
+   }
 
    // subscribe
    if (commitCache.updated) { }
@@ -275,7 +296,7 @@ const StepHistoryModal: React.FC<{ jobDetails: JobDetails, stepId: string | unde
          const url = `${location.pathname}?agentId=${agentId}`;
 
          return <Stack verticalAlign="center" horizontalAlign="center" tokens={{ childrenGap: 0, padding: 0 }} style={{ width: "100%", height: "100%" }}>
-            <a href={url} onClick={(ev) => { ev.preventDefault(); ev.stopPropagation(); navigate(url, {replace: true}); }}><Stack horizontal horizontalAlign={"end"} verticalFill={true} tokens={{ childrenGap: 0, padding: 0 }}><Text>{agentId}</Text></Stack></a>
+            <a href={url} onClick={(ev) => { ev.preventDefault(); ev.stopPropagation(); navigate(url, { replace: true }); }}><Stack horizontal horizontalAlign={"end"} verticalFill={true} tokens={{ childrenGap: 0, padding: 0 }}><Text>{agentId}</Text></Stack></a>
          </Stack>
       }
 
@@ -342,7 +363,7 @@ const StepHistoryModal: React.FC<{ jobDetails: JobDetails, stepId: string | unde
    };
 
    return (<Modal isOpen={true} styles={{ main: { padding: 8, width: 1084, height: '624px', backgroundColor: '#FFFFFF' } }} className={hordeClasses.modal} onDismiss={() => { onClose() }}>
-      {commitState.target && <ChangeContextMenu target={commitState.target} job={jobDetails.jobdata} commit={commitState.commit} rangeCL={ commitState.rangeCL } onDismiss={() => setCommitState({})} />}
+      {commitState.target && <ChangeContextMenu target={commitState.target} job={jobDetails.jobdata} commit={commitState.commit} rangeCL={commitState.rangeCL} onDismiss={() => setCommitState({})} />}
       <Stack styles={{ root: { paddingTop: 8, paddingLeft: 24, paddingRight: 12, paddingBottom: 8 } }}>
          <Stack tokens={{ childrenGap: 12 }}>
             <Stack horizontal styles={{ root: { padding: 8 } }}>
@@ -421,7 +442,7 @@ export const LogList: React.FC<{ logId: string }> = observer(({ logId }) => {
    const [searchState, setSearchState] = useState<{ search?: string, results?: number[], curRequest?: any }>({});
    const [issueHistory, setIssueHistory] = useState(false);
    const [logHistory, setLogHistory] = useState(false);
-   const [logArtifacts, setLogArtifacts] = useState(false);
+   const [logArtifacts, setLogArtifacts] = useState("");
 
    let [historyAgentId, setHistoryAgentId] = useState<string | undefined>(undefined);
 
@@ -649,9 +670,9 @@ export const LogList: React.FC<{ logId: string }> = observer(({ logId }) => {
                }
             }
          }
-      
+
          return (
-            <Stack key={`key_log_line_${item.lineNumber}`} style={{ width: "max-content", height: handler.lineHeight }} onClick={() => {navigate(location.pathname + `?lineindex=${item.lineNumber - 1}`, {replace: true})} }>
+            <Stack key={`key_log_line_${item.lineNumber}`} style={{ width: "max-content", height: handler.lineHeight }} onClick={() => { navigate(location.pathname + `?lineindex=${item.lineNumber - 1}`, { replace: true }) }}>
                <div style={{ position: "relative" }}>
                   <Stack className={styles.logLine} tokens={{ childrenGap: 8 }} horizontal disableShrink={true}>
                      <Stack styles={{ root: { color: "#c0c0c0", width: 80, textAlign: "right", userSelect: "none", fontSize: handler.fontSize } }}>{prefix + item.lineNumber}</Stack>
@@ -745,7 +766,7 @@ export const LogList: React.FC<{ logId: string }> = observer(({ logId }) => {
 
                const lines = logResponse.lines;
 
-            
+
                if (lines.length) {
 
                   logListKey++;
@@ -755,7 +776,7 @@ export const LogList: React.FC<{ logId: string }> = observer(({ logId }) => {
                      curRequest: undefined,
                      results: lines
                   });
-   
+
 
                   handler.stopTrailing();
                   let lineIdx = lines[curSearchIdx] - 10;
@@ -769,14 +790,14 @@ export const LogList: React.FC<{ logId: string }> = observer(({ logId }) => {
                      inTimeout = false;
                      listRef?.scrollToIndex(lineIdx, () => { return handler.lineHeight; }, ScrollToMode.top);
                   }, 250)
-                  
+
                } else {
                   setSearchState({
                      search: newValue,
                      curRequest: undefined,
                      results: lines
                   });
-   
+
                }
 
             };
@@ -821,16 +842,63 @@ export const LogList: React.FC<{ logId: string }> = observer(({ logId }) => {
       menuProps.items.push({
          key: 'jobstep_history',
          text: 'Step History',
-         disabled: !(logSource as JobLogSource)?.step || !fixme.history?.length,
          onClick: () => setLogHistory(true)
       })
 
-      menuProps.items.push({
-         key: 'jobstep_artifacts',
-         disabled: artifacts.length === 0,
-         text: `Step Artifacts`,
-         onClick: () => setLogArtifacts(true)
-      })
+      if (!fixme.jobdata?.useArtifactsV2) {
+         menuProps.items.push({
+            key: 'jobstep_artifacts',
+            disabled: artifacts.length === 0,
+            text: `Step Artifacts`,
+            onClick: () => setLogArtifacts("legacy")
+         })
+      } else {
+
+         const stepArtifacts = (logSource as JobLogSource).artifactsV2;
+
+         const atypes = new Map<ArtifactContextType, number>();
+
+         stepArtifacts?.forEach(a => {
+            let c = atypes.get(a.type) ?? 0;
+            c++;
+            atypes.set(a.type, c);
+         });
+
+         const opsList: IContextualMenuItem[] = [];
+
+         opsList.push({
+            key: 'stepops_artifacts_step',
+            text: "Step Artifacts",
+            iconProps: { iconName: "Folder" },
+            disabled: !atypes.get("step-saved"),
+            onClick: () => { setLogArtifacts("step-saved") }
+         });
+
+         opsList.push({
+            key: 'stepops_artifacts_output',
+            text: "Output Artifacts",
+            iconProps: { iconName: "MenuOpen" },
+            disabled: !atypes.get("step-output"),
+            onClick: () => { setLogArtifacts("step-output") }
+         });
+
+         opsList.push({
+            key: 'stepops_artifacts_trace',
+            text: "Trace Artifacts",
+            iconProps: { iconName: "SearchTemplate" },
+            disabled: !atypes.get("step-trace"),
+            onClick: () => { setLogArtifacts("step-trace") }
+         });
+
+         menuProps.items.push({
+            key: 'jobstep_artifacts',
+            text: `Artifacts`,
+            subMenuProps: {
+               items: opsList
+            }
+         })
+      }
+
    }
 
    function updateError() {
@@ -938,14 +1006,15 @@ export const LogList: React.FC<{ logId: string }> = observer(({ logId }) => {
    return <Stack>
       {!!fixme && <IssueModalV2 issueId={query.get("issue")} popHistoryOnClose={issueHistory} />}
       {!!fixme && logHistory && <StepHistoryModal jobDetails={fixme!} stepId={fixme!.stepByLogId(logId)?.id} onClose={() => setLogHistory(false)} />}
-      {!!fixme && logArtifacts && <StepArtifactsModal jobDetails={fixme!} stepId={fixme!.stepByLogId(logId)?.id} onClose={() => setLogArtifacts(false)} />}
-      {!!historyAgentId && <HistoryModal agentId={historyAgentId} onDismiss={() => { navigate(baseUrl, {replace: true}); setHistoryAgentId(undefined) }} />}
+      {!!fixme && logArtifacts === "legacy" && <StepArtifactsModal jobDetails={fixme!} stepId={fixme!.stepByLogId(logId)?.id} onClose={() => setLogArtifacts("")} />}
+      {!!fixme && logArtifacts && logArtifacts !== "legacy" && <JobArtifactsModal jobId={fixme!.jobdata!.id} stepId={fixme!.stepByLogId(logId)?.id!} artifacts={(logSource as JobLogSource).artifactsV2} contextType={logArtifacts as ArtifactContextType} onClose={() => setLogArtifacts("")} />}
+      {!!historyAgentId && <HistoryModal agentId={historyAgentId} onDismiss={() => { navigate(baseUrl, { replace: true }); setHistoryAgentId(undefined) }} />}
       <Breadcrumbs items={logSource?.crumbs ?? []} title={logSource?.crumbTitle} />
       <Stack tokens={{ childrenGap: 0 }} style={{ backgroundColor: "#FFFFFF", paddingTop: 12 }}>
          <Stack horizontal >
             <div key={`windowsize_logview1_${windowSize.width}_${windowSize.height}`} style={{ width: vw / 2 - (1440 / 2), flexShrink: 0 }} />
             <Stack tokens={{ childrenGap: 0, maxWidth: 1440 }} disableShrink={true} styles={{ root: { width: "100%", backgroundColor: "#FFFFFF", paddingLeft: 4, paddingRight: 24, paddingTop: 12 } }}>
-               <Stack horizontal style={{paddingBottom: 4}}>
+               <Stack horizontal style={{ paddingBottom: 4 }}>
                   <Stack className={hordeClasses.button} horizontal horizontalAlign={"start"} verticalAlign="center" tokens={{ childrenGap: 8 }}>
                      <Stack horizontal tokens={{ childrenGap: 2 }}>
                         <DefaultButton disabled={!errors.length} className={errors.length ? handler.style.errorButton : handler.style.errorButtonDisabled}
diff --git a/Engine/Source/Programs/Horde/HordeDashboard/src/components/artifacts/ArtifactsModal.tsx b/Engine/Source/Programs/Horde/HordeDashboard/src/components/artifacts/ArtifactsModal.tsx
new file mode 100644
--- /dev/null
+++ b/Engine/Source/Programs/Horde/HordeDashboard/src/components/artifacts/ArtifactsModal.tsx
@@ -0,0 +1,675 @@
+// Copyright Epic Games, Inc. All Rights Reserved.
+
+import { DetailsList, DetailsListLayoutMode, FontIcon, IColumn, IconButton, Modal, PrimaryButton, ScrollablePane, Selection, SelectionMode, SelectionZone, Spinner, SpinnerSize, Stack, Text, mergeStyleSets } from "@fluentui/react";
+import { action, makeObservable, observable } from "mobx";
+import { observer } from "mobx-react-lite";
+import { useEffect, useState } from "react";
+import backend from "../../backend";
+import { ArtifactContextType, GetArtifactDirectoryEntryResponse, GetArtifactDirectoryResponse, GetArtifactFileEntryResponse, GetArtifactResponseV2 } from "../../backend/Api";
+import dashboard from "../../backend/Dashboard";
+import { hordeClasses } from "../../styles/Styles";
+
+
+enum BrowserType {
+   Directory,
+   File,
+   NavigateUp
+}
+
+type BrowserItem = {
+   key: string;
+   text: string;
+   icon?: string;
+   size?: number;
+   type: BrowserType;
+}
+
+
+class ArtifactsHandler {
+
+   constructor(jobId: string, stepId: string, contextType: ArtifactContextType, artifacts?:GetArtifactResponseV2[]) {
+      makeObservable(this);
+      this.jobId = jobId;
+      this.stepId = stepId;
+      this.context = contextType;
+      this.artifacts = artifacts;
+
+      this.set();
+   }
+
+   @observable
+   updated = 0
+
+   @action
+   updateReady() {
+      this.updated++;
+   }
+
+   private async set() {
+
+      let artifacts: GetArtifactResponseV2[] | undefined = this.artifacts;
+
+      if (!artifacts) {         
+
+         const key = `job:${this.jobId}/step:${this.stepId}`;
+         try {
+            const v = await backend.getJobArtifactsV2(undefined, [key]);
+            artifacts = v.artifacts;               
+         } catch (err) {
+            console.error(err);
+         } 
+   
+         if (!artifacts) {
+            console.error(`Missing artifacts for job: ${this.jobId} step: ${this.stepId}`);
+            return;
+         }
+   
+         this.artifacts = artifacts;   
+      }      
+
+      if (!artifacts) {
+         console.error(`Missing artifacts for job: ${this.jobId} step: ${this.stepId}`);
+         return;
+      }
+
+      this.artifacts = artifacts;
+
+      this.contexts = [];
+
+      let a = artifacts.find(a => a.type === "step-saved");
+      if (a) {
+         this.contexts.push("step-saved");
+      }
+
+      a = artifacts.find(a => a.type === "step-output");
+      if (a) {
+         this.contexts.push("step-output");
+      }
+
+      a = artifacts.find(a => a.type === "step-trace");
+      if (a) {
+         this.contexts.push("step-trace");
+      }
+
+      if (!this.context) {
+         console.error("Artifact browser has no context");
+         this.updateReady();
+         return;
+      }
+
+      a = artifacts.find(a => a.type === this.context)!;
+
+      if (!a) {
+         console.error("Unable to find artifact for context", this.context, artifacts);
+         this.updateReady();
+         return;
+      }
+
+      this.artifact = a;
+
+      this.loading = true;
+      this.updateReady();
+
+      this.browse = await backend.getBrowseArtifacts(a.id);
+
+      this.loading = false;
+      this.updateReady();
+
+   }
+
+   hasContext(c: ArtifactContextType) {
+
+      return !!this.artifacts?.find(a => a.type === c);
+
+   }
+
+   async browseTo(path: string, push = true) {
+
+      if (!this.artifact) {
+         return;
+      }
+
+      this.loading = true;
+      this.updateReady();
+      this.browse = await backend.getBrowseArtifacts(this.artifact.id, path);
+
+      this.path = path;
+
+      if (push) {
+         this.history.push(path);
+      }
+
+      this.loading = false;
+      this.updateReady();
+
+   }
+
+   get currentSelection(): { filesSelected: number, directoriesSelected: number, size: number, items: BrowserItem[] } {
+
+      let result = {
+         filesSelected: 0,
+         directoriesSelected: 0,
+         size: 0,
+         items: [] as BrowserItem[]
+      }
+
+      const browse = this.browse;
+      if (!browse) {
+         return result;
+      }
+
+      let selection = this.selection?.getSelection() as (BrowserItem[] | undefined);
+
+      if (!selection?.length) {
+         selection = [];
+
+         browse.directories?.forEach(d => {
+            selection!.push({ key: d.hash, text: d.name, icon: "Folder", type: BrowserType.Directory, size: d.length });
+         });
+
+         browse.files?.forEach(d => {
+            selection!.push({ key: d.hash, text: d.name, icon: "Document", type: BrowserType.File, size: d.length });
+         });
+      }
+
+      selection.forEach(b => {
+         result.items.push(b);
+         if (b.type === BrowserType.Directory) {
+            result.directoriesSelected++;
+            result.size += b.size ?? 0;
+         }
+         if (b.type === BrowserType.File) {
+            result.filesSelected++;
+            result.size += b.size ?? 0;
+         }
+      });
+
+      return result;
+
+   }
+
+   get directories(): GetArtifactDirectoryEntryResponse[] {
+      return this.browse?.directories ?? [];
+   }
+
+   get files(): GetArtifactFileEntryResponse[] {
+      return this.browse?.files ?? [];
+   }
+
+   clear() {
+      this.selection = undefined;
+      this.selectionCallback = undefined;
+      this.path = undefined;
+      this.browse = undefined;
+      this.artifact = undefined;
+      this.artifacts = undefined;
+      this.contexts = undefined;
+      this.history = [];
+      this.stepId = "";
+      this.loading = false;
+   }
+   
+
+   selectionCallback?: () => void;
+
+   path?: string;
+
+   selection?: Selection = new Selection({ canSelectItem: (item: any) => { return item.type !== BrowserType.NavigateUp; }, onSelectionChanged: () => { if (this.selectionCallback) this.selectionCallback() } });
+
+   browse?: GetArtifactDirectoryResponse;
+
+   artifact?: GetArtifactResponseV2;
+   artifacts?: GetArtifactResponseV2[];
+
+   private contexts?: ArtifactContextType[];
+
+   readonly context: ArtifactContextType;
+
+   jobId: string;
+   stepId: string;
+
+   history: string[] = [];
+
+   loading = false;
+}
+
+const styles = mergeStyleSets({
+   list: {
+      selectors: {
+         'a': {
+            height: "unset !important",
+         },
+         '.ms-List-cell': {
+
+            borderTop: "1px solid #EDEBE9",
+            borderRight: "1px solid #EDEBE9",
+            borderLeft: "1px solid #EDEBE9"
+         },
+         '.ms-List-cell:nth-last-child(-n + 1)': {
+            borderBottom: "1px solid #EDEBE9"
+         },
+         ".ms-DetailsRow #artifactview": {
+            opacity: 0
+         },
+         ".ms-DetailsRow:hover #artifactview": {
+            opacity: 1
+         },
+      }
+   }
+});
+
+const BrowseHistory: React.FC<{ handler: ArtifactsHandler }> = observer(({ handler }) => {
+
+   // subscribe
+   if (handler.updated) { };
+
+   let backDisabled = !handler.history.length;
+
+   return <Stack style={{ paddingRight: 12 }}>
+      <Stack horizontal tokens={{ childrenGap: 8 }}>
+         <Stack>
+            <IconButton disabled={backDisabled} style={{ fontSize: 14, paddingTop: 1 }} iconProps={{ iconName: 'ArrowLeft' }} onClick={() => {
+               if (handler.history.length === 1) {
+                  handler.history = [];
+                  handler.browseTo("", false);
+               } else {
+                  handler.history.pop();
+                  handler.browseTo(handler.history[handler.history.length - 1], false);
+               }
+            }} />
+         </Stack>
+      </Stack>
+   </Stack>
+
+});
+
+const BrowseBreadCrumbs: React.FC<{ handler: ArtifactsHandler }> = observer(({ handler }) => {
+
+   // subscribe
+   if (handler.updated) { }
+
+   const fontSize = 13;
+
+   let rootName = "Step";
+   if (handler.context === "step-output") {
+      rootName = "Output";
+   }
+   if (handler.context === "step-trace") {
+      rootName = "Trace";
+   }
+
+   if (!handler.path) {
+      return <Stack horizontal verticalAlign="center" styles={{ root: { height: 40 } }}>
+         <Stack>
+            <BrowseHistory handler={handler} />
+         </Stack>
+         <Stack>
+            <Text style={{ fontSize: fontSize, fontWeight: 600, paddingLeft: 2, paddingRight: 4 }}>{rootName} /</Text>
+         </Stack>
+      </Stack>
+
+   }
+
+   const ppath = handler.path.split("/");
+
+   const elements = ppath.map((e, index) => {
+      const path = ppath.slice(0, index + 1).join("/");
+
+      const last = index === (ppath.length - 1);
+
+      const color = last ? undefined : (dashboard.darktheme ? "#55B7FF" : "#0078D4");
+      const cursor = last ? undefined : "pointer";
+
+      return <Stack horizontal onClick={() => handler.browseTo(path)} style={{ cursor: cursor }}>
+         <Stack>
+            <Text style={{ fontSize: fontSize, color: color, fontWeight: last ? 600 : undefined }}>{e}</Text>
+         </Stack>
+         {index !== (ppath.length - 1) && <Stack style={{ paddingLeft: 4, paddingRight: 4 }}><Text style={{ fontSize: fontSize, fontWeight: 600 }}>/</Text></Stack>}
+      </Stack>
+   });
+
+   elements.unshift(<Stack style={{ cursor: "pointer" }} onClick={() => handler.browseTo("")}>
+      <Stack style={{ paddingLeft: 2, paddingRight: 4 }}>
+         <Text style={{ fontSize: fontSize, color: (dashboard.darktheme ? "#55B7FF" : "#0078D4"), fontWeight: 600 }}>{rootName} /</Text>
+      </Stack>
+   </Stack>);
+
+   return <Stack horizontal verticalAlign="center" styles={{ root: { height: 40 } }}>
+      <Stack>
+         <BrowseHistory handler={handler} />
+      </Stack>
+      <Stack horizontal wrap style={{ width: 800 }}>
+         {elements}
+      </Stack>
+   </Stack>
+
+})
+
+function formatBytes(bytes: number, decimals = 2) {
+   if (!+bytes) return '0 Bytes'
+
+   const k = 1024
+   const dm = decimals < 0 ? 0 : decimals
+   const sizes = ['Bytes', 'KB', 'MB', 'GB', 'TB', 'PB', 'EB', 'ZB', 'YB']
+
+   const i = Math.floor(Math.log(bytes) / Math.log(k))
+
+   return `${parseFloat((bytes / Math.pow(k, i)).toFixed(dm))} ${sizes[i]}`
+}
+
+
+const DownloadButton: React.FC<{ handler: ArtifactsHandler }> = observer(({ handler }) => {
+
+   const [downloading, setDownloading] = useState(false);
+   const [selectKey, setSelectionKey] = useState(0);
+
+
+   // subscribe
+   if (handler.updated) { };
+
+   const browse = handler.browse;
+
+   if (!browse) {
+      return null;
+   }
+
+   handler.selectionCallback = () => {
+      setSelectionKey(selectKey + 1);
+   }
+
+   let buttonText = "Download";
+
+
+   const selection = handler.currentSelection;
+
+
+   let sizeText = "0KB";
+   if (selection.size) {
+      sizeText = formatBytes(selection.size, (selection.size < (1024 * 1024)) ? 0 : 1)
+   }
+
+   if (selection.directoriesSelected || selection.filesSelected > 1) {
+      buttonText = `Download (${sizeText})`;
+   } else if (selection.filesSelected === 1) {
+      buttonText = `Download (${sizeText})`;
+   }
+
+   if (downloading) {
+      buttonText = "Downloading";
+   }
+
+   return <Stack horizontal verticalAlign="center" tokens={{ childrenGap: 8 }}>
+      {downloading && <Stack><Spinner size={SpinnerSize.large} /></Stack>}
+      <PrimaryButton styles={{ root: { fontFamily: 'Horde Open Sans SemiBold !important' } }} disabled={!selection.size || downloading} onClick={async () => {
+
+         const selection = handler.currentSelection.items;
+
+         if (!selection?.length || !handler.artifact) {
+            return;
+         }
+
+         // download a single file
+         if (selection.length === 1) {
+            const item = selection[0] as BrowserItem;
+            if (item.type === BrowserType.File) {
+
+               try {
+                  setDownloading(true);
+                  await backend.downloadArtifactV2(handler.artifact.id, (handler.path ? handler.path + "/" : "") + item.text, item.text);
+               } catch (err) {
+                  console.error(err);
+               } finally {
+                  setDownloading(false);
+               }
+
+               return;
+            }
+         }
+
+         const path = handler.path ? handler.path + "/" : "";
+
+         const filters = selection.map(s => {
+            const item = s as BrowserItem;
+
+            if (item.type === BrowserType.NavigateUp) {
+               return "";
+            }
+
+            if (item.type === BrowserType.Directory) {
+               return `${path}${item.text}/...`;
+            }
+            return `${path}${item.text}`;
+         }).filter(f => !!f);
+
+         let context = handler.context;
+         let contextName = "step";
+         if (context === "step-output") {
+            contextName = "output";
+         }
+
+         if (context === "step-trace") {
+            contextName = "trace";
+         }
+
+         const filename = "horde-" + contextName + "-artifacts-" + handler.jobId + '-' + handler.stepId + ".zip";
+
+         try {
+            setDownloading(true);
+            await backend.downloadArtifactZipV2(handler.artifact.id, { filter: filters }, filename);
+         } catch (err) {
+            console.error(err);
+         } finally {
+            setDownloading(false);
+         }
+
+      }}>{buttonText}</PrimaryButton>
+   </Stack>
+
+});
+
+let idcounter = 0;
+
+const JobDetailArtifactsInner: React.FC<{ handler: ArtifactsHandler }> = observer(({ handler }) => {
+
+   // subscribe
+   if (handler.updated) { }
+
+   const browse = handler.browse;
+
+   if (!browse) {
+      return <Stack><Spinner size={ SpinnerSize.large}/></Stack>;
+   }
+
+   const items: BrowserItem[] = [];
+
+   // use the up arrow instead
+   if (handler.path?.length) {
+      items.push({ key: "navigate up", text: "..", type: BrowserType.NavigateUp });
+   }
+
+   browse.directories?.forEach(d => {
+
+      function recurseDirectories(dir: GetArtifactDirectoryEntryResponse, flattened: GetArtifactDirectoryEntryResponse[]) {
+         if (!dir.directories?.length) {
+            const name = flattened.length ? flattened.map(d => d.name).join("/") + "/" + dir.name : dir.name;
+            items.push({ key: d.hash, text: name, icon: "Folder", type: BrowserType.Directory, size: d.length });
+         } else {
+            flattened.push(dir);
+            dir.directories.forEach(d => recurseDirectories(d, [...flattened]));
+         }
+      }
+
+      recurseDirectories(d, []);
+
+
+   });
+
+   browse.files?.forEach(d => {
+      items.push({ key: d.hash, text: d.name, icon: "Document", type: BrowserType.File, size: d.length });
+   });
+
+   const columns: IColumn[] = [
+      { key: 'column1', name: 'Name', minWidth: 794, maxWidth: 794, isResizable: false, isPadded: false },
+      { key: 'column2', name: 'Size', minWidth: 128, maxWidth: 128, isResizable: false, isPadded: false },
+      { key: 'column3', name: 'View_Download', minWidth: 64, maxWidth: 64, isResizable: false, isPadded: false }
+   ];
+
+   const renderItem = (item: any, index?: number, column?: IColumn) => {
+
+      if (!column) {
+         return null;
+      }
+
+      if (column.name === "Size") {
+         if (!item.size) {
+            return null;
+         }
+         return <Stack horizontalAlign="end" verticalAlign="center" verticalFill>
+            <Text>{formatBytes(item.size, (item.size < (1024 * 1024)) ? 0 : 1)}</Text>
+         </Stack>
+
+      }
+
+      if (column.name === "View_Download") {
+
+         if (item.type !== BrowserType.File) {
+            return null;
+         }
+
+         const path = encodeURI(handler.path + "/" + item.text);
+         const server = backend.serverUrl;
+         const href = `${server}/api/v2/artifacts/${handler.artifact!.id}/file?path=${path}`;
+
+
+         return <Stack data-selection-disabled verticalAlign="center" verticalFill horizontal horizontalAlign="end" style={{ paddingTop: 0, paddingBottom: 0 }}>
+            <IconButton id="artifactview" href={`${href}&inline=true`} target="_blank" style={{ paddingTop: 1, color: "#106EBE" }} iconProps={{ iconName: "Eye", styles: { root: { fontSize: "14px" } } }} />
+            <IconButton id="artifactview" href={href} target="_blank" style={{ paddingTop: 1, color: "#106EBE" }} iconProps={{ iconName: "CloudDownload", styles: { root: { fontSize: "14px" } } }} />
+         </Stack>
+      }
+
+      if (column.name === "Name") {
+
+         const path = (item.text as string).split("/");
+
+         const pathElements = path.map((t, index) => {
+            const last = index === (path.length - 1);
+            let color = last ? (dashboard.darktheme ? "#FFFFFF" : "#605E5C") : undefined;
+            const font = last ? undefined : "Horde Open Sans Light";
+            const sep = last ? undefined : "/"
+            return <Text styles={{ root: { fontFamily: font } }} style={{ color: color }}>{t}{sep}</Text>
+         })
+
+         return <Stack verticalFill verticalAlign="center" style={{ cursor: "pointer" }} onClick={(ev) => {
+
+            if (item.type === BrowserType.Directory) {
+               const nbrowse = handler.path ? `${handler.path}/${item.text}` : item.text;
+               handler.browseTo(nbrowse);
+            }
+            if (item.type === BrowserType.NavigateUp && handler.path) {
+               const nbrowse = handler.path.split("/")
+               nbrowse.pop();
+               handler.browseTo(nbrowse.join("/"));
+            }
+
+         }}>
+            {item.type !== BrowserType.NavigateUp && <Stack data-selection-disabled={item.type !== BrowserType.File} horizontal tokens={{ childrenGap: 8 }}>
+               <Stack>
+                  <FontIcon style={{ paddingTop: 1, fontSize: 16 }} iconName={item.icon} />
+               </Stack>
+               <Stack horizontal>
+                  {pathElements}
+               </Stack>
+            </Stack>}
+            {item.type === BrowserType.NavigateUp && <Stack data-selection-disabled verticalFill horizontal verticalAlign="center" tokens={{ childrenGap: 9 }}>
+               <Stack>
+                  <FontIcon style={{ paddingTop: 1, fontSize: 15 }} iconName="ArrowUp" />
+               </Stack>
+               <Stack>
+                  <Text>..</Text>
+               </Stack>
+            </Stack>
+            }
+         </Stack>
+
+      }
+
+      return null;
+
+   }
+
+   return <Stack key={`jobdetailartifacts_${idcounter++}`} tokens={{ childrenGap: 12 }}>
+      <Stack horizontal verticalAlign="center" tokens={{ childrenGap: 18 }} style={{ paddingBottom: 12 }}>
+         <Stack>
+            <BrowseBreadCrumbs handler={handler} />
+         </Stack>
+         <Stack grow />
+         <DownloadButton handler={handler} />
+      </Stack >
+      {handler.loading && <Stack>
+         <Spinner styles={{ root: { opacity: 0, animation: "hordeFadeIn 1s ease-in-out 2s forwards" } }} size={SpinnerSize.large} />
+      </Stack>}
+      {!handler.loading && <Stack style={{ height: 492 + 160, position: "relative" }}>
+         <ScrollablePane style={{ height: 492 + 160 }}>
+            <SelectionZone selection={handler.selection!}>
+               <DetailsList
+                  styles={{ root: { overflowX: "hidden" } }}
+                  className={styles.list}
+                  isHeaderVisible={false}
+                  compact={true}
+                  items={items}
+                  columns={columns}
+                  layoutMode={DetailsListLayoutMode.fixedColumns}
+                  selectionMode={SelectionMode.multiple}
+                  selection={handler.selection}
+                  selectionPreservedOnEmptyClick={true}
+                  //onItemInvoked={this._onItemInvoked} <--- double click*/                  
+                  onRenderItemColumn={renderItem}
+               />
+            </SelectionZone>
+
+         </ScrollablePane>
+      </Stack>}
+   </Stack >
+
+   //
+})
+
+
+
+export const JobArtifactsModal: React.FC<{ jobId: string; stepId: string, artifacts?:GetArtifactResponseV2[], contextType: ArtifactContextType, onClose: () => void }> = ({ jobId, stepId, artifacts, contextType, onClose }) => {
+
+   const [handler] = useState(new ArtifactsHandler(jobId, stepId, contextType, artifacts));
+
+   useEffect(() => {
+      return () => {
+         handler?.clear();
+      };
+   }, [handler]);
+
+   return <Stack>
+      <Modal isOpen={true} isBlocking={true} topOffsetFixed={true} styles={{ main: { padding: 8, width: 1180, height: 820, hasBeenOpened: false, top: "80px", position: "absolute" } }} onDismiss={() => onClose()} className={hordeClasses.modal}>
+         <Stack className="horde-no-darktheme" styles={{ root: { paddingTop: 10, paddingRight: 12 } }}>
+            <Stack style={{ paddingLeft: 24, paddingRight: 24 }}>
+               <Stack tokens={{ childrenGap: 12 }} style={{ height: 800 }}>
+                  <Stack horizontal verticalAlign="start">
+                     <Stack style={{ paddingTop: 3 }}>
+                        <Text variant="mediumPlus" styles={{ root: { fontFamily: "Horde Open Sans SemiBold" } }}>Artifacts</Text>
+                     </Stack>
+                     <Stack grow />
+                     <Stack horizontalAlign="end">
+                        <IconButton
+                           iconProps={{ iconName: 'Cancel' }}
+                           onClick={() => { onClose() }}
+                        />
+                     </Stack>
+
+                  </Stack>
+                  <Stack styles={{ root: { paddingLeft: 4, paddingRight: 0, paddingTop: 8, paddingBottom: 4 } }}>
+                     <JobDetailArtifactsInner handler={handler} />
+                  </Stack>
+               </Stack>
+            </Stack>
+         </Stack>
+      </Modal>
+   </Stack>;
+};
\ No newline at end of file
diff --git a/Engine/Source/Programs/Horde/HordeDashboard/src/components/jobDetailsV2/JobDetailArtifactsV3.tsx b/Engine/Source/Programs/Horde/HordeDashboard/src/components/jobDetailsV2/JobDetailArtifactsV3.tsx
deleted file mode 100644
diff --git a/Engine/Source/Programs/Horde/HordeDashboard/src/components/jobDetailsV2/JobOperationsBar.tsx b/Engine/Source/Programs/Horde/HordeDashboard/src/components/jobDetailsV2/JobOperationsBar.tsx
--- a/Engine/Source/Programs/Horde/HordeDashboard/src/components/jobDetailsV2/JobOperationsBar.tsx
+++ b/Engine/Source/Programs/Horde/HordeDashboard/src/components/jobDetailsV2/JobOperationsBar.tsx
@@ -13,9 +13,9 @@ import { NewBuild } from '../NewBuild';
 import { NotificationDropdown } from '../NotificationDropdown';
 import { PauseStepModal } from '../StepPauseModal';
 import { AbortJobModal } from './AbortJobModal';
-import { JobArtifactsModal } from './JobDetailArtifactsV3';
 import { JobDetailsV2 } from './JobDetailsViewCommon';
 import { StepRetryModal, StepRetryType } from './StepRetryModal';
+import { JobArtifactsModal } from '../artifacts/ArtifactsModal';
 
 enum ParameterState {
    Hidden,
@@ -225,6 +225,7 @@ const StepArtifactsOperations: React.FC<{ jobDetails: JobDetailsV2, stepId: stri
       key: 'stepops_artifacts_trace',
       text: "Trace Artifacts",
       iconProps: { iconName: "SearchTemplate" },
+      disabled: !atypes.get("step-trace"),
       onClick: () => { setArtifactsShown("step-trace")}      
    });
 
@@ -241,7 +242,7 @@ const StepArtifactsOperations: React.FC<{ jobDetails: JobDetailsV2, stepId: stri
    ];
 
    return <Stack>
-      {!!artifactsShown && <JobArtifactsModal stepId={step.id} jobDetails={jobDetails} contextType={artifactsShown} onClose={() => { setArtifactsShown(undefined); }} />}
+      {!!artifactsShown && <JobArtifactsModal stepId={step.id} jobId={jobDetails.jobId!} contextType={artifactsShown} artifacts={stepArtifacts} onClose={() => { setArtifactsShown(undefined); }} />}
       <Stack horizontal styles={{ root: { paddingLeft: 0 } }}>
          <Stack grow />
          <Stack horizontal>
