Horde: Tweak source name for default OpenTelemetry tracer

diff --git a/Engine/Source/Programs/Horde/Horde.Server/Startup.cs b/Engine/Source/Programs/Horde/Horde.Server/Startup.cs
--- a/Engine/Source/Programs/Horde/Horde.Server/Startup.cs
+++ b/Engine/Source/Programs/Horde/Horde.Server/Startup.cs
@@ -1035,7 +1035,7 @@ namespace Horde.Server
 				.WithTracing(builder =>
 				{
 					builder
-						.AddSource(settings.ServiceName)
+						.AddSource("Horde")
 						.AddHttpClientInstrumentation()
 						.AddAspNetCoreInstrumentation()
 						.AddGrpcClientInstrumentation()
