Horde: Improve OpenTelemetry tracing

* Convert MongoDB tracing to use OpenTelemetry
* Create a specific tracer for MongoDB for better grouping
* Add Datadog-specific enrichment for HTTP client (via settings flag)

diff --git a/Engine/Source/Programs/Horde/Horde.Server.Tests/DatabaseIntegrationTest.cs b/Engine/Source/Programs/Horde/Horde.Server.Tests/DatabaseIntegrationTest.cs
--- a/Engine/Source/Programs/Horde/Horde.Server.Tests/DatabaseIntegrationTest.cs
+++ b/Engine/Source/Programs/Horde/Horde.Server.Tests/DatabaseIntegrationTest.cs
@@ -181,7 +181,7 @@ namespace Horde.Server.Tests
 					ss.DatabaseName = _mongoDbInstance.DatabaseName;
 					ss.DatabaseConnectionString = _mongoDbInstance.ConnectionString;
 
-					_mongoService = new MongoService(Options.Create(ss), _loggerFactory);
+					_mongoService = new MongoService(Options.Create(ss), OpenTelemetryTracers.Horde, _loggerFactory);
 				}
 			}
 			return _mongoService;
diff --git a/Engine/Source/Programs/Horde/Horde.Server/Server/MongoService.cs b/Engine/Source/Programs/Horde/Horde.Server/Server/MongoService.cs
--- a/Engine/Source/Programs/Horde/Horde.Server/Server/MongoService.cs
+++ b/Engine/Source/Programs/Horde/Horde.Server/Server/MongoService.cs
@@ -21,17 +21,16 @@ using EpicGames.Core;
 using EpicGames.Redis.Utility;
 using Horde.Server.Secrets;
 using Horde.Server.Utilities;
-using Horde.Server.Utiltiies;
 using Microsoft.Extensions.Diagnostics.HealthChecks;
 using Microsoft.Extensions.Hosting;
 using Microsoft.Extensions.Logging;
 using Microsoft.Extensions.Options;
-using Microsoft.IdentityModel.Tokens;
 using MongoDB.Bson;
 using MongoDB.Bson.Serialization;
 using MongoDB.Bson.Serialization.Attributes;
 using MongoDB.Driver;
 using MongoDB.Driver.Core.Events;
+using OpenTelemetry.Trace;
 using StackExchange.Redis;
 
 namespace Horde.Server.Server
@@ -256,6 +255,11 @@ namespace Horde.Server.Server
 		/// </summary>
 		readonly ILoggerFactory _loggerFactory;
 
+		/// <summary>
+		/// Tracer
+		/// </summary>
+		readonly Tracer _tracer;
+
 		/// <summary>
 		/// Default port for MongoDB connections
 		/// </summary>
@@ -268,10 +272,12 @@ namespace Horde.Server.Server
 		/// Constructor
 		/// </summary>
 		/// <param name="settingsSnapshot">The settings instance</param>
+		/// <param name="tracer">Tracer</param>
 		/// <param name="loggerFactory">Instance of the logger for this service</param>
-		public MongoService(IOptions<ServerSettings> settingsSnapshot, ILoggerFactory loggerFactory)
+		public MongoService(IOptions<ServerSettings> settingsSnapshot, Tracer tracer, ILoggerFactory loggerFactory)
 		{
 			Settings = settingsSnapshot.Value;
+			_tracer = tracer;
 			_logger = loggerFactory.CreateLogger<MongoService>();
 			_loggerFactory = loggerFactory;
 
@@ -627,7 +633,7 @@ namespace Horde.Server.Server
 		/// <returns></returns>
 		public IMongoCollection<T> GetCollection<T>(string name)
 		{
-			return new MongoTracingCollection<T>(Database.GetCollection<T>(name));
+			return new MongoTracingCollection<T>(Database.GetCollection<T>(name), _tracer);
 		}
 
 		/// <summary>
diff --git a/Engine/Source/Programs/Horde/Horde.Server/ServerSettings.cs b/Engine/Source/Programs/Horde/Horde.Server/ServerSettings.cs
--- a/Engine/Source/Programs/Horde/Horde.Server/ServerSettings.cs
+++ b/Engine/Source/Programs/Horde/Horde.Server/ServerSettings.cs
@@ -342,6 +342,11 @@ namespace Horde.Server
 		/// Service version
 		/// </summary>
 		public string? ServiceVersion { get; set; }
+
+		/// <summary>
+		/// Whether to enrich and format telemetry to fit presentation in Datadog
+		/// </summary>
+		public bool EnableDatadogCompatibility { get; set; } = false;
 		
 		/// <summary>
 		/// Extra attributes to set
diff --git a/Engine/Source/Programs/Horde/Horde.Server/Startup.cs b/Engine/Source/Programs/Horde/Horde.Server/Startup.cs
--- a/Engine/Source/Programs/Horde/Horde.Server/Startup.cs
+++ b/Engine/Source/Programs/Horde/Horde.Server/Startup.cs
@@ -6,6 +6,7 @@ using System.Diagnostics;
 using System.Globalization;
 using System.IO;
 using System.Linq;
+using System.Net.Http;
 using System.Reflection;
 using System.Runtime.InteropServices;
 using System.Text.Json;
@@ -1023,20 +1024,35 @@ namespace Horde.Server
 		private static void ConfigureOpenTelemetry(IServiceCollection services, OpenTelemetrySettings settings)
 		{
 			// Always configure OpenTelemetry.Trace.Tracer class as it's used in the codebase even when OpenTelemetry is not configured
-			services.AddSingleton(TracerProvider.Default.GetTracer(settings.ServiceName));
+			services.AddSingleton(OpenTelemetryTracers.Horde);
 			
 			if (!settings.Enabled)
 			{
 				return;
 			}
+
+			void DatadogHttpRequestEnricher(Activity activity, HttpRequestMessage message)
+			{
+				activity.AddTag("service.name", settings.ServiceName + "-http-client");
+				activity.AddTag("operation.name", "http.request");
+				string url = $"{message.Method} {message.Headers.Host}{message.RequestUri?.LocalPath}";  
+				activity.DisplayName = url;
+				activity.AddTag("resource.name", url);
+			}
 			
 			List<KeyValuePair<string,object>> attributes = settings.Attributes.Select(x => new KeyValuePair<string, object>(x.Key, x.Value)).ToList();
 			services.AddOpenTelemetry()
 				.WithTracing(builder =>
 				{
 					builder
-						.AddSource("Horde")
-						.AddHttpClientInstrumentation()
+						.AddSource(OpenTelemetryTracers.SourceNames)
+						.AddHttpClientInstrumentation(options =>
+						{
+							if (settings.EnableDatadogCompatibility)
+							{
+								options.EnrichWithHttpRequestMessage = DatadogHttpRequestEnricher;
+							}
+						})
 						.AddAspNetCoreInstrumentation()
 						.AddGrpcClientInstrumentation()
 						//.AddRedisInstrumentation()
diff --git a/Engine/Source/Programs/Horde/Horde.Server/Utilities/MongoTracingCollection.cs b/Engine/Source/Programs/Horde/Horde.Server/Utilities/MongoTracingCollection.cs
--- a/Engine/Source/Programs/Horde/Horde.Server/Utilities/MongoTracingCollection.cs
+++ b/Engine/Source/Programs/Horde/Horde.Server/Utilities/MongoTracingCollection.cs
@@ -6,10 +6,9 @@ using System.Threading.Tasks;
 using MongoDB.Bson;
 using MongoDB.Bson.Serialization;
 using MongoDB.Driver;
-using OpenTracing;
-using OpenTracing.Util;
+using OpenTelemetry.Trace;
 
-namespace Horde.Server.Utiltiies
+namespace Horde.Server.Utilities
 {
 	/// <summary>
 	/// Wrap a IMongoCollection with trace scopes
@@ -21,550 +20,455 @@ namespace Horde.Server.Utiltiies
 	public class MongoTracingCollection<T> : IMongoCollection<T>
 	{
 		private readonly IMongoCollection<T> _collection;
-		private const string Prefix = "MongoCollection.";
+		private readonly Tracer _tracer;
 
 		/// <summary>
 		/// Constructor
 		/// </summary>
 		/// <param name="collection">Collection to wrap with tracing</param>
-		public MongoTracingCollection(IMongoCollection<T> collection)
+		/// <param name="tracer">Tracer</param>
+		public MongoTracingCollection(IMongoCollection<T> collection, Tracer tracer)
 		{
 			_collection = collection;
+			_tracer = tracer;
 		}
 
 #pragma warning disable CS0618		
 		/// <inheritdoc />
-		public IAsyncCursor<TResult> Aggregate<TResult>(PipelineDefinition<T, TResult> pipeline, AggregateOptions options = null!,
-			CancellationToken  cancellationToken = new CancellationToken())
+		public IAsyncCursor<TResult> Aggregate<TResult>(PipelineDefinition<T, TResult> pipeline, AggregateOptions options = null!, CancellationToken cancellationToken = default)
 		{
 			return _collection.Aggregate(pipeline, options,  cancellationToken);
 		}
 
 		/// <inheritdoc />
-		public IAsyncCursor<TResult> Aggregate<TResult>(IClientSessionHandle session, PipelineDefinition<T, TResult> pipeline,
-			AggregateOptions options = null!, CancellationToken  cancellationToken = new CancellationToken())
+		public IAsyncCursor<TResult> Aggregate<TResult>(IClientSessionHandle session, PipelineDefinition<T, TResult> pipeline, AggregateOptions options = null!, CancellationToken cancellationToken = default)
 		{
 			return _collection.Aggregate(session, pipeline, options,  cancellationToken);
 		}
 
 		/// <inheritdoc />
-		public async Task<IAsyncCursor<TResult>> AggregateAsync<TResult>(PipelineDefinition<T, TResult> pipeline, AggregateOptions options = null!,
-			CancellationToken  cancellationToken = new CancellationToken())
+		public async Task<IAsyncCursor<TResult>> AggregateAsync<TResult>(PipelineDefinition<T, TResult> pipeline, AggregateOptions options = null!, CancellationToken  cancellationToken = default)
 		{
-			using IScope scope = GlobalTracer.Instance.BuildSpan(Prefix + "AggregateAsync").StartActive();
-			scope.Span.SetTag("CollectionName", _collection.CollectionNamespace.CollectionName);
+			using TelemetrySpan span = _tracer.StartMongoDbSpan(nameof(AggregateAsync), _collection);
 			return await _collection.AggregateAsync(pipeline, options,  cancellationToken);
 		}
 
 		/// <inheritdoc />
 		public async Task<IAsyncCursor<TResult>> AggregateAsync<TResult>(IClientSessionHandle session, PipelineDefinition<T, TResult> pipeline, AggregateOptions options = null!,
-			CancellationToken  cancellationToken = new CancellationToken())
+			CancellationToken cancellationToken = default)
 		{
-			using IScope scope = GlobalTracer.Instance.BuildSpan(Prefix + "AggregateAsync").StartActive();
-			scope.Span.SetTag("CollectionName", _collection.CollectionNamespace.CollectionName);
+			using TelemetrySpan span = _tracer.StartMongoDbSpan(nameof(AggregateAsync), _collection);
 			return await _collection.AggregateAsync(session, pipeline, options, cancellationToken);
 		}
 
 		/// <inheritdoc />
 		public void AggregateToCollection<TResult>(PipelineDefinition<T, TResult> pipeline, AggregateOptions? options = null, CancellationToken cancellationToken = default)
 		{
-			using IScope scope = GlobalTracer.Instance.BuildSpan(Prefix + "AggregateToCollection").StartActive();
-			scope.Span.SetTag("CollectionName", _collection.CollectionNamespace.CollectionName);
 			_collection.AggregateToCollection(pipeline, options, cancellationToken);
 		}
 
 		/// <inheritdoc />
 		public void AggregateToCollection<TResult>(IClientSessionHandle session, PipelineDefinition<T, TResult> pipeline, AggregateOptions? options = null, CancellationToken cancellationToken = default)
 		{
-			using IScope scope = GlobalTracer.Instance.BuildSpan(Prefix + "AggregateToCollection").StartActive();
-			scope.Span.SetTag("CollectionName", _collection.CollectionNamespace.CollectionName);
 			_collection.AggregateToCollection(session, pipeline, options, cancellationToken);
 		}
 
 		/// <inheritdoc />
 		public async Task AggregateToCollectionAsync<TResult>(PipelineDefinition<T, TResult> pipeline, AggregateOptions? options = null, CancellationToken cancellationToken = default)
 		{
-			using IScope scope = GlobalTracer.Instance.BuildSpan(Prefix + "AggregateToCollectionAsync").StartActive();
-			scope.Span.SetTag("CollectionName", _collection.CollectionNamespace.CollectionName);
+			using TelemetrySpan span = _tracer.StartMongoDbSpan(nameof(AggregateToCollectionAsync), _collection);
 			await _collection.AggregateToCollectionAsync(pipeline, options, cancellationToken);
 		}
 
 		/// <inheritdoc />
 		public async Task AggregateToCollectionAsync<TResult>(IClientSessionHandle session, PipelineDefinition<T, TResult> pipeline, AggregateOptions? options = null, CancellationToken cancellationToken = default)
 		{
-			using IScope scope = GlobalTracer.Instance.BuildSpan(Prefix + "AggregateToCollectionAsync").StartActive();
-			scope.Span.SetTag("CollectionName", _collection.CollectionNamespace.CollectionName);
+			using TelemetrySpan span = _tracer.StartMongoDbSpan(nameof(AggregateToCollectionAsync), _collection);
 			await _collection.AggregateToCollectionAsync(session, pipeline, options, cancellationToken);
 		}
 
 		/// <inheritdoc />
-		public BulkWriteResult<T> BulkWrite(IEnumerable<WriteModel<T>> requests, BulkWriteOptions options = null!,
-			CancellationToken  cancellationToken = new CancellationToken())
+		public BulkWriteResult<T> BulkWrite(IEnumerable<WriteModel<T>> requests, BulkWriteOptions options = null!, CancellationToken cancellationToken = default)
 		{
 			return _collection.BulkWrite(requests, options,  cancellationToken);
 		}
 
 		/// <inheritdoc />
-		public BulkWriteResult<T> BulkWrite(IClientSessionHandle session, IEnumerable<WriteModel<T>> requests, BulkWriteOptions options = null!,
-			CancellationToken  cancellationToken = new CancellationToken())
+		public BulkWriteResult<T> BulkWrite(IClientSessionHandle session, IEnumerable<WriteModel<T>> requests, BulkWriteOptions options = null!, CancellationToken cancellationToken = default)
 		{
 			return _collection.BulkWrite(session, requests, options,  cancellationToken);
 		}
 
 		/// <inheritdoc />
-		public async Task<BulkWriteResult<T>> BulkWriteAsync(IEnumerable<WriteModel<T>> requests, BulkWriteOptions options = null!,
-			CancellationToken  cancellationToken = new CancellationToken())
+		public async Task<BulkWriteResult<T>> BulkWriteAsync(IEnumerable<WriteModel<T>> requests, BulkWriteOptions options = null!, CancellationToken cancellationToken = default)
 		{
-			using IScope scope = GlobalTracer.Instance.BuildSpan(Prefix + "BulkWriteAsync").StartActive();
-			scope.Span.SetTag("CollectionName", _collection.CollectionNamespace.CollectionName);
+			using TelemetrySpan span = _tracer.StartMongoDbSpan(nameof(BulkWriteAsync), _collection);
 			return await _collection.BulkWriteAsync(requests, options,  cancellationToken);
 		}
 
 		/// <inheritdoc />
-		public async Task<BulkWriteResult<T>> BulkWriteAsync(IClientSessionHandle session, IEnumerable<WriteModel<T>> requests, BulkWriteOptions options = null!,
-			CancellationToken  cancellationToken = new CancellationToken())
+		public async Task<BulkWriteResult<T>> BulkWriteAsync(IClientSessionHandle session, IEnumerable<WriteModel<T>> requests, BulkWriteOptions options = null!, CancellationToken cancellationToken = default)
 		{
-			using IScope scope = GlobalTracer.Instance.BuildSpan(Prefix + "BulkWriteAsync").StartActive();
-			scope.Span.SetTag("CollectionName", _collection.CollectionNamespace.CollectionName);
+			using TelemetrySpan span = _tracer.StartMongoDbSpan(nameof(BulkWriteAsync), _collection);
 			return await _collection.BulkWriteAsync(session, requests, options,  cancellationToken);
 		}
 
 		/// <inheritdoc />
-		public long Count(FilterDefinition<T> filter, CountOptions options = null!,
-			CancellationToken  cancellationToken = new CancellationToken())
+		public long Count(FilterDefinition<T> filter, CountOptions options = null!, CancellationToken cancellationToken = default)
 		{
 			return _collection.Count(filter, options,  cancellationToken);
 		}
 
 		/// <inheritdoc />
-		public long Count(IClientSessionHandle session, FilterDefinition<T> filter, CountOptions options = null!,
-			CancellationToken  cancellationToken = new CancellationToken())
+		public long Count(IClientSessionHandle session, FilterDefinition<T> filter, CountOptions options = null!, CancellationToken cancellationToken = default)
 		{
 			return _collection.Count(session, filter, options,  cancellationToken);
 		}
 
 		/// <inheritdoc />
-		public async Task<long> CountAsync(FilterDefinition<T> filter, CountOptions options = null!,
-			CancellationToken  cancellationToken = new CancellationToken())
+		public async Task<long> CountAsync(FilterDefinition<T> filter, CountOptions options = null!, CancellationToken cancellationToken = default)
 		{
-			using IScope scope = GlobalTracer.Instance.BuildSpan(Prefix + "CountAsync").StartActive();
-			scope.Span.SetTag("CollectionName", _collection.CollectionNamespace.CollectionName);
+			using TelemetrySpan span = _tracer.StartMongoDbSpan(nameof(CountAsync), _collection);
 			return await _collection.CountAsync(filter, options,  cancellationToken);
 		}
 
 		/// <inheritdoc />
-		public async Task<long> CountAsync(IClientSessionHandle session, FilterDefinition<T> filter, CountOptions options = null!,
-			CancellationToken  cancellationToken = new CancellationToken())
+		public async Task<long> CountAsync(IClientSessionHandle session, FilterDefinition<T> filter, CountOptions options = null!, CancellationToken cancellationToken = default)
 		{
-			using IScope scope = GlobalTracer.Instance.BuildSpan(Prefix + "CountAsync").StartActive();
-			scope.Span.SetTag("CollectionName", _collection.CollectionNamespace.CollectionName);
+			using TelemetrySpan span = _tracer.StartMongoDbSpan(nameof(CountAsync), _collection);
 			return await _collection.CountAsync(session, filter, options,  cancellationToken);
 		}
 
 		/// <inheritdoc />
-		public long CountDocuments(FilterDefinition<T> filter, CountOptions options = null!,
-			CancellationToken  cancellationToken = new CancellationToken())
+		public long CountDocuments(FilterDefinition<T> filter, CountOptions options = null!, CancellationToken cancellationToken = default)
 		{
 			return _collection.CountDocuments(filter, options,  cancellationToken);
 		}
 
 		/// <inheritdoc />
-		public long CountDocuments(IClientSessionHandle session, FilterDefinition<T> filter, CountOptions options = null!,
-			CancellationToken  cancellationToken = new CancellationToken())
+		public long CountDocuments(IClientSessionHandle session, FilterDefinition<T> filter, CountOptions options = null!, CancellationToken cancellationToken = default)
 		{
-			using IScope scope = GlobalTracer.Instance.BuildSpan(Prefix + "CountDocumentsAsync").StartActive();
-			scope.Span.SetTag("CollectionName", _collection.CollectionNamespace.CollectionName);
 			return _collection.CountDocuments(session, filter, options,  cancellationToken);
 		}
 
 		/// <inheritdoc />
-		public async Task<long> CountDocumentsAsync(FilterDefinition<T> filter, CountOptions options = null!,
-			CancellationToken  cancellationToken = new CancellationToken())
+		public async Task<long> CountDocumentsAsync(FilterDefinition<T> filter, CountOptions options = null!, CancellationToken cancellationToken = default)
 		{
-			using IScope scope = GlobalTracer.Instance.BuildSpan(Prefix + "CountDocumentsAsync").StartActive();
-			scope.Span.SetTag("CollectionName", _collection.CollectionNamespace.CollectionName);
+			using TelemetrySpan span = _tracer.StartMongoDbSpan(nameof(CountDocumentsAsync), _collection);
 			return await _collection.CountDocumentsAsync(filter, options,  cancellationToken);
 		}
 
 		/// <inheritdoc />
-		public async Task<long> CountDocumentsAsync(IClientSessionHandle session, FilterDefinition<T> filter, CountOptions options = null!,
-			CancellationToken  cancellationToken = new CancellationToken())
+		public async Task<long> CountDocumentsAsync(IClientSessionHandle session, FilterDefinition<T> filter, CountOptions options = null!, CancellationToken cancellationToken = default)
 		{
-			using IScope scope = GlobalTracer.Instance.BuildSpan(Prefix + "CountDocumentsAsync").StartActive();
-			scope.Span.SetTag("CollectionName", _collection.CollectionNamespace.CollectionName);
+			using TelemetrySpan span = _tracer.StartMongoDbSpan(nameof(CountDocumentsAsync), _collection);
 			return await _collection.CountDocumentsAsync(session, filter, options,  cancellationToken);
 		}
 
 		/// <inheritdoc />
-		public DeleteResult DeleteMany(FilterDefinition<T> filter, CancellationToken  cancellationToken = new CancellationToken())
+		public DeleteResult DeleteMany(FilterDefinition<T> filter, CancellationToken cancellationToken = default)
 		{
 			return _collection.DeleteMany(filter,  cancellationToken);
 		}
 
 		/// <inheritdoc />
-		public DeleteResult DeleteMany(FilterDefinition<T> filter, DeleteOptions options,
-			CancellationToken  cancellationToken = new CancellationToken())
+		public DeleteResult DeleteMany(FilterDefinition<T> filter, DeleteOptions options, CancellationToken cancellationToken = default)
 		{
 			return _collection.DeleteMany(filter, options,  cancellationToken);
 		}
 
 		/// <inheritdoc />
-		public DeleteResult DeleteMany(IClientSessionHandle session, FilterDefinition<T> filter, DeleteOptions options = null!,
-			CancellationToken  cancellationToken = new CancellationToken())
+		public DeleteResult DeleteMany(IClientSessionHandle session, FilterDefinition<T> filter, DeleteOptions options = null!, CancellationToken cancellationToken = default)
 		{
 			return _collection.DeleteMany(session, filter, options,  cancellationToken);
 		}
 
 		/// <inheritdoc />
-		public async Task<DeleteResult> DeleteManyAsync(FilterDefinition<T> filter, CancellationToken  cancellationToken = new CancellationToken())
+		public async Task<DeleteResult> DeleteManyAsync(FilterDefinition<T> filter, CancellationToken cancellationToken = default)
 		{
-			using IScope scope = GlobalTracer.Instance.BuildSpan(Prefix + "DeleteManyAsync").StartActive();
-			scope.Span.SetTag("CollectionName", _collection.CollectionNamespace.CollectionName);
+			using TelemetrySpan span = _tracer.StartMongoDbSpan(nameof(DeleteManyAsync), _collection);
 			return await _collection.DeleteManyAsync(filter,  cancellationToken);
 		}
 
 		/// <inheritdoc />
-		public async Task<DeleteResult> DeleteManyAsync(FilterDefinition<T> filter, DeleteOptions options,
-			CancellationToken  cancellationToken = new CancellationToken())
+		public async Task<DeleteResult> DeleteManyAsync(FilterDefinition<T> filter, DeleteOptions options, CancellationToken cancellationToken = default)
 		{
-			using IScope scope = GlobalTracer.Instance.BuildSpan(Prefix + "DeleteManyAsync").StartActive();
-			scope.Span.SetTag("CollectionName", _collection.CollectionNamespace.CollectionName);
+			using TelemetrySpan span = _tracer.StartMongoDbSpan(nameof(DeleteManyAsync), _collection);
 			return await _collection.DeleteManyAsync(filter, options,  cancellationToken);
 		}
 
 		/// <inheritdoc />
-		public async Task<DeleteResult> DeleteManyAsync(IClientSessionHandle session, FilterDefinition<T> filter, DeleteOptions options = null!,
-			CancellationToken  cancellationToken = new CancellationToken())
+		public async Task<DeleteResult> DeleteManyAsync(IClientSessionHandle session, FilterDefinition<T> filter, DeleteOptions options = null!, CancellationToken cancellationToken = default)
 		{
-			using IScope scope = GlobalTracer.Instance.BuildSpan(Prefix + "DeleteManyAsync").StartActive();
-			scope.Span.SetTag("CollectionName", _collection.CollectionNamespace.CollectionName);
+			using TelemetrySpan span = _tracer.StartMongoDbSpan(nameof(DeleteManyAsync), _collection);
 			return await _collection.DeleteManyAsync(session, filter, options,  cancellationToken);
 		}
 
 		/// <inheritdoc />
-		public DeleteResult DeleteOne(FilterDefinition<T> filter, CancellationToken  cancellationToken = new CancellationToken())
+		public DeleteResult DeleteOne(FilterDefinition<T> filter, CancellationToken cancellationToken = default)
 		{
 			return _collection.DeleteOne(filter,  cancellationToken);
 		}
 
 		/// <inheritdoc />
-		public DeleteResult DeleteOne(FilterDefinition<T> filter, DeleteOptions options,
-			CancellationToken  cancellationToken = new CancellationToken())
+		public DeleteResult DeleteOne(FilterDefinition<T> filter, DeleteOptions options, CancellationToken cancellationToken = default)
 		{
 			return _collection.DeleteOne(filter, options,  cancellationToken);
 		}
 
 		/// <inheritdoc />
-		public DeleteResult DeleteOne(IClientSessionHandle session, FilterDefinition<T> filter, DeleteOptions options = null!,
-			CancellationToken  cancellationToken = new CancellationToken())
+		public DeleteResult DeleteOne(IClientSessionHandle session, FilterDefinition<T> filter, DeleteOptions options = null!, CancellationToken cancellationToken = default)
 		{
 			return _collection.DeleteOne(session, filter, options,  cancellationToken);
 		}
 
 		/// <inheritdoc />
-		public async Task<DeleteResult> DeleteOneAsync(FilterDefinition<T> filter, CancellationToken  cancellationToken = new CancellationToken())
+		public async Task<DeleteResult> DeleteOneAsync(FilterDefinition<T> filter, CancellationToken cancellationToken = default)
 		{
-			using IScope scope = GlobalTracer.Instance.BuildSpan(Prefix + "DeleteOneAsync").StartActive();
-			scope.Span.SetTag("CollectionName", _collection.CollectionNamespace.CollectionName);
+			using TelemetrySpan span = _tracer.StartMongoDbSpan(nameof(DeleteOneAsync), _collection);
 			return await _collection.DeleteOneAsync(filter,  cancellationToken);
 		}
 
 		/// <inheritdoc />
-		public async Task<DeleteResult> DeleteOneAsync(FilterDefinition<T> filter, DeleteOptions options,
-			CancellationToken  cancellationToken = new CancellationToken())
+		public async Task<DeleteResult> DeleteOneAsync(FilterDefinition<T> filter, DeleteOptions options, CancellationToken cancellationToken = default)
 		{
-			using IScope scope = GlobalTracer.Instance.BuildSpan(Prefix + "DeleteOneAsync").StartActive();
-			scope.Span.SetTag("CollectionName", _collection.CollectionNamespace.CollectionName);
+			using TelemetrySpan span = _tracer.StartMongoDbSpan(nameof(DeleteOneAsync), _collection);
 			return await _collection.DeleteOneAsync(filter, options,  cancellationToken);
 		}
 
 		/// <inheritdoc />
-		public Task<DeleteResult> DeleteOneAsync(IClientSessionHandle session, FilterDefinition<T> filter, DeleteOptions options = null!,
-			CancellationToken  cancellationToken = new CancellationToken())
+		public Task<DeleteResult> DeleteOneAsync(IClientSessionHandle session, FilterDefinition<T> filter, DeleteOptions options = null!, CancellationToken cancellationToken = default)
 		{
-			using IScope scope = GlobalTracer.Instance.BuildSpan(Prefix + "DeleteOneAsync").StartActive();
-			scope.Span.SetTag("CollectionName", _collection.CollectionNamespace.CollectionName);
+			using TelemetrySpan span = _tracer.StartMongoDbSpan(nameof(DeleteOneAsync), _collection);
 			return _collection.DeleteOneAsync(session, filter, options,  cancellationToken);
 		}
 
 		/// <inheritdoc />
-		public IAsyncCursor<TField> Distinct<TField>(FieldDefinition<T, TField> field, FilterDefinition<T> filter, DistinctOptions options = null!,
-			CancellationToken  cancellationToken = new CancellationToken())
+		public IAsyncCursor<TField> Distinct<TField>(FieldDefinition<T, TField> field, FilterDefinition<T> filter, DistinctOptions options = null!, CancellationToken cancellationToken = default)
 		{
 			return _collection.Distinct(field, filter, options,  cancellationToken);
 		}
 
 		/// <inheritdoc />
-		public IAsyncCursor<TField> Distinct<TField>(IClientSessionHandle session, FieldDefinition<T, TField> field, FilterDefinition<T> filter,
-			DistinctOptions options = null!, CancellationToken  cancellationToken = new CancellationToken())
+		public IAsyncCursor<TField> Distinct<TField>(IClientSessionHandle session, FieldDefinition<T, TField> field, FilterDefinition<T> filter, DistinctOptions options = null!, CancellationToken cancellationToken = default)
 		{
 			return _collection.Distinct(session, field, filter, options,  cancellationToken);
 		}
 
 		/// <inheritdoc />
-		public async Task<IAsyncCursor<TField>> DistinctAsync<TField>(FieldDefinition<T, TField> field, FilterDefinition<T> filter, DistinctOptions options = null!,
-			CancellationToken  cancellationToken = new CancellationToken())
+		public async Task<IAsyncCursor<TField>> DistinctAsync<TField>(FieldDefinition<T, TField> field, FilterDefinition<T> filter, DistinctOptions options = null!, CancellationToken cancellationToken = default)
 		{
-			using IScope scope = GlobalTracer.Instance.BuildSpan(Prefix + "DistinctAsync").StartActive();
-			scope.Span.SetTag("CollectionName", _collection.CollectionNamespace.CollectionName);
+			using TelemetrySpan span = _tracer.StartMongoDbSpan(nameof(DistinctAsync), _collection);
 			return await _collection.DistinctAsync(field, filter, options,  cancellationToken);
 		}
 
 		/// <inheritdoc />
-		public async Task<IAsyncCursor<TField>> DistinctAsync<TField>(IClientSessionHandle session, FieldDefinition<T, TField> field, FilterDefinition<T> filter,
-			DistinctOptions options = null!, CancellationToken  cancellationToken = new CancellationToken())
+		public async Task<IAsyncCursor<TField>> DistinctAsync<TField>(IClientSessionHandle session, FieldDefinition<T, TField> field, FilterDefinition<T> filter, DistinctOptions options = null!, CancellationToken cancellationToken = default)
 		{
-			using IScope scope = GlobalTracer.Instance.BuildSpan(Prefix + "DistinctAsync").StartActive();
-			scope.Span.SetTag("CollectionName", _collection.CollectionNamespace.CollectionName);
+			using TelemetrySpan span = _tracer.StartMongoDbSpan(nameof(DistinctAsync), _collection);
 			return await _collection.DistinctAsync(session, field, filter, options,  cancellationToken);
 		}
 
 		/// <inheritdoc />
-		public long EstimatedDocumentCount(EstimatedDocumentCountOptions options = null!,
-			CancellationToken  cancellationToken = new CancellationToken())
+		public long EstimatedDocumentCount(EstimatedDocumentCountOptions options = null!, CancellationToken cancellationToken = default)
 		{
 			return _collection.EstimatedDocumentCount(options,  cancellationToken);
 		}
 
 		/// <inheritdoc />
-		public async Task<long> EstimatedDocumentCountAsync(EstimatedDocumentCountOptions options = null!,
-			CancellationToken  cancellationToken = new CancellationToken())
+		public async Task<long> EstimatedDocumentCountAsync(EstimatedDocumentCountOptions options = null!, CancellationToken cancellationToken = default)
 		{
-			using IScope scope = GlobalTracer.Instance.BuildSpan(Prefix + "EstimatedDocumentCountAsync").StartActive();
-			scope.Span.SetTag("CollectionName", _collection.CollectionNamespace.CollectionName);
+			using TelemetrySpan span = _tracer.StartMongoDbSpan(nameof(EstimatedDocumentCountAsync), _collection);
 			return await _collection.EstimatedDocumentCountAsync(options,  cancellationToken);
 		}
 
 		/// <inheritdoc />
-		public IAsyncCursor<TProjection> FindSync<TProjection>(FilterDefinition<T> filter, FindOptions<T, TProjection> options = null!,
-			CancellationToken  cancellationToken = new CancellationToken())
+		public IAsyncCursor<TProjection> FindSync<TProjection>(FilterDefinition<T> filter, FindOptions<T, TProjection> options = null!, CancellationToken cancellationToken = default)
 		{
 			return _collection.FindSync(filter, options,  cancellationToken);
 		}
 
 		/// <inheritdoc />
-		public IAsyncCursor<TProjection> FindSync<TProjection>(IClientSessionHandle session, FilterDefinition<T> filter, FindOptions<T, TProjection> options = null!,
-			CancellationToken  cancellationToken = new CancellationToken())
+		public IAsyncCursor<TProjection> FindSync<TProjection>(IClientSessionHandle session, FilterDefinition<T> filter, FindOptions<T, TProjection> options = null!, CancellationToken cancellationToken = default)
 		{
 			return _collection.FindSync(session, filter, options,  cancellationToken);
 		}
 
 		/// <inheritdoc />
-		public async Task<IAsyncCursor<TProjection>> FindAsync<TProjection>(FilterDefinition<T> filter, FindOptions<T, TProjection> options = null!,
-			CancellationToken  cancellationToken = new CancellationToken())
+		public async Task<IAsyncCursor<TProjection>> FindAsync<TProjection>(FilterDefinition<T> filter, FindOptions<T, TProjection> options = null!, CancellationToken cancellationToken = default)
 		{
-			using IScope scope = GlobalTracer.Instance.BuildSpan(Prefix + "FindAsync").StartActive();
-			scope.Span.SetTag("CollectionName", _collection.CollectionNamespace.CollectionName);
+			using TelemetrySpan span = _tracer.StartMongoDbSpan(nameof(FindAsync), _collection);
 			return await _collection.FindAsync(filter, options,  cancellationToken);
 		}
 
 		/// <inheritdoc />
-		public async Task<IAsyncCursor<TProjection>> FindAsync<TProjection>(IClientSessionHandle session, FilterDefinition<T> filter, FindOptions<T, TProjection> options = null!,
-			CancellationToken  cancellationToken = new CancellationToken())
+		public async Task<IAsyncCursor<TProjection>> FindAsync<TProjection>(IClientSessionHandle session, FilterDefinition<T> filter, FindOptions<T, TProjection> options = null!, CancellationToken cancellationToken = default)
 		{
-			using IScope scope = GlobalTracer.Instance.BuildSpan(Prefix + "FindAsync").StartActive();
-			scope.Span.SetTag("CollectionName", _collection.CollectionNamespace.CollectionName);
+			using TelemetrySpan span = _tracer.StartMongoDbSpan(nameof(FindAsync), _collection);
 			return await _collection.FindAsync(session, filter, options,  cancellationToken);
 		}
 
 		/// <inheritdoc />
-		public TProjection FindOneAndDelete<TProjection>(FilterDefinition<T> filter, FindOneAndDeleteOptions<T, TProjection> options = null!,
-			CancellationToken  cancellationToken = new CancellationToken())
+		public TProjection FindOneAndDelete<TProjection>(FilterDefinition<T> filter, FindOneAndDeleteOptions<T, TProjection> options = null!, CancellationToken cancellationToken = default)
 		{
 			return _collection.FindOneAndDelete(filter, options,  cancellationToken);
 		}
 
 		/// <inheritdoc />
-		public TProjection FindOneAndDelete<TProjection>(IClientSessionHandle session, FilterDefinition<T> filter,
-			FindOneAndDeleteOptions<T, TProjection> options = null!, CancellationToken  cancellationToken = new CancellationToken())
+		public TProjection FindOneAndDelete<TProjection>(IClientSessionHandle session, FilterDefinition<T> filter, FindOneAndDeleteOptions<T, TProjection> options = null!, CancellationToken cancellationToken = default)
 		{
 			return _collection.FindOneAndDelete(session, filter, options,  cancellationToken);
 		}
 
 		/// <inheritdoc />
-		public async Task<TProjection> FindOneAndDeleteAsync<TProjection>(FilterDefinition<T> filter, FindOneAndDeleteOptions<T, TProjection> options = null!,
-			CancellationToken  cancellationToken = new CancellationToken())
+		public async Task<TProjection> FindOneAndDeleteAsync<TProjection>(FilterDefinition<T> filter, FindOneAndDeleteOptions<T, TProjection> options = null!, CancellationToken cancellationToken = default)
 		{
-			using IScope scope = GlobalTracer.Instance.BuildSpan(Prefix + "FindOneAndDeleteAsync").StartActive();
-			scope.Span.SetTag("CollectionName", _collection.CollectionNamespace.CollectionName);
+			using TelemetrySpan span = _tracer.StartMongoDbSpan(nameof(FindOneAndDeleteAsync), _collection);
 			return await _collection.FindOneAndDeleteAsync(filter, options,  cancellationToken);
 		}
 
 		/// <inheritdoc />
-		public async Task<TProjection> FindOneAndDeleteAsync<TProjection>(IClientSessionHandle session, FilterDefinition<T> filter,
-			FindOneAndDeleteOptions<T, TProjection> options = null!, CancellationToken  cancellationToken = new CancellationToken())
+		public async Task<TProjection> FindOneAndDeleteAsync<TProjection>(IClientSessionHandle session, FilterDefinition<T> filter, FindOneAndDeleteOptions<T, TProjection> options = null!, CancellationToken cancellationToken = default)
 		{
-			using IScope scope = GlobalTracer.Instance.BuildSpan(Prefix + "FindOneAndDeleteAsync").StartActive();
-			scope.Span.SetTag("CollectionName", _collection.CollectionNamespace.CollectionName);
+			using TelemetrySpan span = _tracer.StartMongoDbSpan(nameof(FindOneAndDeleteAsync), _collection);
 			return await _collection.FindOneAndDeleteAsync(session, filter, options,  cancellationToken);
 		}
 
 		/// <inheritdoc />
-		public TProjection FindOneAndReplace<TProjection>(FilterDefinition<T> filter, T replacement,
-			FindOneAndReplaceOptions<T, TProjection> options = null!, CancellationToken  cancellationToken = new CancellationToken())
+		public TProjection FindOneAndReplace<TProjection>(FilterDefinition<T> filter, T replacement, FindOneAndReplaceOptions<T, TProjection> options = null!, CancellationToken cancellationToken = default)
 		{
 			return _collection.FindOneAndReplace(filter, replacement, options,  cancellationToken);
 		}
 
 		/// <inheritdoc />
-		public TProjection FindOneAndReplace<TProjection>(IClientSessionHandle session, FilterDefinition<T> filter, T replacement,
-			FindOneAndReplaceOptions<T, TProjection> options = null!, CancellationToken  cancellationToken = new CancellationToken())
+		public TProjection FindOneAndReplace<TProjection>(IClientSessionHandle session, FilterDefinition<T> filter, T replacement, FindOneAndReplaceOptions<T, TProjection> options = null!, CancellationToken cancellationToken = default)
 		{
 			return _collection.FindOneAndReplace(session, filter, replacement, options,  cancellationToken);
 		}
 
 		/// <inheritdoc />
-		public async Task<TProjection> FindOneAndReplaceAsync<TProjection>(FilterDefinition<T> filter, T replacement,
-			FindOneAndReplaceOptions<T, TProjection> options = null!, CancellationToken  cancellationToken = new CancellationToken())
+		public async Task<TProjection> FindOneAndReplaceAsync<TProjection>(FilterDefinition<T> filter, T replacement, FindOneAndReplaceOptions<T, TProjection> options = null!, CancellationToken cancellationToken = default)
 		{
-			using IScope scope = GlobalTracer.Instance.BuildSpan(Prefix + "FindOneAndReplaceAsync").StartActive();
-			scope.Span.SetTag("CollectionName", _collection.CollectionNamespace.CollectionName);
+			using TelemetrySpan span = _tracer.StartMongoDbSpan(nameof(FindOneAndReplaceAsync), _collection);
 			return await _collection.FindOneAndReplaceAsync(filter, replacement, options,  cancellationToken);
 		}
 
 		/// <inheritdoc />
-		public async Task<TProjection> FindOneAndReplaceAsync<TProjection>(IClientSessionHandle session, FilterDefinition<T> filter, T replacement,
-			FindOneAndReplaceOptions<T, TProjection> options = null!, CancellationToken  cancellationToken = new CancellationToken())
+		public async Task<TProjection> FindOneAndReplaceAsync<TProjection>(IClientSessionHandle session, FilterDefinition<T> filter, T replacement, FindOneAndReplaceOptions<T, TProjection> options = null!, CancellationToken cancellationToken = default)
 		{
-			using IScope scope = GlobalTracer.Instance.BuildSpan(Prefix + "FindOneAndReplaceAsync").StartActive();
-			scope.Span.SetTag("CollectionName", _collection.CollectionNamespace.CollectionName);
+			using TelemetrySpan span = _tracer.StartMongoDbSpan(nameof(FindOneAndReplaceAsync), _collection);
 			return await _collection.FindOneAndReplaceAsync(session, filter, replacement, options,  cancellationToken);
 		}
 
 		/// <inheritdoc />
 		public TProjection FindOneAndUpdate<TProjection>(FilterDefinition<T> filter, UpdateDefinition<T> update,
-			FindOneAndUpdateOptions<T, TProjection> options = null!, CancellationToken  cancellationToken = new CancellationToken())
+			FindOneAndUpdateOptions<T, TProjection> options = null!, CancellationToken cancellationToken = default)
 		{
 			return _collection.FindOneAndUpdate(filter, update, options,  cancellationToken);
 		}
 
 		/// <inheritdoc />
-		public TProjection FindOneAndUpdate<TProjection>(IClientSessionHandle session, FilterDefinition<T> filter,
-			UpdateDefinition<T> update, FindOneAndUpdateOptions<T, TProjection> options = null!,
-			CancellationToken  cancellationToken = new CancellationToken())
+		public TProjection FindOneAndUpdate<TProjection>(IClientSessionHandle session, FilterDefinition<T> filter, UpdateDefinition<T> update, FindOneAndUpdateOptions<T, TProjection> options = null!, CancellationToken cancellationToken = default)
 		{
 			return _collection.FindOneAndUpdate(session, filter, update, options,  cancellationToken);
 		}
 
 		/// <inheritdoc />
-		public async Task<TProjection> FindOneAndUpdateAsync<TProjection>(FilterDefinition<T> filter, UpdateDefinition<T> update,
-			FindOneAndUpdateOptions<T, TProjection> options = null!, CancellationToken  cancellationToken = new CancellationToken())
+		public async Task<TProjection> FindOneAndUpdateAsync<TProjection>(FilterDefinition<T> filter, UpdateDefinition<T> update, FindOneAndUpdateOptions<T, TProjection> options = null!, CancellationToken cancellationToken = default)
 		{
-			using IScope scope = GlobalTracer.Instance.BuildSpan(Prefix + "FindOneAndUpdateAsync").StartActive();
-			scope.Span.SetTag("CollectionName", _collection.CollectionNamespace.CollectionName);
+			using TelemetrySpan span = _tracer.StartMongoDbSpan(nameof(FindOneAndUpdateAsync), _collection);
 			return await _collection.FindOneAndUpdateAsync(filter, update, options,  cancellationToken);
 		}
 
 		/// <inheritdoc />
-		public async Task<TProjection> FindOneAndUpdateAsync<TProjection>(IClientSessionHandle session, FilterDefinition<T> filter, UpdateDefinition<T> update,
-			FindOneAndUpdateOptions<T, TProjection> options = null!, CancellationToken  cancellationToken = new CancellationToken())
+		public async Task<TProjection> FindOneAndUpdateAsync<TProjection>(IClientSessionHandle session, FilterDefinition<T> filter, UpdateDefinition<T> update, FindOneAndUpdateOptions<T, TProjection> options = null!, CancellationToken cancellationToken = default)
 		{
-			using IScope scope = GlobalTracer.Instance.BuildSpan(Prefix + "FindOneAndUpdateAsync").StartActive();
-			scope.Span.SetTag("CollectionName", _collection.CollectionNamespace.CollectionName);
+			using TelemetrySpan span = _tracer.StartMongoDbSpan(nameof(FindOneAndUpdateAsync), _collection);
 			return await _collection.FindOneAndUpdateAsync(session, filter, update, options,  cancellationToken);
 		}
 
 		/// <inheritdoc />
-		public void InsertOne( T document, InsertOneOptions options = null!,
-			CancellationToken  cancellationToken = new CancellationToken())
+		public void InsertOne(T document, InsertOneOptions options = null!, CancellationToken cancellationToken = default)
 		{
 			_collection.InsertOne(document, options,  cancellationToken);
 		}
 
 		/// <inheritdoc />
-		public void InsertOne(IClientSessionHandle session, T document, InsertOneOptions options = null!,
-			CancellationToken  cancellationToken = new CancellationToken())
+		public void InsertOne(IClientSessionHandle session, T document, InsertOneOptions options = null!, CancellationToken cancellationToken = default)
 		{
 			_collection.InsertOne(session, document, options,  cancellationToken);
 		}
 
 		/// <inheritdoc />
-		public async Task InsertOneAsync( T document, CancellationToken cancellationToken)
+		public async Task InsertOneAsync(T document, CancellationToken cancellationToken)
 		{
-			using IScope scope = GlobalTracer.Instance.BuildSpan(Prefix + "InsertOneAsync").StartActive();
-			scope.Span.SetTag("CollectionName", _collection.CollectionNamespace.CollectionName);
+			using TelemetrySpan span = _tracer.StartMongoDbSpan(nameof(InsertOneAsync), _collection);
 			await _collection.InsertOneAsync(document, cancellationToken);
 		}
 
 		/// <inheritdoc />
-		public async Task InsertOneAsync( T document, InsertOneOptions options = null!,
-			CancellationToken  cancellationToken = new CancellationToken())
+		public async Task InsertOneAsync(T document, InsertOneOptions options = null!, CancellationToken cancellationToken = default)
 		{
-			using IScope scope = GlobalTracer.Instance.BuildSpan(Prefix + "InsertOneAsync").StartActive();
-			scope.Span.SetTag("CollectionName", _collection.CollectionNamespace.CollectionName);
+			using TelemetrySpan span = _tracer.StartMongoDbSpan(nameof(InsertOneAsync), _collection);
 			await _collection.InsertOneAsync(document, options,  cancellationToken);
 		}
 
 		/// <inheritdoc />
-		public async Task InsertOneAsync(IClientSessionHandle session, T document, InsertOneOptions options = null!,
-			CancellationToken cancellationToken = new CancellationToken())
+		public async Task InsertOneAsync(IClientSessionHandle session, T document, InsertOneOptions options = null!, CancellationToken cancellationToken = default)
 		{
-			using IScope scope = GlobalTracer.Instance.BuildSpan(Prefix + "InsertOneAsync").StartActive();
-			scope.Span.SetTag("CollectionName", _collection.CollectionNamespace.CollectionName);
+			using TelemetrySpan span = _tracer.StartMongoDbSpan(nameof(InsertOneAsync), _collection);
 			await _collection.InsertOneAsync(session, document, options,  cancellationToken);
 		}
 
 		/// <inheritdoc />
-		public void InsertMany(IEnumerable<T> documents, InsertManyOptions options = null!,
-			CancellationToken  cancellationToken = new CancellationToken())
+		public void InsertMany(IEnumerable<T> documents, InsertManyOptions options = null!, CancellationToken cancellationToken = default)
 		{
 			_collection.InsertMany(documents, options,  cancellationToken);
 		}
 
 		/// <inheritdoc />
-		public void InsertMany(IClientSessionHandle session, IEnumerable<T> documents, InsertManyOptions options = null!,
-			CancellationToken  cancellationToken = new CancellationToken())
+		public void InsertMany(IClientSessionHandle session, IEnumerable<T> documents, InsertManyOptions options = null!, CancellationToken cancellationToken = default)
 		{
 			_collection.InsertMany(session, documents, options,  cancellationToken);
 		}
 
 		/// <inheritdoc />
-		public async Task InsertManyAsync(IEnumerable<T> documents, InsertManyOptions options = null!,
-			CancellationToken  cancellationToken = new CancellationToken())
+		public async Task InsertManyAsync(IEnumerable<T> documents, InsertManyOptions options = null!, CancellationToken cancellationToken = default)
 		{
-			using IScope scope = GlobalTracer.Instance.BuildSpan(Prefix + "InsertManyAsync").StartActive();
-			scope.Span.SetTag("CollectionName", _collection.CollectionNamespace.CollectionName);
+			using TelemetrySpan span = _tracer.StartMongoDbSpan(nameof(InsertManyAsync), _collection);
 			await _collection.InsertManyAsync(documents, options,  cancellationToken);
 		}
 
 		/// <inheritdoc />
-		public async Task InsertManyAsync(IClientSessionHandle session, IEnumerable<T> documents, InsertManyOptions options = null!,
-			CancellationToken  cancellationToken = new CancellationToken())
+		public async Task InsertManyAsync(IClientSessionHandle session, IEnumerable<T> documents, InsertManyOptions options = null!, CancellationToken cancellationToken = default)
 		{
-			using IScope scope = GlobalTracer.Instance.BuildSpan(Prefix + "InsertManyAsync").StartActive();
-			scope.Span.SetTag("CollectionName", _collection.CollectionNamespace.CollectionName);
+			using TelemetrySpan span = _tracer.StartMongoDbSpan(nameof(InsertManyAsync), _collection);
 			await _collection.InsertManyAsync(session, documents, options,  cancellationToken);
 		}
 
 		/// <inheritdoc />
-		public IAsyncCursor<TResult> MapReduce<TResult>(BsonJavaScript map, BsonJavaScript reduce, MapReduceOptions<T, TResult> options = null!,
-			CancellationToken  cancellationToken = new CancellationToken())
+		public IAsyncCursor<TResult> MapReduce<TResult>(BsonJavaScript map, BsonJavaScript reduce, MapReduceOptions<T, TResult> options = null!, CancellationToken cancellationToken = default)
 		{
 			return _collection.MapReduce(map, reduce, options,  cancellationToken);
 		}
 
 		/// <inheritdoc />
-		public IAsyncCursor<TResult> MapReduce<TResult>(IClientSessionHandle session, BsonJavaScript map, BsonJavaScript reduce,
-			MapReduceOptions<T, TResult> options = null!, CancellationToken  cancellationToken = new CancellationToken())
+		public IAsyncCursor<TResult> MapReduce<TResult>(IClientSessionHandle session, BsonJavaScript map, BsonJavaScript reduce, MapReduceOptions<T, TResult> options = null!, CancellationToken cancellationToken = default)
 		{
 			return _collection.MapReduce(session, map, reduce, options,  cancellationToken);
 		}
 
 		/// <inheritdoc />
-		public async Task<IAsyncCursor<TResult>> MapReduceAsync<TResult>(BsonJavaScript map, BsonJavaScript reduce, MapReduceOptions<T, TResult> options = null!,
-			CancellationToken  cancellationToken = new CancellationToken())
+		public async Task<IAsyncCursor<TResult>> MapReduceAsync<TResult>(BsonJavaScript map, BsonJavaScript reduce, MapReduceOptions<T, TResult> options = null!, CancellationToken cancellationToken = default)
 		{
-			using IScope scope = GlobalTracer.Instance.BuildSpan(Prefix + "MapReduceAsync").StartActive();
-			scope.Span.SetTag("CollectionName", _collection.CollectionNamespace.CollectionName);
+			using TelemetrySpan span = _tracer.StartMongoDbSpan(nameof(MapReduceAsync), _collection);
 			return await _collection.MapReduceAsync(map, reduce, options,  cancellationToken);
 		}
 
 		/// <inheritdoc />
-		public async Task<IAsyncCursor<TResult>> MapReduceAsync<TResult>(IClientSessionHandle session, BsonJavaScript map, BsonJavaScript reduce,
-			MapReduceOptions<T, TResult> options = null!, CancellationToken  cancellationToken = new CancellationToken())
+		public async Task<IAsyncCursor<TResult>> MapReduceAsync<TResult>(IClientSessionHandle session, BsonJavaScript map, BsonJavaScript reduce, MapReduceOptions<T, TResult> options = null!, CancellationToken cancellationToken = default)
 		{
-			using IScope scope = GlobalTracer.Instance.BuildSpan(Prefix + "MapReduceAsync").StartActive();
-			scope.Span.SetTag("CollectionName", _collection.CollectionNamespace.CollectionName);
+			using TelemetrySpan span = _tracer.StartMongoDbSpan(nameof(MapReduceAsync), _collection);
 			return await _collection.MapReduceAsync(session, map, reduce, options,  cancellationToken);
 		}
 
@@ -575,162 +479,132 @@ namespace Horde.Server.Utiltiies
 		}
 
 		/// <inheritdoc />
-		public ReplaceOneResult ReplaceOne(FilterDefinition<T> filter, T replacement, ReplaceOptions options = null!,
-			CancellationToken  cancellationToken = new CancellationToken())
+		public ReplaceOneResult ReplaceOne(FilterDefinition<T> filter, T replacement, ReplaceOptions options = null!, CancellationToken cancellationToken = default)
 		{
 			return _collection.ReplaceOne(filter, replacement, options,  cancellationToken);
 		}
 
 		/// <inheritdoc />
-		public ReplaceOneResult ReplaceOne(FilterDefinition<T> filter, T replacement, UpdateOptions options,
-			CancellationToken  cancellationToken = new CancellationToken())
+		public ReplaceOneResult ReplaceOne(FilterDefinition<T> filter, T replacement, UpdateOptions options, CancellationToken cancellationToken = default)
 		{
 			return _collection.ReplaceOne(filter, replacement, options,  cancellationToken);
 		}
 
 		/// <inheritdoc />
-		public ReplaceOneResult ReplaceOne(IClientSessionHandle session, FilterDefinition<T> filter, T replacement,
-			ReplaceOptions options = null!, CancellationToken  cancellationToken = new CancellationToken())
+		public ReplaceOneResult ReplaceOne(IClientSessionHandle session, FilterDefinition<T> filter, T replacement, ReplaceOptions options = null!, CancellationToken cancellationToken = default)
 		{
 			return _collection.ReplaceOne(session, filter, replacement, options,  cancellationToken);
 		}
 
 		/// <inheritdoc />
-		public ReplaceOneResult ReplaceOne(IClientSessionHandle session, FilterDefinition<T> filter, T replacement, UpdateOptions options,
-			CancellationToken  cancellationToken = new CancellationToken())
+		public ReplaceOneResult ReplaceOne(IClientSessionHandle session, FilterDefinition<T> filter, T replacement, UpdateOptions options, CancellationToken cancellationToken = default)
 		{
 			return _collection.ReplaceOne(session, filter, replacement, options,  cancellationToken);
 		}
 
 		/// <inheritdoc />
-		public async Task<ReplaceOneResult> ReplaceOneAsync(FilterDefinition<T> filter, T replacement, ReplaceOptions options = null!,
-			CancellationToken  cancellationToken = new CancellationToken())
+		public async Task<ReplaceOneResult> ReplaceOneAsync(FilterDefinition<T> filter, T replacement, ReplaceOptions options = null!, CancellationToken cancellationToken = default)
 		{
-			using IScope scope = GlobalTracer.Instance.BuildSpan(Prefix + "ReplaceOneAsync").StartActive();
-			scope.Span.SetTag("CollectionName", _collection.CollectionNamespace.CollectionName);
+			using TelemetrySpan span = _tracer.StartMongoDbSpan(nameof(ReplaceOneAsync), _collection);
 			return await _collection.ReplaceOneAsync(filter, replacement, options,  cancellationToken);
 		}
 
 		/// <inheritdoc />
-		public async Task<ReplaceOneResult> ReplaceOneAsync(FilterDefinition<T> filter, T replacement, UpdateOptions options,
-			CancellationToken  cancellationToken = new CancellationToken())
+		public async Task<ReplaceOneResult> ReplaceOneAsync(FilterDefinition<T> filter, T replacement, UpdateOptions options, CancellationToken cancellationToken = default)
 		{
-			using IScope scope = GlobalTracer.Instance.BuildSpan(Prefix + "ReplaceOneAsync").StartActive();
-			scope.Span.SetTag("CollectionName", _collection.CollectionNamespace.CollectionName);
+			using TelemetrySpan span = _tracer.StartMongoDbSpan(nameof(ReplaceOneAsync), _collection);
 			return await _collection.ReplaceOneAsync(filter, replacement, options,  cancellationToken);
 		}
 
 		/// <inheritdoc />
-		public async Task<ReplaceOneResult> ReplaceOneAsync(IClientSessionHandle session, FilterDefinition<T> filter, T replacement,
-			ReplaceOptions options = null!, CancellationToken  cancellationToken = new CancellationToken())
+		public async Task<ReplaceOneResult> ReplaceOneAsync(IClientSessionHandle session, FilterDefinition<T> filter, T replacement, ReplaceOptions options = null!, CancellationToken cancellationToken = default)
 		{
-			using IScope scope = GlobalTracer.Instance.BuildSpan(Prefix + "ReplaceOneAsync").StartActive();
-			scope.Span.SetTag("CollectionName", _collection.CollectionNamespace.CollectionName);
+			using TelemetrySpan span = _tracer.StartMongoDbSpan(nameof(ReplaceOneAsync), _collection);
 			return await _collection.ReplaceOneAsync(session, filter, replacement, options,  cancellationToken);
 		}
 
 		/// <inheritdoc />
-		public async Task<ReplaceOneResult> ReplaceOneAsync(IClientSessionHandle session, FilterDefinition<T> filter, T replacement, UpdateOptions options,
-			CancellationToken  cancellationToken = new CancellationToken())
+		public async Task<ReplaceOneResult> ReplaceOneAsync(IClientSessionHandle session, FilterDefinition<T> filter, T replacement, UpdateOptions options, CancellationToken cancellationToken = default)
 		{
-			using IScope scope = GlobalTracer.Instance.BuildSpan(Prefix + "ReplaceOneAsync").StartActive();
-			scope.Span.SetTag("CollectionName", _collection.CollectionNamespace.CollectionName);
+			using TelemetrySpan span = _tracer.StartMongoDbSpan(nameof(ReplaceOneAsync), _collection);
 			return await _collection.ReplaceOneAsync(session, filter, replacement, options,  cancellationToken);
 		}
 
 		/// <inheritdoc />
-		public UpdateResult UpdateMany(FilterDefinition<T> filter, UpdateDefinition<T> update, UpdateOptions options = null!,
-			CancellationToken  cancellationToken = new CancellationToken())
+		public UpdateResult UpdateMany(FilterDefinition<T> filter, UpdateDefinition<T> update, UpdateOptions options = null!, CancellationToken cancellationToken = default)
 		{
 			return _collection.UpdateMany(filter, update, options,  cancellationToken);
 		}
 
 		/// <inheritdoc />
-		public UpdateResult UpdateMany(IClientSessionHandle session, FilterDefinition<T> filter, UpdateDefinition<T> update,
-			UpdateOptions options = null!, CancellationToken  cancellationToken = new CancellationToken())
+		public UpdateResult UpdateMany(IClientSessionHandle session, FilterDefinition<T> filter, UpdateDefinition<T> update, UpdateOptions options = null!, CancellationToken cancellationToken = default)
 		{
 			return _collection.UpdateMany(session, filter, update, options,  cancellationToken);
 		}
 
 		/// <inheritdoc />
-		public async Task<UpdateResult> UpdateManyAsync(FilterDefinition<T> filter, UpdateDefinition<T> update, UpdateOptions options = null!,
-			CancellationToken  cancellationToken = new CancellationToken())
+		public async Task<UpdateResult> UpdateManyAsync(FilterDefinition<T> filter, UpdateDefinition<T> update, UpdateOptions options = null!, CancellationToken cancellationToken = default)
 		{
-			using IScope scope = GlobalTracer.Instance.BuildSpan(Prefix + "UpdateManyAsync").StartActive();
-			scope.Span.SetTag("CollectionName", _collection.CollectionNamespace.CollectionName);
+			using TelemetrySpan span = _tracer.StartMongoDbSpan(nameof(UpdateManyAsync), _collection);
 			return await _collection.UpdateManyAsync(filter, update, options,  cancellationToken);
 		}
 
 		/// <inheritdoc />
-		public async Task<UpdateResult> UpdateManyAsync(IClientSessionHandle session, FilterDefinition<T> filter, UpdateDefinition<T> update,
-			UpdateOptions options = null!, CancellationToken  cancellationToken = new CancellationToken())
+		public async Task<UpdateResult> UpdateManyAsync(IClientSessionHandle session, FilterDefinition<T> filter, UpdateDefinition<T> update, UpdateOptions options = null!, CancellationToken cancellationToken = default)
 		{
-			using IScope scope = GlobalTracer.Instance.BuildSpan(Prefix + "UpdateManyAsync").StartActive();
-			scope.Span.SetTag("CollectionName", _collection.CollectionNamespace.CollectionName);
+			using TelemetrySpan span = _tracer.StartMongoDbSpan(nameof(UpdateManyAsync), _collection);
 			return await _collection.UpdateManyAsync(session, filter, update, options,  cancellationToken);
 		}
 
 		/// <inheritdoc />
-		public UpdateResult UpdateOne(FilterDefinition<T> filter, UpdateDefinition<T> update, UpdateOptions options = null!,
-			CancellationToken  cancellationToken = new CancellationToken())
+		public UpdateResult UpdateOne(FilterDefinition<T> filter, UpdateDefinition<T> update, UpdateOptions options = null!, CancellationToken cancellationToken = default)
 		{
 			return _collection.UpdateOne(filter, update, options,  cancellationToken);
 		}
 
 		/// <inheritdoc />
-		public UpdateResult UpdateOne(IClientSessionHandle session, FilterDefinition<T> filter, UpdateDefinition<T> update,
-			UpdateOptions options = null!, CancellationToken  cancellationToken = new CancellationToken())
+		public UpdateResult UpdateOne(IClientSessionHandle session, FilterDefinition<T> filter, UpdateDefinition<T> update, UpdateOptions options = null!, CancellationToken cancellationToken = default)
 		{
 			return _collection.UpdateOne(session, filter, update, options,  cancellationToken);
 		}
 
 		/// <inheritdoc />
-		public async Task<UpdateResult> UpdateOneAsync(FilterDefinition<T> filter, UpdateDefinition<T> update, UpdateOptions options = null!,
-			CancellationToken  cancellationToken = new CancellationToken())
+		public async Task<UpdateResult> UpdateOneAsync(FilterDefinition<T> filter, UpdateDefinition<T> update, UpdateOptions options = null!, CancellationToken cancellationToken = default)
 		{
-			using IScope scope = GlobalTracer.Instance.BuildSpan(Prefix + "UpdateOneAsync").StartActive();
-			scope.Span.SetTag("CollectionName", _collection.CollectionNamespace.CollectionName);
+			using TelemetrySpan span = _tracer.StartMongoDbSpan(nameof(UpdateOneAsync), _collection);
 			return await _collection.UpdateOneAsync(filter, update, options,  cancellationToken);
 		}
 
 		/// <inheritdoc />
-		public async Task<UpdateResult> UpdateOneAsync(IClientSessionHandle session, FilterDefinition<T> filter, UpdateDefinition<T> update,
-			UpdateOptions options = null!, CancellationToken  cancellationToken = new CancellationToken())
+		public async Task<UpdateResult> UpdateOneAsync(IClientSessionHandle session, FilterDefinition<T> filter, UpdateDefinition<T> update, UpdateOptions options = null!, CancellationToken cancellationToken = default)
 		{
-			using IScope scope = GlobalTracer.Instance.BuildSpan(Prefix + "UpdateOneAsync").StartActive();
-			scope.Span.SetTag("CollectionName", _collection.CollectionNamespace.CollectionName);
+			using TelemetrySpan span = _tracer.StartMongoDbSpan(nameof(UpdateOneAsync), _collection);
 			return await _collection.UpdateOneAsync(session, filter, update, options,  cancellationToken);
 		}
 
 		/// <inheritdoc />
-		public IChangeStreamCursor<TResult> Watch<TResult>(PipelineDefinition<ChangeStreamDocument<T>, TResult> pipeline, ChangeStreamOptions options = null!,
-			CancellationToken  cancellationToken = new CancellationToken())
+		public IChangeStreamCursor<TResult> Watch<TResult>(PipelineDefinition<ChangeStreamDocument<T>, TResult> pipeline, ChangeStreamOptions options = null!, CancellationToken cancellationToken = default)
 		{
 			return _collection.Watch(pipeline, options,  cancellationToken);
 		}
 
 		/// <inheritdoc />
-		public IChangeStreamCursor<TResult> Watch<TResult>(IClientSessionHandle session, PipelineDefinition<ChangeStreamDocument<T>, TResult> pipeline,
-			ChangeStreamOptions options = null!, CancellationToken  cancellationToken = new CancellationToken())
+		public IChangeStreamCursor<TResult> Watch<TResult>(IClientSessionHandle session, PipelineDefinition<ChangeStreamDocument<T>, TResult> pipeline, ChangeStreamOptions options = null!, CancellationToken cancellationToken = default)
 		{
 			return _collection.Watch(session, pipeline, options,  cancellationToken);
 		}
 
 		/// <inheritdoc />
-		public async Task<IChangeStreamCursor<TResult>> WatchAsync<TResult>(PipelineDefinition<ChangeStreamDocument<T>, TResult> pipeline, ChangeStreamOptions options = null!,
-			CancellationToken  cancellationToken = new CancellationToken())
+		public async Task<IChangeStreamCursor<TResult>> WatchAsync<TResult>(PipelineDefinition<ChangeStreamDocument<T>, TResult> pipeline, ChangeStreamOptions options = null!, CancellationToken cancellationToken = default)
 		{
-			using IScope scope = GlobalTracer.Instance.BuildSpan(Prefix + "WatchAsync").StartActive();
-			scope.Span.SetTag("CollectionName", _collection.CollectionNamespace.CollectionName);
+			using TelemetrySpan span = _tracer.StartMongoDbSpan(nameof(WatchAsync), _collection);
 			return await _collection.WatchAsync(pipeline, options,  cancellationToken);
 		}
 
 		/// <inheritdoc />
-		public async Task<IChangeStreamCursor<TResult>> WatchAsync<TResult>(IClientSessionHandle session, PipelineDefinition<ChangeStreamDocument<T>, TResult> pipeline, ChangeStreamOptions options = null!,
-			CancellationToken  cancellationToken = new CancellationToken())
+		public async Task<IChangeStreamCursor<TResult>> WatchAsync<TResult>(IClientSessionHandle session, PipelineDefinition<ChangeStreamDocument<T>, TResult> pipeline, ChangeStreamOptions options = null!, CancellationToken cancellationToken = default)
 		{
-			using IScope scope = GlobalTracer.Instance.BuildSpan(Prefix + "WatchAsync").StartActive();
-			scope.Span.SetTag("CollectionName", _collection.CollectionNamespace.CollectionName);
+			using TelemetrySpan span = _tracer.StartMongoDbSpan(nameof(WatchAsync), _collection);
 			return await _collection.WatchAsync(session, pipeline, options,  cancellationToken);
 		}
 
diff --git a/Engine/Source/Programs/Horde/Horde.Server/Utilities/TracingExtensions.cs b/Engine/Source/Programs/Horde/Horde.Server/Utilities/TracingExtensions.cs
--- a/Engine/Source/Programs/Horde/Horde.Server/Utilities/TracingExtensions.cs
+++ b/Engine/Source/Programs/Horde/Horde.Server/Utilities/TracingExtensions.cs
@@ -3,9 +3,11 @@
 using System;
 using System.Collections.Generic;
 using System.Linq;
+using System.Runtime.CompilerServices;
 using EpicGames.Core;
 using Horde.Server.Jobs;
 using Horde.Server.Streams;
+using MongoDB.Driver;
 using OpenTelemetry.Trace;
 using OpenTracing;
 
@@ -54,6 +56,38 @@ namespace Horde.Server.Utilities
 			return span;
 		}
 	}
+
+	/// <summary>
+	/// Static initialization of all available OpenTelemetry tracers
+	/// </summary>
+	public static class OpenTelemetryTracers
+	{
+		/// <summary>
+		/// Name of default Horde tracer (aka activity source)
+		/// </summary>
+		public const string HordeName = "Horde";
+		
+		/// <summary>
+		/// Name of MongoDB tracer (aka activity source)
+		/// </summary>
+		public const string MongoDbName = "MongoDB";
+
+		/// <summary>
+		/// List of all source names configured in this class.
+		/// They are needed at startup when initializing OpenTelemetry
+		/// </summary>
+		public static string[] SourceNames => new[] { HordeName, MongoDbName };
+		
+		/// <summary>
+		/// Default tracer used inside Horde
+		/// </summary>
+		public static readonly Tracer Horde = TracerProvider.Default.GetTracer(HordeName);
+		
+		/// <summary>
+		/// Tracer specific to MongoDB
+		/// </summary>
+		public static readonly Tracer MongoDb = TracerProvider.Default.GetTracer(MongoDbName);
+	}
 	
 	/// <summary>
 	/// Extensions to handle Horde specific data types in the OpenTelemetry library
@@ -106,5 +140,30 @@ namespace Horde.Server.Utilities
 		
 		/// <inheritdoc cref="ISpan.SetTag(System.String, System.String)"/>
 		public static TelemetrySpan SetAttribute(this TelemetrySpan span, string key, TemplateId[]? values) => span.SetAttribute(key, values != null ? String.Join(',', values.Select(x => x.Id.ToString())) : null);
+
+		/// <summary>
+		/// Start a MongoDB-based tracing span
+		/// </summary>
+		/// <param name="tracer">Current tracer being extended</param>
+		/// <param name="spanName">Name of the span</param>
+		/// <param name="collection">An optional MongoDB collection, the name will be used as an attribute</param>
+		/// <returns>A new telemetry span</returns>
+		[MethodImpl(MethodImplOptions.AggressiveInlining)]
+		public static TelemetrySpan StartMongoDbSpan<T>(this Tracer tracer, string spanName, IMongoCollection<T>? collection = null)
+		{
+			string name = "mongodb." + spanName;
+			TelemetrySpan span = OpenTelemetryTracers.MongoDb
+				.StartActiveSpan(name, parentContext: Tracer.CurrentSpan.Context)
+				.SetAttribute("type", "db")
+				.SetAttribute("operation.name", name)
+				.SetAttribute("service.name", OpenTelemetryTracers.MongoDbName);
+
+			if (collection != null)
+			{
+				span.SetAttribute("collection", collection.CollectionNamespace.CollectionName);
+			}
+			
+			return span;
+		}
 	}
 }
\ No newline at end of file
