Horde: More artifacts and log view work
#skipci
#rnx

diff --git a/Engine/Source/Programs/Horde/HordeDashboard/src/backend/JobDetails.ts b/Engine/Source/Programs/Horde/HordeDashboard/src/backend/JobDetails.ts
--- a/Engine/Source/Programs/Horde/HordeDashboard/src/backend/JobDetails.ts
+++ b/Engine/Source/Programs/Horde/HordeDashboard/src/backend/JobDetails.ts
@@ -20,7 +20,7 @@ export type JobLabel = GetLabelResponse & {
 const defaultUpdateMS = 5000;
 
 /**
- * @deprecated This class evolved over time to be a fragile grab bag of methods related to jobs.  
+ * @deprecated This class evolved over time to be a fragile grab bag of methods related to jobs.
  */
 export class JobDetails {
 
@@ -751,7 +751,7 @@ export class JobDetails {
             }
 
         });
-        
+
         this.jobdata?.batches?.forEach(b => {
             b.steps.forEach(s => s.reports?.forEach(r => {
                 if (s.id === this.stepId) {
@@ -806,9 +806,26 @@ export class JobDetails {
 
             // job timing is very expensive, don't delay load for it until we have job data
             if (this.jobdata) {
-                doTiming = !this.updateTimingInfo;
-                this.updateTimingInfo++;
-                this.updateTimingInfo %= 3;
+
+                let skip = false;
+                if (this.isLogView && this.logId) {
+
+                    if (!this.getLogActive(this.logId)) {
+                        // skippinng timing info on finished log, it is only used for projected finish time
+                        skip = true;                        
+                    }
+
+                    if (this.updateTimingInfo) {
+                        //skippinng timing info already loaded
+                        skip = true;
+                    }
+                }
+
+                if (!skip) {
+                    doTiming = !this.updateTimingInfo;
+                    this.updateTimingInfo++;
+                    this.updateTimingInfo %= 3;
+                }
             }
 
             if (doTiming) {
@@ -1172,7 +1189,7 @@ export const getStepSummaryMarkdown = (jobDetails: JobDetails, stepId: string):
             aborted += ` by ${jobDetails.jobdata?.abortedByUserInfo.name}.`;
         } else {
             aborted = "The step was canceled";
-            
+
             if (step.error === JobStepError.TimedOut) {
                 aborted = "The step was canceled due to reaching the maximum run time limit";
             }
diff --git a/Engine/Source/Programs/Horde/HordeDashboard/src/components/JobDetailTrace.tsx b/Engine/Source/Programs/Horde/HordeDashboard/src/components/JobDetailTrace.tsx
deleted file mode 100644
diff --git a/Engine/Source/Programs/Horde/HordeDashboard/src/components/LogSource.tsx b/Engine/Source/Programs/Horde/HordeDashboard/src/components/LogSource.tsx
--- a/Engine/Source/Programs/Horde/HordeDashboard/src/components/LogSource.tsx
+++ b/Engine/Source/Programs/Horde/HordeDashboard/src/components/LogSource.tsx
@@ -196,8 +196,8 @@ export abstract class LogSource {
 
       if (this.active) {
          this.pollID = setTimeout(() => { this.poll(); }, this.pollMS);
-      } 
-      
+      }
+
       if (this.polling) {
          return;
       }
@@ -250,7 +250,7 @@ export abstract class LogSource {
    }
 
    @observable
-   logItemsUpdated = 0;      
+   logItemsUpdated = 0;
 
    get logItems(): LogItem[] {
       // subscribe
@@ -306,10 +306,19 @@ export class JobLogSource extends LogSource {
 
             const details = this.jobDetails;
 
-            const done = () => {
+
+            const done = async () => {
+
                if (details.fatalError) {
                   this.setFatalError(details.fatalError);
                }
+
+               if (this.logData?.id) {
+                  this.setActive(details.getLogActive(this.logData.id));
+               }
+
+               await this.updateArtifacts();
+
                this.detailsUpdated();
                this.refreshJobData();
                this.initComplete();
@@ -391,6 +400,25 @@ export class JobLogSource extends LogSource {
 
    }
 
+   private async updateArtifacts() {
+      const details = this.jobDetails;
+
+      if (!details.jobdata?.useArtifactsV2 || this.artifactsV2 !== undefined || !this.logData?.id || details.getLogActive(this.logData.id)) {
+         return;
+      }
+
+      const step = details.getSteps().find(s => s.logId === this.logId);
+      if (step) {
+         const key = `job:${details.id!}/step:${step.id}`;
+         try {
+            const v = await backend.getJobArtifactsV2(undefined, [key]);
+            this.artifactsV2 = v.artifacts;
+         } catch (err) {
+            console.error(err);
+         }
+      }
+   }
+
    private async detailsUpdated() {
 
       const details = this.jobDetails;
@@ -405,18 +433,7 @@ export class JobLogSource extends LogSource {
       const active = details.getLogActive(this.logData!.id);
 
       if (this.active !== active) {
-         
-         const step = details.getSteps().find(s => s.logId === this.logId);
-         if (step) {
-            const key = `job:${details.id!}/step:${step.id}`;
-            try {
-               const v = await backend.getJobArtifactsV2(undefined, [key]);
-               this.artifactsV2 = v.artifacts;
-            } catch (err) {
-               console.error(err);
-            }                   
-         }
-
+         await this.updateArtifacts();
          this.setActive(active);
       }
    }
diff --git a/Engine/Source/Programs/Horde/HordeDashboard/src/components/LogView.tsx b/Engine/Source/Programs/Horde/HordeDashboard/src/components/LogView.tsx
--- a/Engine/Source/Programs/Horde/HordeDashboard/src/components/LogView.tsx
+++ b/Engine/Source/Programs/Horde/HordeDashboard/src/components/LogView.tsx
@@ -206,7 +206,7 @@ const StepHistoryModal: React.FC<{ jobDetails: JobDetails, stepId: string | unde
       })
 
       return <Modal isOpen={true} styles={{ main: { padding: 8, width: 1084, height: '624px', backgroundColor: '#FFFFFF' } }} className={hordeClasses.modal} onDismiss={() => { onClose() }}>
-         <Stack style={{paddingTop: 24}} horizontalAlign='center' tokens={{childrenGap: 18}}>
+         <Stack style={{ paddingTop: 24 }} horizontalAlign='center' tokens={{ childrenGap: 18 }}>
             <Stack>
                <Text variant='mediumPlus'>Loading Step History</Text>
             </Stack>
@@ -430,6 +430,19 @@ const StepArtifactsModal: React.FC<{ jobDetails: JobDetails, stepId: string | un
 
 });
 
+const LogProgressIndicator: React.FC<{ logSource: LogSource }> = observer(({ logSource }) => {
+
+   // subscribe
+   const details = (logSource as JobLogSource).jobDetails;
+   if (details?.updated) { }
+
+   const percentComplete = logSource.percentComplete;
+
+   return <Stack>
+      <ProgressIndicator percentComplete={percentComplete} barHeight={2} styles={{ root: { paddingLeft: 12, width: 364 } }} />
+   </Stack>
+});
+
 
 export const LogList: React.FC<{ logId: string }> = observer(({ logId }) => {
 
@@ -820,9 +833,7 @@ export const LogList: React.FC<{ logId: string }> = observer(({ logId }) => {
 
       }
 
-   }
-
-   const percentComplete = logSource.percentComplete;
+   }   
 
    // we need job details here, though need to make a better accessor
    const fixme = (logSource as any).jobDetails as JobDetails | undefined;
@@ -1061,7 +1072,7 @@ export const LogList: React.FC<{ logId: string }> = observer(({ logId }) => {
 
                      </Stack>
                      <Stack horizontal >
-                        {logSource.active && <ProgressIndicator percentComplete={percentComplete} barHeight={2} styles={{ root: { paddingLeft: 12, width: 364 } }} />}
+                        {logSource.active && <LogProgressIndicator logSource={logSource} />}
                      </Stack>
                   </Stack>
 
