Horde: Fix log progress indicator width, again
#rnx
#skipci

diff --git a/Engine/Source/Programs/Horde/HordeDashboard/src/components/LogView.tsx b/Engine/Source/Programs/Horde/HordeDashboard/src/components/LogView.tsx
--- a/Engine/Source/Programs/Horde/HordeDashboard/src/components/LogView.tsx
+++ b/Engine/Source/Programs/Horde/HordeDashboard/src/components/LogView.tsx
@@ -439,7 +439,7 @@ const LogProgressIndicator: React.FC<{ logSource: LogSource }> = observer(({ log
    const percentComplete = logSource.percentComplete;
 
    return <Stack>
-      <ProgressIndicator percentComplete={percentComplete} barHeight={2} styles={{ root: { paddingLeft: 12, width: 364 } }} />
+      <ProgressIndicator percentComplete={percentComplete} barHeight={2} styles={{ root: { paddingLeft: 12, width: 300 } }} />
    </Stack>
 });
 
@@ -1071,7 +1071,7 @@ export const LogList: React.FC<{ logId: string }> = observer(({ logId }) => {
                         {!menuProps.items.find(i => !i.disabled) && <Stack style={{ paddingLeft: 18 }}><DefaultButton disabled={true} className={handler.style.warningButtonDisabled} text="View" style={{ color: "rgb(97, 110, 133)", padding: 15 }} /> </Stack>}
 
                      </Stack>
-                     <Stack horizontal >
+                     <Stack>
                         {logSource.active && <LogProgressIndicator logSource={logSource} />}
                      </Stack>
                   </Stack>
