Horde: Enrich ASP.NET Core traces when Datadog compatibility is enabled

diff --git a/Engine/Source/Programs/Horde/Horde.Server/Startup.cs b/Engine/Source/Programs/Horde/Horde.Server/Startup.cs
--- a/Engine/Source/Programs/Horde/Horde.Server/Startup.cs
+++ b/Engine/Source/Programs/Horde/Horde.Server/Startup.cs
@@ -1040,6 +1040,15 @@ namespace Horde.Server
 				activity.AddTag("resource.name", url);
 			}
 			
+			void DatadogAspNetRequestEnricher(Activity activity, HttpRequest request)
+			{
+				activity.AddTag("service.name", settings.ServiceName);
+				activity.AddTag("operation.name", "http.request");
+				string url = $"{request.Method} {request.Headers.Host}{request.Path}";  
+				activity.DisplayName = url;
+				activity.AddTag("resource.name", url);
+			}
+			
 			List<KeyValuePair<string,object>> attributes = settings.Attributes.Select(x => new KeyValuePair<string, object>(x.Key, x.Value)).ToList();
 			services.AddOpenTelemetry()
 				.WithTracing(builder =>
@@ -1053,7 +1062,13 @@ namespace Horde.Server
 								options.EnrichWithHttpRequestMessage = DatadogHttpRequestEnricher;
 							}
 						})
-						.AddAspNetCoreInstrumentation()
+						.AddAspNetCoreInstrumentation(options =>
+						{
+							if (settings.EnableDatadogCompatibility)
+							{
+								options.EnrichWithHttpRequest = DatadogAspNetRequestEnricher;
+							}
+						})
 						.AddGrpcClientInstrumentation()
 						//.AddRedisInstrumentation()
 						.SetResourceBuilder(ResourceBuilder.CreateDefault()
