Horde: Do not trace incoming health or status based requests

They are frequent and traces of them provide no value.

diff --git a/Engine/Source/Programs/Horde/Horde.Server/Startup.cs b/Engine/Source/Programs/Horde/Horde.Server/Startup.cs
--- a/Engine/Source/Programs/Horde/Horde.Server/Startup.cs
+++ b/Engine/Source/Programs/Horde/Horde.Server/Startup.cs
@@ -1048,6 +1048,16 @@ namespace Horde.Server
 				activity.DisplayName = url;
 				activity.AddTag("resource.name", url);
 			}
+
+			bool FilterHttpRequests(HttpContext context)
+			{
+				if (context.Request.Path.Value != null && context.Request.Path.Value.Contains("health/", StringComparison.Ordinal))
+				{
+					return false;
+				}
+				
+				return true;
+			}
 			
 			List<KeyValuePair<string,object>> attributes = settings.Attributes.Select(x => new KeyValuePair<string, object>(x.Key, x.Value)).ToList();
 			services.AddOpenTelemetry()
@@ -1064,6 +1074,8 @@ namespace Horde.Server
 						})
 						.AddAspNetCoreInstrumentation(options =>
 						{
+							options.Filter = FilterHttpRequests;
+							
 							if (settings.EnableDatadogCompatibility)
 							{
 								options.EnrichWithHttpRequest = DatadogAspNetRequestEnricher;
