Horde: Prevent exceptions in telemetry service from terminating the host.

diff --git a/Engine/Source/Programs/Horde/Horde.Agent/Services/TelemetryService.cs b/Engine/Source/Programs/Horde/Horde.Agent/Services/TelemetryService.cs
--- a/Engine/Source/Programs/Horde/Horde.Agent/Services/TelemetryService.cs
+++ b/Engine/Source/Programs/Horde/Horde.Agent/Services/TelemetryService.cs
@@ -14,6 +14,7 @@ using HordeCommon.Rpc.Messages.Telemetry;
 using Microsoft.Extensions.Logging;
 using Microsoft.Extensions.Options;
 using Microsoft.Management.Infrastructure;
+using Horde.Agent.Utility;
 
 namespace Horde.Agent.Services;
 
@@ -249,6 +250,25 @@ class TelemetryService : BackgroundService
 
 	/// <inheritdoc />
 	protected override async Task ExecuteAsync(CancellationToken stoppingToken)
+	{
+		while (!stoppingToken.IsCancellationRequested)
+		{
+			try
+			{
+				await ExecuteInternalAsync(stoppingToken);
+			}
+			catch (OperationCanceledException) when (stoppingToken.IsCancellationRequested)
+			{
+				break;
+			}
+			catch (Exception ex)
+			{
+				_logger.LogWarning(ex, "Exception in TelemetryService: {Message}", ex.Message);
+			}
+		}
+	}
+
+	private async Task ExecuteInternalAsync(CancellationToken stoppingToken)
 	{
 		using GrpcChannel channel = _grpcService.CreateGrpcChannel();
 		HordeRpc.HordeRpcClient client = new (channel);
