Horde: Add a delay before restarting agent's TelemetryService

diff --git a/Engine/Source/Programs/Horde/Horde.Agent/Services/TelemetryService.cs b/Engine/Source/Programs/Horde/Horde.Agent/Services/TelemetryService.cs
--- a/Engine/Source/Programs/Horde/Horde.Agent/Services/TelemetryService.cs
+++ b/Engine/Source/Programs/Horde/Horde.Agent/Services/TelemetryService.cs
@@ -264,6 +264,7 @@ class TelemetryService : BackgroundService
 			catch (Exception ex)
 			{
 				_logger.LogWarning(ex, "Exception in TelemetryService: {Message}", ex.Message);
+				await Task.Delay(TimeSpan.FromSeconds(20), stoppingToken);
 			}
 		}
 	}
