Horde: Default all tests/suite on when selecting automation project
#skipci
#rnx

diff --git a/Engine/Source/Programs/Horde/HordeDashboard/src/backend/AutomationTestData.ts b/Engine/Source/Programs/Horde/HordeDashboard/src/backend/AutomationTestData.ts
--- a/Engine/Source/Programs/Horde/HordeDashboard/src/backend/AutomationTestData.ts
+++ b/Engine/Source/Programs/Horde/HordeDashboard/src/backend/AutomationTestData.ts
@@ -58,25 +58,24 @@ export class TestDataHandler {
                 this.state.streams = undefined;
             }
 
-            //this.state.tests = undefined;
-            //this.state.suites = undefined;
+            if (initStreams) {
+                const streams = this.getAutomationStreams(automation);
+                this.state.streams = streams.map(s => s);
+            }
+
+            this.state.tests = undefined;
+            this.state.suites = undefined;
 
             const tests = this.streamTests;
-            this.state.tests = this.state.tests?.filter(test => !!tests.find(t => t.name === test));
-            if (!this.state.tests?.length) {
-                this.state.tests = undefined;
+            if (tests?.length) {
+                this.state.tests = tests.map(t => t.name);
             }
-
+            
             const suites = this.streamSuites;
-            this.state.suites = this.state.suites?.filter(suite => !!suites.find(s => s.name === suite));
-            if (!this.state.suites?.length) {
-                this.state.suites = undefined;
+            if (suites?.length) {
+                this.state.suites = suites.map(t => t.name);
             }
 
-            if (initStreams) {
-                const streams = this.getAutomationStreams(automation);
-                this.state.streams = streams.map(s => s);
-            }
 
         }
 
