Horde: Ensure OS is supported before starting telemetry collection on agent

diff --git a/Engine/Source/Programs/Horde/Horde.Agent/Services/TelemetryService.cs b/Engine/Source/Programs/Horde/Horde.Agent/Services/TelemetryService.cs
--- a/Engine/Source/Programs/Horde/Horde.Agent/Services/TelemetryService.cs
+++ b/Engine/Source/Programs/Horde/Horde.Agent/Services/TelemetryService.cs
@@ -255,7 +255,10 @@ class TelemetryService : BackgroundService
 		{
 			try
 			{
-				await ExecuteInternalAsync(stoppingToken);
+				if (!await ExecuteInternalAsync(stoppingToken))
+				{
+					break;
+				}
 			}
 			catch (OperationCanceledException) when (stoppingToken.IsCancellationRequested)
 			{
@@ -264,17 +267,24 @@ class TelemetryService : BackgroundService
 			catch (Exception ex)
 			{
 				_logger.LogWarning(ex, "Exception in TelemetryService: {Message}", ex.Message);
-				await Task.Delay(TimeSpan.FromSeconds(20), stoppingToken);
 			}
+			
+			// Wait a moment before attempting to restart the background work
+			await Task.Delay(TimeSpan.FromSeconds(20), stoppingToken);
 		}
 	}
 
-	private async Task ExecuteInternalAsync(CancellationToken stoppingToken)
+	private async Task<bool> ExecuteInternalAsync(CancellationToken stoppingToken)
 	{
+		if (_systemMetrics == null)
+		{
+			return false;
+		}
+		
 		using GrpcChannel channel = _grpcService.CreateGrpcChannel();
 		HordeRpc.HordeRpcClient client = new (channel);
 		
-		while (!stoppingToken.IsCancellationRequested && _systemMetrics != null)
+		while (!stoppingToken.IsCancellationRequested)
 		{
 			_logger.LogDebug("Sending telemetry events to server...");
 			
@@ -309,6 +319,8 @@ class TelemetryService : BackgroundService
 			await client.SendTelemetryEventsAsync(request, new CallOptions(cancellationToken: stoppingToken));
 			await Task.Delay(_reportInterval, stoppingToken);
 		}
+
+		return true;
 	}
 
 	AgentMetadataEvent GetAgentMetadataEvent()
