Horde: Make alert dismissable in breadcrumbs
#rnx
#skipci

diff --git a/Engine/Source/Programs/Horde/HordeDashboard/src/components/Breadcrumbs.tsx b/Engine/Source/Programs/Horde/HordeDashboard/src/components/Breadcrumbs.tsx
--- a/Engine/Source/Programs/Horde/HordeDashboard/src/components/Breadcrumbs.tsx
+++ b/Engine/Source/Programs/Horde/HordeDashboard/src/components/Breadcrumbs.tsx
@@ -2,7 +2,7 @@
 
 import { mergeStyleSets, MessageBar, MessageBarType, Separator, Spinner, SpinnerSize, Stack, Text } from '@fluentui/react';
 import { observer } from 'mobx-react-lite';
-import React from 'react';
+import React, { useState } from 'react';
 import { Link } from 'react-router-dom';
 import dashboard from '../backend/Dashboard';
 import notices from '../backend/Notices';
@@ -35,6 +35,8 @@ export const classes = mergeStyleSets({
 
 export const Breadcrumbs: React.FC<{ items: BreadcrumbItem[], title?: string, suppressHome?: boolean, spinner?: boolean }> = observer((({ items: itemsIn, title, suppressHome, spinner }) => {
 
+   const [hideAlert, setHideAlert] = useState(false);
+
    if (notices.updated) { }
 
    // hack fix to stop the entire stack item acting as a link
@@ -110,11 +112,11 @@ export const Breadcrumbs: React.FC<{ items: BreadcrumbItem[], title?: string, su
             </Stack>
          </Stack>}
 
-         {!!alert && <Stack horizontalAlign="center" disableShrink={true} style={{ position: "absolute", width: "100%", pointerEvents: "none" }} >
+         {!!alert && !hideAlert && <Stack horizontalAlign="center" disableShrink={true} style={{ position: "absolute", width: "100%"}} >
             <Stack horizontal>
                <Stack grow />
                <Stack>
-                  <MessageBar
+                  <MessageBar onDismiss={() => setHideAlert(true)}
                      messageBarType={MessageBarType.severeWarning} isMultiline={false} >
                      <Text variant={"small"} style={{ fontFamily: "Horde Open Sans Bold" }}>{alert}</Text>
                   </MessageBar>
