Horde: Make alert dismissal more sticky
#rnx
#skipci

diff --git a/Engine/Source/Programs/Horde/HordeDashboard/src/backend/Dashboard.ts b/Engine/Source/Programs/Horde/HordeDashboard/src/backend/Dashboard.ts
--- a/Engine/Source/Programs/Horde/HordeDashboard/src/backend/Dashboard.ts
+++ b/Engine/Source/Programs/Horde/HordeDashboard/src/backend/Dashboard.ts
@@ -576,6 +576,8 @@ export class Dashboard {
 
     private config?: GetDashboardConfigResponse;
 
+    alertSquelch?: string;
+
 }
 
 const dashboard = new Dashboard();
diff --git a/Engine/Source/Programs/Horde/HordeDashboard/src/components/Breadcrumbs.tsx b/Engine/Source/Programs/Horde/HordeDashboard/src/components/Breadcrumbs.tsx
--- a/Engine/Source/Programs/Horde/HordeDashboard/src/components/Breadcrumbs.tsx
+++ b/Engine/Source/Programs/Horde/HordeDashboard/src/components/Breadcrumbs.tsx
@@ -112,11 +112,11 @@ export const Breadcrumbs: React.FC<{ items: BreadcrumbItem[], title?: string, su
             </Stack>
          </Stack>}
 
-         {!!alert && !hideAlert && <Stack horizontalAlign="center" disableShrink={true} style={{ position: "absolute", width: "100%"}} >
+         {!!alert && dashboard.alertSquelch !== alert && <Stack horizontalAlign="center" disableShrink={true} style={{ position: "absolute", width: "100%" }} >
             <Stack horizontal>
                <Stack grow />
                <Stack>
-                  <MessageBar onDismiss={() => setHideAlert(true)}
+                  <MessageBar onDismiss={() => { dashboard.alertSquelch = alert; setHideAlert(true) }}
                      messageBarType={MessageBarType.severeWarning} isMultiline={false} >
                      <Text variant={"small"} style={{ fontFamily: "Horde Open Sans Bold" }}>{alert}</Text>
                   </MessageBar>
