Horde: Fix issue with mislinked test ref in failure report
#skipci
#rnx

diff --git a/Engine/Source/Programs/Horde/HordeDashboard/src/components/AutomationReport.tsx b/Engine/Source/Programs/Horde/HordeDashboard/src/components/AutomationReport.tsx
--- a/Engine/Source/Programs/Horde/HordeDashboard/src/components/AutomationReport.tsx
+++ b/Engine/Source/Programs/Horde/HordeDashboard/src/components/AutomationReport.tsx
@@ -169,7 +169,7 @@ class TestReportGenerator {
       let details = await backend.getTestDetails(errorRefIds);
 
       details.forEach((d, index) => {
-         const refId = errorRefIds[index];
+         const refId = d.id;
          const metaEvent = allMetaEvents.find(e => e.ref.id === refId);
          if (!metaEvent) {
             console.error("Unable to find meta event");
@@ -392,7 +392,7 @@ const AutomationFailureInner: React.FC<{ generator: TestReportGenerator, test: G
          const duration = `${msecToElapsed(moment.duration(event.ref.duration).asMilliseconds(), true, true)}`;
 
          if (event.jobId && event.stepId) {
-            infoItems.push(infoItem("Horde", "Job Step", `/job/${event.jobId!}/?step=${event.stepId}`));
+            infoItems.push(infoItem("Horde", "Job Step", `/job/${event.jobId!}?step=${event.stepId}`));
          }
 
          infoItems.push(infoItem("Date", getNiceTime(event.time)));
