Horde: Connect log events in Datadog with OpenTelemetry-based traces

Uses IDs from System.Diagnostics.Activity instead of OpenTracing. Also moves OpenTelemetry related config to a separate file.

diff --git a/Engine/Source/Programs/Horde/Horde.Server/Program.cs b/Engine/Source/Programs/Horde/Horde.Server/Program.cs
--- a/Engine/Source/Programs/Horde/Horde.Server/Program.cs
+++ b/Engine/Source/Programs/Horde/Horde.Server/Program.cs
@@ -9,6 +9,7 @@ using System.Runtime.InteropServices;
 using System.Threading.Tasks;
 using EpicGames.Core;
 using Horde.Server.Commands;
+using Horde.Server.Utilities;
 using Microsoft.Extensions.Configuration;
 using Microsoft.Extensions.DependencyInjection;
 using Microsoft.Extensions.Hosting;
@@ -52,15 +53,21 @@ namespace Horde.Server
 
 		public static LoggerConfiguration WithHordeConfig(this LoggerConfiguration configuration, ServerSettings settings)
 		{
-			if (settings.WithDatadog)
+			if (settings.OpenTelemetry.EnableDatadogCompatibility)
 			{
-				configuration = configuration.Enrich.With<DatadogLogEnricher>();
+				configuration = configuration.Enrich.With<OpenTelemetryDatadogLogEnricher>();
 			}
+			else if (settings.WithDatadog)
+			{
+				// OpenTracing is replaced with OpenTelemetry. To be removed.
+				configuration = configuration.Enrich.With<OpenTracingDatadogLogEnricher>();
+			}
+
 			return configuration;
 		}
 	}
 
-	class DatadogLogEnricher : ILogEventEnricher
+	class OpenTracingDatadogLogEnricher : ILogEventEnricher
 	{
 		public void Enrich(Serilog.Events.LogEvent logEvent, ILogEventPropertyFactory propertyFactory)
 		{
diff --git a/Engine/Source/Programs/Horde/Horde.Server/Startup.cs b/Engine/Source/Programs/Horde/Horde.Server/Startup.cs
--- a/Engine/Source/Programs/Horde/Horde.Server/Startup.cs
+++ b/Engine/Source/Programs/Horde/Horde.Server/Startup.cs
@@ -98,9 +98,6 @@ using Horde.Server.Notifications.Sinks;
 using StatusCode = Grpc.Core.StatusCode;
 using Horde.Server.Artifacts;
 using Horde.Server.Compute;
-using OpenTelemetry.Exporter;
-using OpenTelemetry.Resources;
-using OpenTelemetry.Trace;
 
 namespace Horde.Server
 {
@@ -347,7 +344,7 @@ namespace Horde.Server
 			BindServerSettings(settings);
 			settings.Validate();
 
-			ConfigureOpenTelemetry(services, settings.OpenTelemetry);
+			OpenTelemetryHelper.Configure(services, settings.OpenTelemetry);
 
 			if (settings.GlobalThreadPoolMinSize != null)
 			{
@@ -1021,91 +1018,6 @@ namespace Horde.Server
 			}
 		}
 
-		private static void ConfigureOpenTelemetry(IServiceCollection services, OpenTelemetrySettings settings)
-		{
-			// Always configure OpenTelemetry.Trace.Tracer class as it's used in the codebase even when OpenTelemetry is not configured
-			services.AddSingleton(OpenTelemetryTracers.Horde);
-			
-			if (!settings.Enabled)
-			{
-				return;
-			}
-
-			void DatadogHttpRequestEnricher(Activity activity, HttpRequestMessage message)
-			{
-				activity.AddTag("service.name", settings.ServiceName + "-http-client");
-				activity.AddTag("operation.name", "http.request");
-				string url = $"{message.Method} {message.Headers.Host}{message.RequestUri?.LocalPath}";  
-				activity.DisplayName = url;
-				activity.AddTag("resource.name", url);
-			}
-			
-			void DatadogAspNetRequestEnricher(Activity activity, HttpRequest request)
-			{
-				activity.AddTag("service.name", settings.ServiceName);
-				activity.AddTag("operation.name", "http.request");
-				string url = $"{request.Method} {request.Headers.Host}{request.Path}";  
-				activity.DisplayName = url;
-				activity.AddTag("resource.name", url);
-			}
-
-			bool FilterHttpRequests(HttpContext context)
-			{
-				if (context.Request.Path.Value != null && context.Request.Path.Value.Contains("health/", StringComparison.Ordinal))
-				{
-					return false;
-				}
-				
-				return true;
-			}
-			
-			List<KeyValuePair<string,object>> attributes = settings.Attributes.Select(x => new KeyValuePair<string, object>(x.Key, x.Value)).ToList();
-			services.AddOpenTelemetry()
-				.WithTracing(builder =>
-				{
-					builder
-						.AddSource(OpenTelemetryTracers.SourceNames)
-						.AddHttpClientInstrumentation(options =>
-						{
-							if (settings.EnableDatadogCompatibility)
-							{
-								options.EnrichWithHttpRequestMessage = DatadogHttpRequestEnricher;
-							}
-						})
-						.AddAspNetCoreInstrumentation(options =>
-						{
-							options.Filter = FilterHttpRequests;
-							
-							if (settings.EnableDatadogCompatibility)
-							{
-								options.EnrichWithHttpRequest = DatadogAspNetRequestEnricher;
-							}
-						})
-						.AddGrpcClientInstrumentation()
-						//.AddRedisInstrumentation()
-						.SetResourceBuilder(ResourceBuilder.CreateDefault()
-							.AddService(settings.ServiceName, serviceNamespace: settings.ServiceNamespace, serviceVersion: settings.ServiceVersion)
-							.AddAttributes(attributes)
-							.AddTelemetrySdk()
-							.AddEnvironmentVariableDetector());
-
-					if (settings.EnableConsoleExporter)
-					{
-						builder.AddConsoleExporter();
-					}
-					
-					foreach ((string name, OpenTelemetryProtocolExporterSettings exporterSettings) in settings.ProtocolExporters)
-					{
-						builder.AddOtlpExporter(name, exporter =>
-						{
-							exporter.Endpoint = exporterSettings.Endpoint;
-							exporter.Protocol = Enum.TryParse(exporterSettings.Protocol, true, out OtlpExportProtocol protocol) ? protocol : OtlpExportProtocol.Grpc;
-						});
-					}
-				});
-
-		}
-
 		private static void OnAddHealthChecks(IServiceCollection services)
 		{
 			services.AddHealthChecks()
diff --git a/Engine/Source/Programs/Horde/Horde.Server/Utilities/OpenTelemetry.cs b/Engine/Source/Programs/Horde/Horde.Server/Utilities/OpenTelemetry.cs
new file mode 100644
--- /dev/null
+++ b/Engine/Source/Programs/Horde/Horde.Server/Utilities/OpenTelemetry.cs
@@ -0,0 +1,175 @@
+// Copyright Epic Games, Inc. All Rights Reserved.
+
+using System;
+using System.Collections.Generic;
+using System.Diagnostics;
+using System.Linq;
+using System.Net.Http;
+using Microsoft.AspNetCore.Http;
+using Microsoft.Extensions.DependencyInjection;
+using Microsoft.Extensions.Logging;
+using OpenTelemetry.Exporter;
+using OpenTelemetry.Logs;
+using OpenTelemetry.Resources;
+using OpenTelemetry.Trace;
+using Serilog.Core;
+
+namespace Horde.Server.Utilities;
+	
+/// <summary>
+/// Serilog event enricher attaching trace and span ID for Datadog using current System.Diagnostics.Activity
+/// </summary>
+public class OpenTelemetryDatadogLogEnricher : ILogEventEnricher
+{
+	/// <inheritdoc />
+	public void Enrich(Serilog.Events.LogEvent logEvent, ILogEventPropertyFactory propertyFactory)
+	{
+		if (Activity.Current != null)
+		{
+			string stringTraceId = Activity.Current.TraceId.ToString();
+			string stringSpanId = Activity.Current.SpanId.ToString();
+			string ddTraceId = Convert.ToUInt64(stringTraceId.Substring(16), 16).ToString();
+			string ddSpanId = Convert.ToUInt64(stringSpanId, 16).ToString();
+				
+			logEvent.AddPropertyIfAbsent(propertyFactory.CreateProperty("dd.trace_id", ddTraceId));
+			logEvent.AddPropertyIfAbsent(propertyFactory.CreateProperty("dd.span_id", ddSpanId));
+		}
+	}
+}
+
+/// <summary>
+/// Helper for configuring OpenTelemetry
+/// </summary>
+public static class OpenTelemetryHelper
+{
+	private static ResourceBuilder? s_resourceBuilder;
+
+	/// <summary>
+	/// Configure OpenTelemetry in Horde and ASP.NET
+	/// </summary>
+	/// <param name="services">Service collection for DI</param>
+	/// <param name="settings">Current server settings</param>
+	public static void Configure(IServiceCollection services, OpenTelemetrySettings settings)
+	{
+		// Always configure OpenTelemetry.Trace.Tracer class as it's used in the codebase even when OpenTelemetry is not configured
+		services.AddSingleton(OpenTelemetryTracers.Horde);
+			
+		if (!settings.Enabled)
+		{
+			return;
+		}
+		
+		services.AddOpenTelemetry()
+			.WithTracing(builder => ConfigureTracing(builder, settings));
+	}
+	
+	private static void ConfigureTracing(TracerProviderBuilder builder, OpenTelemetrySettings settings)
+	{
+		void DatadogHttpRequestEnricher(Activity activity, HttpRequestMessage message)
+		{
+			activity.AddTag("service.name", settings.ServiceName + "-http-client");
+			activity.AddTag("operation.name", "http.request");
+			string url = $"{message.Method} {message.Headers.Host}{message.RequestUri?.LocalPath}";  
+			activity.DisplayName = url;
+			activity.AddTag("resource.name", url);
+		}
+			
+		void DatadogAspNetRequestEnricher(Activity activity, HttpRequest request)
+		{
+			activity.AddTag("service.name", settings.ServiceName);
+			activity.AddTag("operation.name", "http.request");
+			string url = $"{request.Method} {request.Headers.Host}{request.Path}";  
+			activity.DisplayName = url;
+			activity.AddTag("resource.name", url);
+		}
+
+		bool FilterHttpRequests(HttpContext context)
+		{
+			if (context.Request.Path.Value != null && context.Request.Path.Value.Contains("health/", StringComparison.OrdinalIgnoreCase))
+			{
+				return false;
+			}
+				
+			return true;
+		}
+		
+		builder
+			.AddSource(OpenTelemetryTracers.SourceNames)
+			.AddHttpClientInstrumentation(options =>
+			{
+				if (settings.EnableDatadogCompatibility)
+				{
+					options.EnrichWithHttpRequestMessage = DatadogHttpRequestEnricher;
+				}
+			})
+			.AddAspNetCoreInstrumentation(options =>
+			{
+				options.Filter = FilterHttpRequests;
+						
+				if (settings.EnableDatadogCompatibility)
+				{
+					options.EnrichWithHttpRequest = DatadogAspNetRequestEnricher;
+				}
+			})
+			.AddGrpcClientInstrumentation()
+			//.AddRedisInstrumentation()
+			.SetResourceBuilder(OpenTelemetryHelper.GetResourceBuilder(settings));
+
+		if (settings.EnableConsoleExporter)
+		{
+			builder.AddConsoleExporter();
+		}
+				
+		foreach ((string name, OpenTelemetryProtocolExporterSettings exporterSettings) in settings.ProtocolExporters)
+		{
+			builder.AddOtlpExporter(name, exporter =>
+			{
+				exporter.Endpoint = exporterSettings.Endpoint;
+				exporter.Protocol = Enum.TryParse(exporterSettings.Protocol, true, out OtlpExportProtocol protocol) ? protocol : OtlpExportProtocol.Grpc;
+			});
+		}
+	}
+	
+	/// <summary>
+	/// Configure .NET logging with OpenTelemetry
+	/// </summary>
+	/// <param name="builder">Logging builder to modify</param>
+	/// <param name="settings">Current server settings</param>
+	public static void ConfigureLogging(ILoggingBuilder builder, OpenTelemetrySettings settings)
+	{
+		if (!settings.Enabled)
+		{
+			return;
+		}
+		
+		builder.AddOpenTelemetry(options =>
+		{
+			options.IncludeScopes = true;
+			options.IncludeFormattedMessage = true;
+			options.ParseStateValues = true;
+			options.SetResourceBuilder(GetResourceBuilder(settings));
+			
+			if (settings.EnableConsoleExporter)
+			{
+				options.AddConsoleExporter();
+			}
+		});
+	}
+
+	private static ResourceBuilder GetResourceBuilder(OpenTelemetrySettings settings)
+	{
+		if (s_resourceBuilder != null)
+		{
+			return s_resourceBuilder;
+		}
+		
+		List<KeyValuePair<string, object>> attributes = settings.Attributes.Select(x => new KeyValuePair<string, object>(x.Key, x.Value)).ToList();
+		s_resourceBuilder = ResourceBuilder.CreateDefault()
+			.AddService(settings.ServiceName, serviceNamespace: settings.ServiceNamespace, serviceVersion: settings.ServiceVersion)
+			.AddAttributes(attributes)
+			.AddTelemetrySdk()
+			.AddEnvironmentVariableDetector();
+
+		return s_resourceBuilder;
+	}
+}
\ No newline at end of file
