Horde: Enrich gRPC client with Datadog trace IDs from OpenTelemetry

diff --git a/Engine/Source/Programs/Horde/Horde.Server/Utilities/OpenTelemetry.cs b/Engine/Source/Programs/Horde/Horde.Server/Utilities/OpenTelemetry.cs
--- a/Engine/Source/Programs/Horde/Horde.Server/Utilities/OpenTelemetry.cs
+++ b/Engine/Source/Programs/Horde/Horde.Server/Utilities/OpenTelemetry.cs
@@ -67,20 +67,20 @@ public static class OpenTelemetryHelper
 	{
 		void DatadogHttpRequestEnricher(Activity activity, HttpRequestMessage message)
 		{
-			activity.AddTag("service.name", settings.ServiceName + "-http-client");
-			activity.AddTag("operation.name", "http.request");
+			activity.SetTag("service.name", settings.ServiceName + "-http-client");
+			activity.SetTag("operation.name", "http.request");
 			string url = $"{message.Method} {message.Headers.Host}{message.RequestUri?.LocalPath}";  
 			activity.DisplayName = url;
-			activity.AddTag("resource.name", url);
+			activity.SetTag("resource.name", url);
 		}
 			
 		void DatadogAspNetRequestEnricher(Activity activity, HttpRequest request)
 		{
-			activity.AddTag("service.name", settings.ServiceName);
-			activity.AddTag("operation.name", "http.request");
+			activity.SetTag("service.name", settings.ServiceName);
+			activity.SetTag("operation.name", "http.request");
 			string url = $"{request.Method} {request.Headers.Host}{request.Path}";  
 			activity.DisplayName = url;
-			activity.AddTag("resource.name", url);
+			activity.SetTag("resource.name", url);
 		}
 
 		bool FilterHttpRequests(HttpContext context)
@@ -111,9 +111,15 @@ public static class OpenTelemetryHelper
 					options.EnrichWithHttpRequest = DatadogAspNetRequestEnricher;
 				}
 			})
-			.AddGrpcClientInstrumentation()
+			.AddGrpcClientInstrumentation(options =>
+			{
+				if (settings.EnableDatadogCompatibility)
+				{
+					options.EnrichWithHttpRequestMessage = DatadogHttpRequestEnricher;
+				}
+			})
 			//.AddRedisInstrumentation()
-			.SetResourceBuilder(OpenTelemetryHelper.GetResourceBuilder(settings));
+			.SetResourceBuilder(GetResourceBuilder(settings));
 
 		if (settings.EnableConsoleExporter)
 		{
