Horde: Catch any errors when initializing system metrics collector for agent telemetry

On certain agents/OS versions "Memory" performance counter category does not exist.

diff --git a/Engine/Source/Programs/Horde/Horde.Agent/Services/TelemetryService.cs b/Engine/Source/Programs/Horde/Horde.Agent/Services/TelemetryService.cs
--- a/Engine/Source/Programs/Horde/Horde.Agent/Services/TelemetryService.cs
+++ b/Engine/Source/Programs/Horde/Horde.Agent/Services/TelemetryService.cs
@@ -238,7 +238,14 @@ class TelemetryService : BackgroundService
 	{
 		if (RuntimeInformation.IsOSPlatform(OSPlatform.Windows))
 		{
-			_systemMetrics ??= new WindowsSystemMetrics();
+			try
+			{
+				_systemMetrics ??= new WindowsSystemMetrics();
+			}
+			catch (Exception e)
+			{
+				_logger.LogError(e, "Unable to initialize system metric collector for telemetry. Disabling. Reason: {Message}", e.Message);
+			}
 		}
 		else
 		{
