Horde: Add support for OpenTelemetry metrics

diff --git a/Engine/Source/Programs/Horde/Horde.Server.Tests/TestSetup.cs b/Engine/Source/Programs/Horde/Horde.Server.Tests/TestSetup.cs
--- a/Engine/Source/Programs/Horde/Horde.Server.Tests/TestSetup.cs
+++ b/Engine/Source/Programs/Horde/Horde.Server.Tests/TestSetup.cs
@@ -2,6 +2,7 @@
 
 using System;
 using System.Collections.Generic;
+using System.Diagnostics.Metrics;
 using System.Reflection;
 using System.Security.Claims;
 using System.Text.Json;
@@ -189,6 +190,7 @@ namespace Horde.Server.Tests
 			services.AddSingleton<IAuditLog<AgentId>>(sp => sp.GetRequiredService<IAuditLogFactory<AgentId>>().Create("Agents.Log", "AgentId"));
 			services.AddSingleton<ITelemetrySink, NullTelemetrySink>();
 			services.AddSingleton<OpenTelemetry.Trace.Tracer>(sp => TracerProvider.Default.GetTracer("TestTracer"));
+			services.AddSingleton(sp => new Meter("TestMeter"));
 
 			services.AddSingleton<ConfigService>();
 			services.AddSingleton<IOptionsFactory<GlobalConfig>>(sp => sp.GetRequiredService<ConfigService>());
diff --git a/Engine/Source/Programs/Horde/Horde.Server/Utilities/OpenTelemetry.cs b/Engine/Source/Programs/Horde/Horde.Server/Utilities/OpenTelemetry.cs
--- a/Engine/Source/Programs/Horde/Horde.Server/Utilities/OpenTelemetry.cs
+++ b/Engine/Source/Programs/Horde/Horde.Server/Utilities/OpenTelemetry.cs
@@ -3,15 +3,23 @@
 using System;
 using System.Collections.Generic;
 using System.Diagnostics;
+using System.Diagnostics.Metrics;
 using System.Linq;
 using System.Net.Http;
+using System.Runtime.CompilerServices;
+using EpicGames.Core;
+using Horde.Server.Jobs;
+using Horde.Server.Streams;
 using Microsoft.AspNetCore.Http;
 using Microsoft.Extensions.DependencyInjection;
 using Microsoft.Extensions.Logging;
+using MongoDB.Driver;
 using OpenTelemetry.Exporter;
 using OpenTelemetry.Logs;
+using OpenTelemetry.Metrics;
 using OpenTelemetry.Resources;
 using OpenTelemetry.Trace;
+using OpenTracing;
 using Serilog.Core;
 
 namespace Horde.Server.Utilities;
@@ -51,8 +59,9 @@ public static class OpenTelemetryHelper
 	/// <param name="settings">Current server settings</param>
 	public static void Configure(IServiceCollection services, OpenTelemetrySettings settings)
 	{
-		// Always configure OpenTelemetry.Trace.Tracer class as it's used in the codebase even when OpenTelemetry is not configured
+		// Always configure tracers/meters as they are used in the codebase even when OpenTelemetry is not configured
 		services.AddSingleton(OpenTelemetryTracers.Horde);
+		services.AddSingleton(OpenTelemetryMeters.Horde);
 			
 		if (!settings.Enabled)
 		{
@@ -60,7 +69,8 @@ public static class OpenTelemetryHelper
 		}
 		
 		services.AddOpenTelemetry()
-			.WithTracing(builder => ConfigureTracing(builder, settings));
+			.WithTracing(builder => ConfigureTracing(builder, settings))
+			.WithMetrics(builder => ConfigureMetrics(builder, settings));
 	}
 	
 	private static void ConfigureTracing(TracerProviderBuilder builder, OpenTelemetrySettings settings)
@@ -135,6 +145,27 @@ public static class OpenTelemetryHelper
 			});
 		}
 	}
+
+	private static void ConfigureMetrics(MeterProviderBuilder builder, OpenTelemetrySettings settings)
+	{
+		builder.AddMeter(OpenTelemetryMeters.MeterNames);
+		builder.AddAspNetCoreInstrumentation();
+		builder.AddHttpClientInstrumentation();
+		
+		if (settings.EnableConsoleExporter)
+		{
+			builder.AddConsoleExporter();
+		}
+				
+		foreach ((string name, OpenTelemetryProtocolExporterSettings exporterSettings) in settings.ProtocolExporters)
+		{
+			builder.AddOtlpExporter(name, exporter =>
+			{
+				exporter.Endpoint = exporterSettings.Endpoint;
+				exporter.Protocol = Enum.TryParse(exporterSettings.Protocol, true, out OtlpExportProtocol protocol) ? protocol : OtlpExportProtocol.Grpc;
+			});
+		}
+	}
 	
 	/// <summary>
 	/// Configure .NET logging with OpenTelemetry
@@ -178,4 +209,136 @@ public static class OpenTelemetryHelper
 
 		return s_resourceBuilder;
 	}
+}
+
+/// <summary>
+/// Static initialization of all available OpenTelemetry tracers
+/// </summary>
+public static class OpenTelemetryTracers
+{
+	/// <summary>
+	/// Name of default Horde tracer (aka activity source)
+	/// </summary>
+	public const string HordeName = "Horde";
+	
+	/// <summary>
+	/// Name of MongoDB tracer (aka activity source)
+	/// </summary>
+	public const string MongoDbName = "MongoDB";
+
+	/// <summary>
+	/// List of all source names configured in this class.
+	/// They are needed at startup when initializing OpenTelemetry
+	/// </summary>
+	public static string[] SourceNames => new[] { HordeName, MongoDbName };
+	
+	/// <summary>
+	/// Default tracer used in Horde
+	/// </summary>
+	public static readonly Tracer Horde = TracerProvider.Default.GetTracer(HordeName);
+	
+	/// <summary>
+	/// Tracer specific to MongoDB
+	/// </summary>
+	public static readonly Tracer MongoDb = TracerProvider.Default.GetTracer(MongoDbName);
+}
+
+/// <summary>
+/// Extensions to handle Horde specific data types in the OpenTelemetry library
+/// </summary>
+public static class OpenTelemetrySpanExtensions
+{
+	/// <summary>Set a key:value tag on the span</summary>
+	/// <returns>This span instance, for chaining</returns>
+	public static TelemetrySpan SetAttribute(this TelemetrySpan span, string key, ContentHash value)
+	{
+		span.SetAttribute(key, value.ToString());
+		return span;
+	}
+	
+	/// <summary>Set a key:value tag on the span</summary>
+	/// <returns>This span instance, for chaining</returns>
+	public static TelemetrySpan SetAttribute(this TelemetrySpan span, string key, SubResourceId value)
+	{
+		span.SetAttribute(key, value.ToString());
+		return span;
+	}
+	
+	/// <summary>Set a key:value tag on the span</summary>
+	/// <returns>This span instance, for chaining</returns>
+	public static TelemetrySpan SetAttribute(this TelemetrySpan span, string key, int? value)
+	{
+		if (value != null)
+		{
+			span.SetAttribute(key, value.Value);
+		}
+		return span;
+	}
+	
+	/// <summary>Set a key:value tag on the span</summary>
+	/// <returns>This span instance, for chaining</returns>
+	public static TelemetrySpan SetAttribute(this TelemetrySpan span, string key, DateTimeOffset? value)
+	{
+		if (value != null)
+		{
+			span.SetAttribute(key, value.ToString());
+		}
+		return span;
+	}
+	
+	/// <inheritdoc cref="ISpan.SetTag(System.String, System.String)"/>
+	public static TelemetrySpan SetAttribute(this TelemetrySpan span, string key, StreamId? value) => span.SetAttribute(key, value?.ToString());
+	
+	/// <inheritdoc cref="ISpan.SetTag(System.String, System.String)"/>
+	public static TelemetrySpan SetAttribute(this TelemetrySpan span, string key, TemplateId? value) => span.SetAttribute(key, value?.ToString());
+	
+	/// <inheritdoc cref="ISpan.SetTag(System.String, System.String)"/>
+	public static TelemetrySpan SetAttribute(this TelemetrySpan span, string key, TemplateId[]? values) => span.SetAttribute(key, values != null ? String.Join(',', values.Select(x => x.Id.ToString())) : null);
+
+	/// <summary>
+	/// Start a MongoDB-based tracing span
+	/// </summary>
+	/// <param name="tracer">Current tracer being extended</param>
+	/// <param name="spanName">Name of the span</param>
+	/// <param name="collection">An optional MongoDB collection, the name will be used as an attribute</param>
+	/// <returns>A new telemetry span</returns>
+	[MethodImpl(MethodImplOptions.AggressiveInlining)]
+	public static TelemetrySpan StartMongoDbSpan<T>(this Tracer tracer, string spanName, IMongoCollection<T>? collection = null)
+	{
+		string name = "mongodb." + spanName;
+		TelemetrySpan span = OpenTelemetryTracers.MongoDb
+			.StartActiveSpan(name, parentContext: Tracer.CurrentSpan.Context)
+			.SetAttribute("type", "db")
+			.SetAttribute("operation.name", name)
+			.SetAttribute("service.name", OpenTelemetryTracers.MongoDbName);
+
+		if (collection != null)
+		{
+			span.SetAttribute("collection", collection.CollectionNamespace.CollectionName);
+		}
+		
+		return span;
+	}
+}
+
+/// <summary>
+/// Static initialization of all available OpenTelemetry meters
+/// </summary>
+public static class OpenTelemetryMeters
+{
+	/// <summary>
+	/// Name of default Horde meter
+	/// </summary>
+	public const string HordeName = "Horde";
+
+	/// <summary>
+	/// List of all source names configured in this class.
+	/// They are needed at startup when initializing OpenTelemetry
+	/// </summary>
+	public static string[] MeterNames => new[] { HordeName };
+	
+	/// <summary>
+	/// Default meter used in Horde
+	/// </summary>
+	public static readonly Meter Horde = new (HordeName);
 }
\ No newline at end of file
diff --git a/Engine/Source/Programs/Horde/Horde.Server/Utilities/TracingExtensions.cs b/Engine/Source/Programs/Horde/Horde.Server/Utilities/TracingExtensions.cs
--- a/Engine/Source/Programs/Horde/Horde.Server/Utilities/TracingExtensions.cs
+++ b/Engine/Source/Programs/Horde/Horde.Server/Utilities/TracingExtensions.cs
@@ -56,114 +56,4 @@ namespace Horde.Server.Utilities
 			return span;
 		}
 	}
-
-	/// <summary>
-	/// Static initialization of all available OpenTelemetry tracers
-	/// </summary>
-	public static class OpenTelemetryTracers
-	{
-		/// <summary>
-		/// Name of default Horde tracer (aka activity source)
-		/// </summary>
-		public const string HordeName = "Horde";
-		
-		/// <summary>
-		/// Name of MongoDB tracer (aka activity source)
-		/// </summary>
-		public const string MongoDbName = "MongoDB";
-
-		/// <summary>
-		/// List of all source names configured in this class.
-		/// They are needed at startup when initializing OpenTelemetry
-		/// </summary>
-		public static string[] SourceNames => new[] { HordeName, MongoDbName };
-		
-		/// <summary>
-		/// Default tracer used inside Horde
-		/// </summary>
-		public static readonly Tracer Horde = TracerProvider.Default.GetTracer(HordeName);
-		
-		/// <summary>
-		/// Tracer specific to MongoDB
-		/// </summary>
-		public static readonly Tracer MongoDb = TracerProvider.Default.GetTracer(MongoDbName);
-	}
-	
-	/// <summary>
-	/// Extensions to handle Horde specific data types in the OpenTelemetry library
-	/// </summary>
-	public static class OpenTelemetrySpanExtensions
-	{
-		/// <summary>Set a key:value tag on the span</summary>
-		/// <returns>This span instance, for chaining</returns>
-		public static TelemetrySpan SetAttribute(this TelemetrySpan span, string key, ContentHash value)
-		{
-			span.SetAttribute(key, value.ToString());
-			return span;
-		}
-		
-		/// <summary>Set a key:value tag on the span</summary>
-		/// <returns>This span instance, for chaining</returns>
-		public static TelemetrySpan SetAttribute(this TelemetrySpan span, string key, SubResourceId value)
-		{
-			span.SetAttribute(key, value.ToString());
-			return span;
-		}
-		
-		/// <summary>Set a key:value tag on the span</summary>
-		/// <returns>This span instance, for chaining</returns>
-		public static TelemetrySpan SetAttribute(this TelemetrySpan span, string key, int? value)
-		{
-			if (value != null)
-			{
-				span.SetAttribute(key, value.Value);
-			}
-			return span;
-		}
-		
-		/// <summary>Set a key:value tag on the span</summary>
-		/// <returns>This span instance, for chaining</returns>
-		public static TelemetrySpan SetAttribute(this TelemetrySpan span, string key, DateTimeOffset? value)
-		{
-			if (value != null)
-			{
-				span.SetAttribute(key, value.ToString());
-			}
-			return span;
-		}
-		
-		/// <inheritdoc cref="ISpan.SetTag(System.String, System.String)"/>
-		public static TelemetrySpan SetAttribute(this TelemetrySpan span, string key, StreamId? value) => span.SetAttribute(key, value?.ToString());
-		
-		/// <inheritdoc cref="ISpan.SetTag(System.String, System.String)"/>
-		public static TelemetrySpan SetAttribute(this TelemetrySpan span, string key, TemplateId? value) => span.SetAttribute(key, value?.ToString());
-		
-		/// <inheritdoc cref="ISpan.SetTag(System.String, System.String)"/>
-		public static TelemetrySpan SetAttribute(this TelemetrySpan span, string key, TemplateId[]? values) => span.SetAttribute(key, values != null ? String.Join(',', values.Select(x => x.Id.ToString())) : null);
-
-		/// <summary>
-		/// Start a MongoDB-based tracing span
-		/// </summary>
-		/// <param name="tracer">Current tracer being extended</param>
-		/// <param name="spanName">Name of the span</param>
-		/// <param name="collection">An optional MongoDB collection, the name will be used as an attribute</param>
-		/// <returns>A new telemetry span</returns>
-		[MethodImpl(MethodImplOptions.AggressiveInlining)]
-		public static TelemetrySpan StartMongoDbSpan<T>(this Tracer tracer, string spanName, IMongoCollection<T>? collection = null)
-		{
-			string name = "mongodb." + spanName;
-			TelemetrySpan span = OpenTelemetryTracers.MongoDb
-				.StartActiveSpan(name, parentContext: Tracer.CurrentSpan.Context)
-				.SetAttribute("type", "db")
-				.SetAttribute("operation.name", name)
-				.SetAttribute("service.name", OpenTelemetryTracers.MongoDbName);
-
-			if (collection != null)
-			{
-				span.SetAttribute("collection", collection.CollectionNamespace.CollectionName);
-			}
-			
-			return span;
-		}
-	}
 }
\ No newline at end of file
