Horde: Prevent config exceptions from logging massive messages that can't be parsed correctly (all public properties are serialized by default).

diff --git a/Engine/Source/Programs/Horde/Horde.Server/Configuration/ConfigContext.cs b/Engine/Source/Programs/Horde/Horde.Server/Configuration/ConfigContext.cs
--- a/Engine/Source/Programs/Horde/Horde.Server/Configuration/ConfigContext.cs
+++ b/Engine/Source/Programs/Horde/Horde.Server/Configuration/ConfigContext.cs
@@ -13,7 +13,7 @@ namespace Horde.Server.Configuration
 	/// Context for reading a tree of config files
 	/// </summary>
 	[DebuggerDisplay("{CurrentFile}")]
-	public class ConfigContext
+	class ConfigContext
 	{
 		/// <summary>
 		/// Options for serializing config files
diff --git a/Engine/Source/Programs/Horde/Horde.Server/Configuration/ConfigType.cs b/Engine/Source/Programs/Horde/Horde.Server/Configuration/ConfigType.cs
--- a/Engine/Source/Programs/Horde/Horde.Server/Configuration/ConfigType.cs
+++ b/Engine/Source/Programs/Horde/Horde.Server/Configuration/ConfigType.cs
@@ -55,10 +55,17 @@ namespace Horde.Server.Configuration
 	/// </summary>
 	public sealed class ConfigException : Exception
 	{
+		readonly ConfigContext _context;
+
+		/// <summary>
+		/// Stack of properties
+		/// </summary>
+		public IEnumerable<string> ScopeStack => _context.ScopeStack;
+
 		/// <summary>
-		/// Context for the parser
+		/// Stack of objects
 		/// </summary>
-		public ConfigContext Context { get; }
+		public IEnumerable<object> IncludeContextStack => _context.IncludeContextStack;
 
 		/// <summary>
 		/// Constructor
@@ -69,8 +76,13 @@ namespace Horde.Server.Configuration
 		internal ConfigException(ConfigContext context, string message, Exception? innerException = null)
 			: base(message, innerException)
 		{
-			Context = context;
+			_context = context;
 		}
+
+		/// <summary>
+		/// Gets the parser context when this exception was thrown. This is not exposed as a public property to avoid serializing the whole thing to Serilog.
+		/// </summary>
+		internal ConfigContext GetContext() => _context;
 	}
 
 	/// <summary>
diff --git a/Engine/Source/Programs/Horde/Horde.Server/Notifications/Sinks/SlackNotificationSink.cs b/Engine/Source/Programs/Horde/Horde.Server/Notifications/Sinks/SlackNotificationSink.cs
--- a/Engine/Source/Programs/Horde/Horde.Server/Notifications/Sinks/SlackNotificationSink.cs
+++ b/Engine/Source/Programs/Horde/Horde.Server/Notifications/Sinks/SlackNotificationSink.cs
@@ -1959,7 +1959,7 @@ namespace Horde.Server.Notifications.Sinks
 				List<string> details = new List<string>();
 				if (ex is ConfigException configEx)
 				{
-					ConfigContext context = configEx.Context;
+					ConfigContext context = configEx.GetContext();
 					if (context.IncludeStack.TryPeek(out IConfigFile? blame))
 					{
 						string line = String.Empty;
