Horde: Do not prevent posting Slack messages to channels in the dev environment.

diff --git a/Engine/Source/Programs/Horde/Horde.Server/Notifications/Sinks/SlackNotificationSink.cs b/Engine/Source/Programs/Horde/Horde.Server/Notifications/Sinks/SlackNotificationSink.cs
--- a/Engine/Source/Programs/Horde/Horde.Server/Notifications/Sinks/SlackNotificationSink.cs
+++ b/Engine/Source/Programs/Horde/Horde.Server/Notifications/Sinks/SlackNotificationSink.cs
@@ -2293,7 +2293,7 @@ namespace Horde.Server.Notifications.Sinks
 
 		bool IsRecipientAllowed(string recipient, string? description)
 		{
-			if (_allowUsers != null && !_allowUsers.Contains(recipient) && !recipient.StartsWith("#", StringComparison.Ordinal))
+			if (_allowUsers != null && !_allowUsers.Contains(recipient) && !recipient.StartsWith("#", StringComparison.Ordinal) && !recipient.StartsWith("C", StringComparison.Ordinal))
 			{
 				_logger.LogDebug("Suppressing message to {Recipient}: {Description}", recipient, description);
 				return false;
