Horde: Set environment variables specifying the Horde server URL and token for the current lease when executing a job.

diff --git a/Engine/Source/Programs/Horde/Horde.Agent/Execution/JobExecutor.cs b/Engine/Source/Programs/Horde/Horde.Agent/Execution/JobExecutor.cs
--- a/Engine/Source/Programs/Horde/Horde.Agent/Execution/JobExecutor.cs
+++ b/Engine/Source/Programs/Horde/Horde.Agent/Execution/JobExecutor.cs
@@ -231,6 +231,9 @@ namespace Horde.Agent.Execution
 			_token = options.Token;
 			_jobOptions = options.JobOptions;
 
+			_envVars["UE_HORDE_URL"] = options.Session.ServerUrl.ToString();
+			_envVars["UE_HORDE_TOKEN"] = options.Token;
+
 			_logger = logger;
 		}
 
diff --git a/Engine/Source/Programs/Horde/Horde.Server/Jobs/JobTaskSource.cs b/Engine/Source/Programs/Horde/Horde.Server/Jobs/JobTaskSource.cs
--- a/Engine/Source/Programs/Horde/Horde.Server/Jobs/JobTaskSource.cs
+++ b/Engine/Source/Programs/Horde/Horde.Server/Jobs/JobTaskSource.cs
@@ -832,6 +832,7 @@ namespace Horde.Server.Jobs
 
 			// Create a bearer token for the job executor
 			List<AclClaimConfig> claims = new List<AclClaimConfig>();
+			claims.Add(HordeClaims.AgentRoleClaim);
 			claims.Add(new AclClaimConfig(HordeClaimTypes.Lease, leaseId.ToString()));
 
 			string storagePrefix = $"{job.StreamId}/{job.Change}-{job.Id}";
