Horde: Force a compute lease timeout if no data is received for 15m.

diff --git a/Engine/Source/Programs/Horde/Horde.Agent/Leases/Handlers/ComputeHandler.cs b/Engine/Source/Programs/Horde/Horde.Agent/Leases/Handlers/ComputeHandler.cs
--- a/Engine/Source/Programs/Horde/Horde.Agent/Leases/Handlers/ComputeHandler.cs
+++ b/Engine/Source/Programs/Horde/Horde.Agent/Leases/Handlers/ComputeHandler.cs
@@ -1,6 +1,8 @@
 // Copyright Epic Games, Inc. All Rights Reserved.
 
 using System;
+using System.Buffers;
+using System.Diagnostics;
 using System.Net.Sockets;
 using System.Threading;
 using System.Threading.Tasks;
@@ -19,6 +21,40 @@ namespace Horde.Agent.Leases.Handlers
 	/// </summary>
 	class ComputeHandler : LeaseHandler<ComputeTask>
 	{
+		class TcpTransportWithTimeout : IComputeTransport
+		{
+			readonly TcpTransport _inner;
+			long _lastPingTicks;
+
+			public long Position => _inner.Position;
+
+			public TcpTransportWithTimeout(Socket socket)
+			{
+				_inner = new TcpTransport(socket);
+			}
+
+			public TimeSpan TimeSinceActivity => TimeSpan.FromTicks(Stopwatch.GetTimestamp() - Interlocked.CompareExchange(ref _lastPingTicks, 0, 0));
+
+			public ValueTask MarkCompleteAsync(CancellationToken cancellationToken) => _inner.MarkCompleteAsync(cancellationToken);
+
+			public async ValueTask<int> ReadPartialAsync(Memory<byte> buffer, CancellationToken cancellationToken)
+			{
+				int result = await ReadPartialAsync(buffer, cancellationToken);
+				if (result > 0)
+				{
+					Interlocked.Exchange(ref _lastPingTicks, Stopwatch.GetTimestamp());
+				}
+				return result;
+			}
+
+			public async ValueTask WriteAsync(ReadOnlySequence<byte> buffer, CancellationToken cancellationToken)
+			{
+				await WriteAsync(buffer, cancellationToken);
+				Interlocked.Exchange(ref _lastPingTicks, Stopwatch.GetTimestamp());
+			}
+		}
+
+
 		readonly ComputeListenerService _listenerService;
 		readonly IMemoryCache _memoryCache;
 		readonly DirectoryReference _sandboxDir;
@@ -54,12 +90,17 @@ namespace Horde.Agent.Leases.Handlers
 
 				_logger.LogInformation("Matched connection for {Nonce}", StringUtils.FormatHexString(computeTask.Nonce.Span));
 
-				await using (IComputeSocket socket = ComputeSocket.Create(new TcpTransport(tcpClient.Client), _logger))
+				TcpTransportWithTimeout transport = new TcpTransportWithTimeout(tcpClient.Client);
+				using (CancellationTokenSource cts = CancellationTokenSource.CreateLinkedTokenSource(cancellationToken))
 				{
-					ComputeWorker worker = new ComputeWorker(_sandboxDir, _memoryCache, _logger);
-					await worker.RunAsync(socket, cancellationToken);
-					await socket.CloseAsync(cancellationToken);
-					return LeaseResult.Success;
+					await using BackgroundTask timeoutTask = BackgroundTask.StartNew(ctx => TickTimeoutAsync(transport, cts, ctx));
+					await using (IComputeSocket socket = ComputeSocket.Create(transport, _logger))
+					{
+						ComputeWorker worker = new ComputeWorker(_sandboxDir, _memoryCache, _logger);
+						await worker.RunAsync(socket, cts.Token);
+						await socket.CloseAsync(cts.Token);
+						return LeaseResult.Success;
+					}
 				}
 			}
 			catch (Exception ex)
@@ -72,6 +113,19 @@ namespace Horde.Agent.Leases.Handlers
 				tcpClient?.Dispose();
 			}
 		}
+
+		static async Task TickTimeoutAsync(TcpTransportWithTimeout transport, CancellationTokenSource cts, CancellationToken cancellationToken)
+		{
+			while(!cancellationToken.IsCancellationRequested)
+			{
+				if (transport.TimeSinceActivity > TimeSpan.FromMinutes(10))
+				{
+					cts.Cancel();
+					break;
+				}
+				await Task.Delay(TimeSpan.FromSeconds(20), cancellationToken);
+			}
+		}
 	}
 }
 
diff --git a/Engine/Source/Programs/Shared/EpicGames.Horde/Compute/Transports/TcpTransport.cs b/Engine/Source/Programs/Shared/EpicGames.Horde/Compute/Transports/TcpTransport.cs
--- a/Engine/Source/Programs/Shared/EpicGames.Horde/Compute/Transports/TcpTransport.cs
+++ b/Engine/Source/Programs/Shared/EpicGames.Horde/Compute/Transports/TcpTransport.cs
@@ -12,7 +12,7 @@ namespace EpicGames.Horde.Compute.Transports
 	/// <summary>
 	/// Implementation of <see cref="IComputeTransport"/> for communicating over a socket
 	/// </summary>
-	public class TcpTransport : IComputeTransport
+	public sealed class TcpTransport : IComputeTransport
 	{
 		readonly Socket _socket;
 
