Horde: Fix agent tests.

diff --git a/Engine/Source/Programs/Horde/Horde.Agent.Tests/Execution/WorkspaceExecutorTest.cs b/Engine/Source/Programs/Horde/Horde.Agent.Tests/Execution/WorkspaceExecutorTest.cs
--- a/Engine/Source/Programs/Horde/Horde.Agent.Tests/Execution/WorkspaceExecutorTest.cs
+++ b/Engine/Source/Programs/Horde/Horde.Agent.Tests/Execution/WorkspaceExecutorTest.cs
@@ -101,7 +101,6 @@ public sealed class WorkspaceExecutorTest : IDisposable
 		Assert.AreEqual("1", envVars["uebp_CL"]);
 		Assert.AreEqual("0", envVars["uebp_CodeCL"]);
 		Assert.AreEqual(autoSdkSettings.DirectoryPath.FullName, envVars["UE_SDKS_ROOT"]);
-		Assert.AreEqual(7, envVars.Count);
 		
 		await executor.FinalizeAsync(_logger, CancellationToken.None);
 	}
diff --git a/Engine/Source/Programs/Horde/Horde.Agent.Tests/WorkerServiceTest.cs b/Engine/Source/Programs/Horde/Horde.Agent.Tests/WorkerServiceTest.cs
--- a/Engine/Source/Programs/Horde/Horde.Agent.Tests/WorkerServiceTest.cs
+++ b/Engine/Source/Programs/Horde/Horde.Agent.Tests/WorkerServiceTest.cs
@@ -265,6 +265,7 @@ namespace Horde.Agent.Tests
 		public static ISession CreateSession(IRpcConnection rpcConnection)
 		{
 			Mock<ISession> fakeSession = new Mock<ISession>(MockBehavior.Strict);
+			fakeSession.Setup(x => x.ServerUrl).Returns(new Uri("https://localhost:9999"));
 			fakeSession.Setup(x => x.AgentId).Returns("LocalAgent");
 			fakeSession.Setup(x => x.SessionId).Returns("Session");
 			fakeSession.Setup(x => x.RpcConnection).Returns(rpcConnection);
