Horde: Tighten up log poll time on initial complete and increase frequency
#prefight none
#skipci
#rnx

diff --git a/Engine/Source/Programs/Horde/HordeDashboard/src/backend/Api.ts b/Engine/Source/Programs/Horde/HordeDashboard/src/backend/Api.ts
--- a/Engine/Source/Programs/Horde/HordeDashboard/src/backend/Api.ts
+++ b/Engine/Source/Programs/Horde/HordeDashboard/src/backend/Api.ts
@@ -1183,7 +1183,7 @@ export type GetStepResponse = {
 	startTime?: Date | string;
 
 	/**Time at which the run finished (UTC) */
-	finishTime: Date | string;
+	finishTime?: Date | string;
 
 	/** List of reports */
 	reports?: GetReportResponse[];
diff --git a/Engine/Source/Programs/Horde/HordeDashboard/src/backend/JobDetails.ts b/Engine/Source/Programs/Horde/HordeDashboard/src/backend/JobDetails.ts
--- a/Engine/Source/Programs/Horde/HordeDashboard/src/backend/JobDetails.ts
+++ b/Engine/Source/Programs/Horde/HordeDashboard/src/backend/JobDetails.ts
@@ -629,12 +629,12 @@ export class JobDetails {
 
         const batch = this.batches.find((b) => b.logId === logId);
         if (batch) {
-            return batch.state === JobStepBatchState.Running;
+            return !batch.finishTime;
         }
 
         const step = this.getSteps().find(s => s.logId === logId);
         if (step) {
-            return step.state === JobStepState.Running;
+            return !step.finishTime;
         }
 
         return false;
diff --git a/Engine/Source/Programs/Horde/HordeDashboard/src/components/LogSource.tsx b/Engine/Source/Programs/Horde/HordeDashboard/src/components/LogSource.tsx
--- a/Engine/Source/Programs/Horde/HordeDashboard/src/components/LogSource.tsx
+++ b/Engine/Source/Programs/Horde/HordeDashboard/src/components/LogSource.tsx
@@ -67,7 +67,7 @@ export abstract class LogSource {
 
    initComplete() {
       if (this.active) {
-         setTimeout(() => { this.poll(); }, this.pollMS);
+         setTimeout(() => { this.poll(); }, 2000);
       }
    }
 
@@ -194,11 +194,10 @@ export abstract class LogSource {
          return;
       }
 
-      if (this.active) {
-         this.pollID = setTimeout(() => { this.poll(); }, this.pollMS);
-      }
-
       if (this.polling) {
+         if (this.active) {
+            this.pollID = setTimeout(() => { this.poll(); }, !!this._logItems.find(i => !!i.line) ? this.pollMS : 3000);
+         }         
          return;
       }
 
@@ -212,6 +211,8 @@ export abstract class LogSource {
          this.polling = false;
          if (!this.active) {
             this.stopPolling();
+         } else {
+            this.pollID = setTimeout(() => { this.poll(); }, !!this._logItems.find(i => !!i.line) ? this.pollMS : 3000);
          }
       });
 
@@ -272,7 +273,7 @@ export abstract class LogSource {
    trailing?: boolean;
 
    pollID?: any = undefined;
-   pollMS = 10000;
+   private pollMS = 5000;
    polling = false;
 
    startTime?: Moment;
