Horde: Fix complete tool deployments not being returned correctly.

diff --git a/Engine/Source/Programs/Horde/Horde.Server/Tools/ITool.cs b/Engine/Source/Programs/Horde/Horde.Server/Tools/ITool.cs
--- a/Engine/Source/Programs/Horde/Horde.Server/Tools/ITool.cs
+++ b/Engine/Source/Programs/Horde/Horde.Server/Tools/ITool.cs
@@ -116,7 +116,7 @@ namespace Horde.Server.Tools
 			int idx = tool.Deployments.Count - 1;
 			for (; idx >= 0; idx--)
 			{
-				if (phase < tool.Deployments[idx].GetProgressValue(utcNow) || idx == 0)
+				if (phase <= tool.Deployments[idx].GetProgressValue(utcNow) || idx == 0)
 				{
 					return tool.Deployments[idx];
 				}
