Horde: Fix infinite loop in compute handler.

diff --git a/Engine/Source/Programs/Horde/Horde.Agent/Leases/Handlers/ComputeHandler.cs b/Engine/Source/Programs/Horde/Horde.Agent/Leases/Handlers/ComputeHandler.cs
--- a/Engine/Source/Programs/Horde/Horde.Agent/Leases/Handlers/ComputeHandler.cs
+++ b/Engine/Source/Programs/Horde/Horde.Agent/Leases/Handlers/ComputeHandler.cs
@@ -39,7 +39,7 @@ namespace Horde.Agent.Leases.Handlers
 
 			public async ValueTask<int> ReadPartialAsync(Memory<byte> buffer, CancellationToken cancellationToken)
 			{
-				int result = await ReadPartialAsync(buffer, cancellationToken);
+				int result = await _inner.ReadPartialAsync(buffer, cancellationToken);
 				if (result > 0)
 				{
 					Interlocked.Exchange(ref _lastPingTicks, Stopwatch.GetTimestamp());
@@ -49,7 +49,7 @@ namespace Horde.Agent.Leases.Handlers
 
 			public async ValueTask WriteAsync(ReadOnlySequence<byte> buffer, CancellationToken cancellationToken)
 			{
-				await WriteAsync(buffer, cancellationToken);
+				await _inner.WriteAsync(buffer, cancellationToken);
 				Interlocked.Exchange(ref _lastPingTicks, Stopwatch.GetTimestamp());
 			}
 		}
