Horde: Additional logging for log tailing requests.

diff --git a/Engine/Source/Programs/Horde/Horde.Agent/Utility/JsonRpcLogSink.cs b/Engine/Source/Programs/Horde/Horde.Agent/Utility/JsonRpcLogSink.cs
--- a/Engine/Source/Programs/Horde/Horde.Agent/Utility/JsonRpcLogSink.cs
+++ b/Engine/Source/Programs/Horde/Horde.Agent/Utility/JsonRpcLogSink.cs
@@ -167,6 +167,8 @@ namespace Horde.Agent.Utility
 
 				// Update the next tailing position
 				int newTailNext = await UpdateLogTailAsync(tailNext, tailData);
+				_logger.LogDebug("Log {LogId} tail next = {TailNext}", _logId, newTailNext);
+
 				if (newTailNext != tailNext)
 				{
 					tailNext = newTailNext;
diff --git a/Engine/Source/Programs/Horde/Horde.Server/Logs/LogRpcService.cs b/Engine/Source/Programs/Horde/Horde.Server/Logs/LogRpcService.cs
--- a/Engine/Source/Programs/Horde/Horde.Server/Logs/LogRpcService.cs
+++ b/Engine/Source/Programs/Horde/Horde.Server/Logs/LogRpcService.cs
@@ -9,6 +9,7 @@ using EpicGames.Horde.Storage;
 using Horde.Server.Storage;
 using Microsoft.Extensions.Logging;
 using System.Threading;
+using System;
 
 namespace Horde.Server.Logs
 {
@@ -89,7 +90,14 @@ namespace Horde.Server.Logs
 						await Task.WhenAny(waitTask, moveNextTask);
 						cancellationSource.Cancel();
 
-						response.TailNext = await waitTask;
+						try
+						{
+							response.TailNext = await waitTask;
+						}
+						catch (Exception ex)
+						{
+							_logger.LogWarning(ex, "Exception while waiting for tail next");
+						}
 					}
 				}
 
