Horde: Fix agent timeout not waiting the correct amount of time.

diff --git a/Engine/Source/Programs/Horde/Horde.Agent/Leases/Handlers/ComputeHandler.cs b/Engine/Source/Programs/Horde/Horde.Agent/Leases/Handlers/ComputeHandler.cs
--- a/Engine/Source/Programs/Horde/Horde.Agent/Leases/Handlers/ComputeHandler.cs
+++ b/Engine/Source/Programs/Horde/Horde.Agent/Leases/Handlers/ComputeHandler.cs
@@ -31,6 +31,7 @@ namespace Horde.Agent.Leases.Handlers
 			public TcpTransportWithTimeout(Socket socket)
 			{
 				_inner = new TcpTransport(socket);
+				_lastPingTicks = Stopwatch.GetTimestamp();
 			}
 
 			public TimeSpan TimeSinceActivity => TimeSpan.FromTicks(Stopwatch.GetTimestamp() - Interlocked.CompareExchange(ref _lastPingTicks, 0, 0));
@@ -113,12 +114,13 @@ namespace Horde.Agent.Leases.Handlers
 			}
 		}
 
-		static async Task TickTimeoutAsync(TcpTransportWithTimeout transport, CancellationTokenSource cts, CancellationToken cancellationToken)
+		async Task TickTimeoutAsync(TcpTransportWithTimeout transport, CancellationTokenSource cts, CancellationToken cancellationToken)
 		{
 			while(!cancellationToken.IsCancellationRequested)
 			{
 				if (transport.TimeSinceActivity > TimeSpan.FromMinutes(10))
 				{
+					_logger.LogWarning("Terminating compute task due to timeout (last tick at {Time})", DateTime.UtcNow - transport.TimeSinceActivity);
 					cts.Cancel();
 					break;
 				}
