Horde: Use a sandbox directory under the session working dir.

diff --git a/Engine/Source/Programs/Horde/Horde.Agent/Leases/Handlers/ComputeHandler.cs b/Engine/Source/Programs/Horde/Horde.Agent/Leases/Handlers/ComputeHandler.cs
--- a/Engine/Source/Programs/Horde/Horde.Agent/Leases/Handlers/ComputeHandler.cs
+++ b/Engine/Source/Programs/Horde/Horde.Agent/Leases/Handlers/ComputeHandler.cs
@@ -57,7 +57,6 @@ namespace Horde.Agent.Leases.Handlers
 
 		readonly ComputeListenerService _listenerService;
 		readonly IMemoryCache _memoryCache;
-		readonly DirectoryReference _sandboxDir;
 		readonly ILogger _logger;
 
 		/// <summary>
@@ -67,7 +66,6 @@ namespace Horde.Agent.Leases.Handlers
 		{
 			_listenerService = listenerService;
 			_memoryCache = memoryCache;
-			_sandboxDir = DirectoryReference.Combine(Program.DataDir, "Sandbox");
 			_logger = logger;
 		}
 
@@ -96,10 +94,20 @@ namespace Horde.Agent.Leases.Handlers
 					await using BackgroundTask timeoutTask = BackgroundTask.StartNew(ctx => TickTimeoutAsync(transport, cts, ctx));
 					await using (IComputeSocket socket = ComputeSocket.Create(transport, _logger))
 					{
-						ComputeWorker worker = new ComputeWorker(_sandboxDir, _memoryCache, _logger);
-						await worker.RunAsync(socket, cts.Token);
-						await socket.CloseAsync(cts.Token);
-						return LeaseResult.Success;
+						DirectoryReference sandboxDir = DirectoryReference.Combine(session.WorkingDir, "Sandbox", leaseId);
+						try
+						{
+							DirectoryReference.CreateDirectory(sandboxDir);
+
+							ComputeWorker worker = new ComputeWorker(sandboxDir, _memoryCache, _logger);
+							await worker.RunAsync(socket, cts.Token);
+							await socket.CloseAsync(cts.Token);
+							return LeaseResult.Success;
+						}
+						finally
+						{
+							FileUtils.ForceDeleteDirectory(sandboxDir);
+						}
 					}
 				}
 			}
