Horde: Store issues affecting steps in history
#rnx

diff --git a/Engine/Source/Programs/Horde/Horde.Server.Tests/IssueServiceTests.cs b/Engine/Source/Programs/Horde/Horde.Server.Tests/IssueServiceTests.cs
--- a/Engine/Source/Programs/Horde/Horde.Server.Tests/IssueServiceTests.cs
+++ b/Engine/Source/Programs/Horde/Horde.Server.Tests/IssueServiceTests.cs
@@ -20,9 +20,7 @@ using Horde.Server.Projects;
 using Horde.Server.Streams;
 using Horde.Server.Server;
 using Horde.Server.Tests.Stubs.Services;
-using Horde.Server.Utilities;
 using HordeCommon;
-using Microsoft.Extensions.DependencyInjection;
 using Microsoft.Extensions.Logging;
 using Microsoft.VisualStudio.TestTools.UnitTesting;
 using MongoDB.Bson;
@@ -255,14 +253,14 @@ namespace Horde.Server.Tests
 			IJobStepBatch batch = job.Batches[batchIdx];
 			IJobStep step = batch.Steps[stepIdx];
 
+			JobStepRefId jobStepRefId = new JobStepRefId(job.Id, batch.Id, step.Id);
+			string nodeName = _graph.Groups[batch.GroupIdx].Nodes[step.NodeIdx].Name;
+			await JobStepRefCollection.InsertOrReplaceAsync(jobStepRefId, "TestJob", nodeName, job.StreamId, job.TemplateId, job.Change, step.LogId, null, null, outcome, job.UpdateIssues, null, null, 0.0f, 0.0f, job.CreateTimeUtc, step.StartTimeUtc!.Value, step.StartTimeUtc);
+
 			if (job.UpdateIssues)
 			{
 				await IssueService.UpdateCompleteStep(job, _graph, batch.Id, step.Id);
 			}			
-
-			JobStepRefId jobStepRefId = new JobStepRefId(job.Id, batch.Id, step.Id);
-			string nodeName = _graph.Groups[batch.GroupIdx].Nodes[step.NodeIdx].Name;
-			await JobStepRefCollection.InsertOrReplaceAsync(jobStepRefId, "TestJob", nodeName, job.StreamId, job.TemplateId, job.Change, step.LogId, null, null, outcome, job.UpdateIssues, null, null, 0.0f, 0.0f, job.CreateTimeUtc, step.StartTimeUtc!.Value, step.StartTimeUtc);
 		}
 
 		async Task AddEvent(IJob job, int batchIdx, int stepIdx, object data, EventSeverity severity = EventSeverity.Error)
@@ -347,11 +345,16 @@ namespace Horde.Server.Tests
 				IJob job = CreateJob(_mainStreamId, 105, "Test Build", _graph);
 				await AddEvent(job, 0, 0, new { level = nameof(LogLevel.Warning), message = "" }, EventSeverity.Warning);
 				await UpdateCompleteStep(job, 0, 0, JobStepOutcome.Warnings);
+				IJobStepRef? stepRef = await JobStepRefCollection.FindAsync(job.Id, job.Batches[0].Id, job.Batches[0].Steps[0].Id);
+				
 
 				List<IIssue> issues = await IssueCollection.FindIssuesAsync();
 				Assert.AreEqual(1, issues.Count);
 
 				Assert.AreEqual("Warnings in Update Version Files", issues[0].Summary);
+
+				Assert.AreEqual(1, stepRef!.IssueIds!.Count);
+				Assert.AreEqual(1, stepRef!.IssueIds![0], issues[0].Id );
 			}
 
 			// #2
diff --git a/Engine/Source/Programs/Horde/Horde.Server.Tests/Stubs/Collections/JobStepRefCollectionStub.cs b/Engine/Source/Programs/Horde/Horde.Server.Tests/Stubs/Collections/JobStepRefCollectionStub.cs
deleted file mode 100644
diff --git a/Engine/Source/Programs/Horde/Horde.Server/Issues/IssueService.cs b/Engine/Source/Programs/Horde/Horde.Server/Issues/IssueService.cs
--- a/Engine/Source/Programs/Horde/Horde.Server/Issues/IssueService.cs
+++ b/Engine/Source/Programs/Horde/Horde.Server/Issues/IssueService.cs
@@ -597,6 +597,13 @@ namespace Horde.Server.Issues
 					break;
 				}
 			}
+
+			List<IIssue> issues = await Collection.FindIssuesForJobAsync(job, graph, stepId, batchId);
+			if (issues.Count > 0)
+			{
+				await _jobStepRefs.UpdateAsync(job.Id, batchId, stepId, issues.Select(i => i.Id).ToList());
+			}
+
 		}
 
 		/// <summary>
diff --git a/Engine/Source/Programs/Horde/Horde.Server/Jobs/IJobStepRef.cs b/Engine/Source/Programs/Horde/Horde.Server/Jobs/IJobStepRef.cs
--- a/Engine/Source/Programs/Horde/Horde.Server/Jobs/IJobStepRef.cs
+++ b/Engine/Source/Programs/Horde/Horde.Server/Jobs/IJobStepRef.cs
@@ -1,11 +1,11 @@
 // Copyright Epic Games, Inc. All Rights Reserved.
 
 using System;
+using System.Collections.Generic;
 using Horde.Server.Agents;
 using Horde.Server.Agents.Pools;
 using Horde.Server.Logs;
 using Horde.Server.Streams;
-using Horde.Server.Utilities;
 using HordeCommon;
 using MongoDB.Bson.Serialization;
 using MongoDB.Bson.Serialization.Attributes;
@@ -65,6 +65,12 @@ namespace Horde.Server.Jobs
 		{
 			return $"{JobId}:{BatchId}:{StepId}";
 		}
+
+		/// <inheritdoc/>
+		public bool Equals(JobStepRefId other)
+		{
+			return JobId.Equals(other.JobId) && BatchId.Equals(other.BatchId) && StepId.Equals(other.StepId);
+		}
 	}
 
 	/// <summary>
@@ -160,6 +166,11 @@ namespace Horde.Server.Jobs
 		/// </summary>
 		public bool UpdateIssues { get; }
 
+		/// <summary>
+		/// Issues ids affecting this job step
+		/// </summary>
+		public IReadOnlyList<int>? IssueIds { get; }
+
 		/// <summary>
 		/// The last change that succeeded. Note that this is only set when the ref is updated; it is not necessarily consistent with steps run later.
 		/// </summary>
diff --git a/Engine/Source/Programs/Horde/Horde.Server/Jobs/IJobStepRefCollection.cs b/Engine/Source/Programs/Horde/Horde.Server/Jobs/IJobStepRefCollection.cs
--- a/Engine/Source/Programs/Horde/Horde.Server/Jobs/IJobStepRefCollection.cs
+++ b/Engine/Source/Programs/Horde/Horde.Server/Jobs/IJobStepRefCollection.cs
@@ -8,7 +8,6 @@ using Horde.Server.Agents.Pools;
 using Horde.Server.Jobs.Graphs;
 using Horde.Server.Logs;
 using Horde.Server.Streams;
-using Horde.Server.Utilities;
 using HordeCommon;
 using Microsoft.Extensions.Logging;
 
@@ -42,6 +41,25 @@ namespace Horde.Server.Jobs
 		/// <param name="finishTimeUtc">Finish time for the step, if known</param>
 		Task<IJobStepRef> InsertOrReplaceAsync(JobStepRefId id, string jobName, string stepName, StreamId streamId, TemplateId templateId, int change, LogId? logId, PoolId? poolId, AgentId? agentId, JobStepOutcome? outcome, bool updateIssues, int? lastSuccess, int? lastWarning, float waitTime, float initTime, DateTime jobStartTimeUtc, DateTime startTimeUtc, DateTime? finishTimeUtc);
 
+		/// <summary>
+		/// Updates a job step ref 
+		/// </summary>
+		/// <param name="jobId"></param>
+		/// <param name="batchId"></param>
+		/// <param name="stepId"></param>
+		/// <param name="issueIds"></param>
+		/// <returns></returns>
+		Task<IJobStepRef?> UpdateAsync(JobId jobId, SubResourceId batchId, SubResourceId stepId, List<int>? issueIds);
+
+		/// <summary>
+		/// Gets a specific job step ref
+		/// </summary>
+		/// <param name="jobId"></param>
+		/// <param name="batchId"></param>
+		/// <param name="stepId"></param>
+		/// <returns></returns>
+		Task<IJobStepRef?> FindAsync(JobId jobId, SubResourceId batchId, SubResourceId stepId);
+
 		/// <summary>
 		/// Gets the history of a given node
 		/// </summary>
diff --git a/Engine/Source/Programs/Horde/Horde.Server/Jobs/JobMessages.cs b/Engine/Source/Programs/Horde/Horde.Server/Jobs/JobMessages.cs
--- a/Engine/Source/Programs/Horde/Horde.Server/Jobs/JobMessages.cs
+++ b/Engine/Source/Programs/Horde/Horde.Server/Jobs/JobMessages.cs
@@ -1122,6 +1122,11 @@ namespace Horde.Server.Jobs
 		/// </summary>
 		public JobStepOutcome? Outcome { get; set; }
 
+		/// <summary>
+		/// The issues which affected this step
+		/// </summary>
+		public List<int>? IssueIds { get; }
+
 		/// <summary>
 		/// Time at which the step started.
 		/// </summary>
@@ -1146,6 +1151,7 @@ namespace Horde.Server.Jobs
 			PoolId = jobStepRef.PoolId?.ToString();
 			AgentId = jobStepRef.AgentId?.ToString();
 			Outcome = jobStepRef.Outcome;
+			IssueIds = jobStepRef.IssueIds?.Select(id => id).ToList();
 			StartTime = jobStepRef.StartTimeUtc;
 			FinishTime = jobStepRef.FinishTimeUtc;
 		}
diff --git a/Engine/Source/Programs/Horde/Horde.Server/Jobs/JobStepRefCollection.cs b/Engine/Source/Programs/Horde/Horde.Server/Jobs/JobStepRefCollection.cs
--- a/Engine/Source/Programs/Horde/Horde.Server/Jobs/JobStepRefCollection.cs
+++ b/Engine/Source/Programs/Horde/Horde.Server/Jobs/JobStepRefCollection.cs
@@ -65,10 +65,14 @@ namespace Horde.Server.Jobs
 			[BsonIgnoreIfNull]
 			public bool? UpdateIssues { get; set; }
 
+			[BsonIgnoreIfNull]
+			public List<int>? IssueIds { get; set; }
+
 			DateTime IJobStepRef.StartTimeUtc => StartTimeUtc ?? StartTime?.UtcDateTime ?? default;
 			DateTime? IJobStepRef.FinishTimeUtc => FinishTimeUtc ?? FinishTime?.UtcDateTime;
 			string IJobStepRef.NodeName => Name;
-			bool IJobStepRef.UpdateIssues => UpdateIssues ?? false;
+			bool IJobStepRef.UpdateIssues => UpdateIssues ?? false;			
+			IReadOnlyList<int>? IJobStepRef.IssueIds => IssueIds;
 
 			public JobStepRef(JobStepRefId id, string jobName, string nodeName, StreamId streamId, TemplateId templateId, int change, LogId? logId, PoolId? poolId, AgentId? agentId, JobStepOutcome? outcome, bool updateIssues, int? lastSuccess, int? lastWarning, float batchWaitTime, float batchInitTime, DateTime jobStartTimeUtc, DateTime startTimeUtc, DateTime? finishTimeUtc)
 			{
@@ -141,6 +145,33 @@ namespace Horde.Server.Jobs
 			return newJobStepRef;
 		}
 
+		/// <inheritdoc/>
+		public async Task<IJobStepRef?> UpdateAsync(JobId jobId, SubResourceId batchId, SubResourceId stepId, List<int>? issueIds)
+		{
+			UpdateDefinitionBuilder<JobStepRef> updateBuilder = Builders<JobStepRef>.Update;
+			List<UpdateDefinition<JobStepRef>> updates = new List<UpdateDefinition<JobStepRef>>();
+
+			if (issueIds != null)
+			{
+				updates.Add(updateBuilder.Set(x => x.IssueIds, issueIds));
+			}
+
+			if (updates.Count == 0) 
+			{
+				return await FindAsync(jobId, batchId, stepId);
+			}
+			
+			JobStepRefId id = new JobStepRefId(jobId, batchId, stepId);
+			return await _jobStepRefs.FindOneAndUpdateAsync(x => x.Id.Equals(id), updateBuilder.Combine(updates));
+		}
+
+		/// <inheritdoc/>
+		public async Task<IJobStepRef?> FindAsync(JobId jobId, SubResourceId batchId, SubResourceId stepId)
+		{
+			JobStepRefId id = new JobStepRefId(jobId, batchId, stepId);
+			return await _jobStepRefs.Find(x => x.Id.Equals(id)).FirstOrDefaultAsync();
+		}
+
 		/// <inheritdoc/>
 		public async Task<List<IJobStepRef>> GetStepsForNodeAsync(StreamId streamId, TemplateId templateId, string nodeName, int? change, bool includeFailed, int count)
 		{
