Horde: Additional logging for log tailing.

diff --git a/Engine/Source/Programs/Horde/Horde.Agent/Utility/JsonRpcLogSink.cs b/Engine/Source/Programs/Horde/Horde.Agent/Utility/JsonRpcLogSink.cs
--- a/Engine/Source/Programs/Horde/Horde.Agent/Utility/JsonRpcLogSink.cs
+++ b/Engine/Source/Programs/Horde/Horde.Agent/Utility/JsonRpcLogSink.cs
@@ -161,13 +161,15 @@ namespace Horde.Agent.Utility
 				// If we don't have any updates for the server, wait until we do.
 				if (tailNext != -1 && tailData.IsEmpty)
 				{
+					_logger.LogInformation("No tail data available for log {LogId} after {TailNext}; waiting for more...", _logId, tailNext);
 					await newTailDataTask;
 					continue;
 				}
 
 				// Update the next tailing position
+				_logger.LogInformation("Setting log {LogId} tail = {TailNext}, data = {TailDataSize} bytes, {NumLines} lines", _logId, tailNext, tailData.Length, CountLines(tailData.Span));
 				int newTailNext = await UpdateLogTailAsync(tailNext, tailData);
-				_logger.LogDebug("Log {LogId} tail next = {TailNext}", _logId, newTailNext);
+				_logger.LogInformation("Log {LogId} tail next = {TailNext}", _logId, newTailNext);
 
 				if (newTailNext != tailNext)
 				{
@@ -177,6 +179,19 @@ namespace Horde.Agent.Utility
 			}
 		}
 
+		static int CountLines(ReadOnlySpan<byte> data)
+		{
+			int lines = 0;
+			for (int idx = 0; idx < data.Length; idx++)
+			{
+				if (data[idx] == '\n')
+				{
+					lines++;
+				}
+			}
+			return lines;
+		}
+
 		/// <inheritdoc/>
 		public async Task SetOutcomeAsync(JobStepOutcome outcome, CancellationToken cancellationToken)
 		{
