Horde: Add temp storage exception for tbbmalloc.pdb.

diff --git a/Engine/Source/Programs/Horde/Horde.Agent/Utility/TempStorage.cs b/Engine/Source/Programs/Horde/Horde.Agent/Utility/TempStorage.cs
--- a/Engine/Source/Programs/Horde/Horde.Agent/Utility/TempStorage.cs
+++ b/Engine/Source/Programs/Horde/Horde.Agent/Utility/TempStorage.cs
@@ -781,7 +781,7 @@ namespace Horde.Storage.Utility
 			{
 				return true;
 			}
-			if (fileName.Equals("tbbmalloc.dll", StringComparison.OrdinalIgnoreCase) || fileName.Equals("libtbbmalloc.dylib", StringComparison.OrdinalIgnoreCase))
+			if (fileName.Equals("tbbmalloc.dll", StringComparison.OrdinalIgnoreCase) || fileName.Equals("tbbmalloc.pdb", StringComparison.OrdinalIgnoreCase) || fileName.Equals("libtbbmalloc.dylib", StringComparison.OrdinalIgnoreCase))
 			{
 				return true;
 			}
