Horde: Add more logging for server-side log tailing.

diff --git a/Engine/Source/Programs/Horde/Horde.Server/Logs/LogRpcService.cs b/Engine/Source/Programs/Horde/Horde.Server/Logs/LogRpcService.cs
--- a/Engine/Source/Programs/Horde/Horde.Server/Logs/LogRpcService.cs
+++ b/Engine/Source/Programs/Horde/Horde.Server/Logs/LogRpcService.cs
@@ -72,6 +72,7 @@ namespace Horde.Server.Logs
 			{
 				UpdateLogTailRequest request = requestStream.Current;
 				LogId logId = LogId.Parse(request.LogId);
+				_logger.LogDebug("Updating log tail for {LogId}: line {LineIdx}, size {Size}", logId, request.TailNext, request.TailData.Length);
 
 				if (request.TailData.Length > 0)
 				{
diff --git a/Engine/Source/Programs/Horde/Horde.Server/Logs/LogTailService.cs b/Engine/Source/Programs/Horde/Horde.Server/Logs/LogTailService.cs
--- a/Engine/Source/Programs/Horde/Horde.Server/Logs/LogTailService.cs
+++ b/Engine/Source/Programs/Horde/Horde.Server/Logs/LogTailService.cs
@@ -223,6 +223,7 @@ namespace Horde.Server.Logs
 		{
 			List<SortedSetEntry<ReadOnlyMemory<byte>>> entries = new List<SortedSetEntry<ReadOnlyMemory<byte>>>();
 			int newLineCount = SplitLogDataToChunks(tailNext, tailData, ChunkLineCount, entries);
+			_logger.LogDebug("New line count = {NewLineCount}, entries = {EntryCount}", newLineCount, entries.Count);
 
 			if (entries.Count > 0)
 			{
