Horde: More debug logging.

diff --git a/Engine/Source/Programs/Horde/Horde.Agent/Utility/JsonRpcLogSink.cs b/Engine/Source/Programs/Horde/Horde.Agent/Utility/JsonRpcLogSink.cs
--- a/Engine/Source/Programs/Horde/Horde.Agent/Utility/JsonRpcLogSink.cs
+++ b/Engine/Source/Programs/Horde/Horde.Agent/Utility/JsonRpcLogSink.cs
@@ -2,6 +2,7 @@
 
 using System;
 using System.Collections.Generic;
+using System.Text;
 using System.Threading;
 using System.Threading.Tasks;
 using EpicGames.Core;
@@ -167,7 +168,14 @@ namespace Horde.Agent.Utility
 				}
 
 				// Update the next tailing position
-				_logger.LogInformation("Setting log {LogId} tail = {TailNext}, data = {TailDataSize} bytes, {NumLines} lines", _logId, tailNext, tailData.Length, CountLines(tailData.Span));
+				int numLines = CountLines(tailData.Span);
+				_logger.LogInformation("Setting log {LogId} tail = {TailNext}, data = {TailDataSize} bytes, {NumLines} lines", _logId, tailNext, tailData.Length, numLines);
+				if (tailData.Length > 0 && numLines == 0)
+				{
+					string text = Encoding.UTF8.GetString(tailData.Slice(0, Math.Min(tailData.Length, 256)).Span);
+					_logger.LogInformation("Initial text: {Text}", text);
+				}
+
 				int newTailNext = await UpdateLogTailAsync(tailNext, tailData);
 				_logger.LogInformation("Log {LogId} tail next = {TailNext}", _logId, newTailNext);
 
