Horde: More log tail logging.

diff --git a/Engine/Source/Programs/Horde/Horde.Server/Logs/LogRpcService.cs b/Engine/Source/Programs/Horde/Horde.Server/Logs/LogRpcService.cs
--- a/Engine/Source/Programs/Horde/Horde.Server/Logs/LogRpcService.cs
+++ b/Engine/Source/Programs/Horde/Horde.Server/Logs/LogRpcService.cs
@@ -86,6 +86,7 @@ namespace Horde.Server.Logs
 				{
 					using (CancellationTokenSource cancellationSource = new CancellationTokenSource())
 					{
+						_logger.LogDebug("Waiting for tail next on log {LogId}", logId);
 						Task<int> waitTask = _logTailService.WaitForTailNextAsync(logId, cancellationSource.Token);
 
 						await Task.WhenAny(waitTask, moveNextTask);
@@ -102,6 +103,7 @@ namespace Horde.Server.Logs
 					}
 				}
 
+				_logger.LogDebug("Return tail next for log {LogId} = {TailNext}", logId, response.TailNext);
 				await responseStream.WriteAsync(response);
 			}
 
