Horde: More logging for logging.

diff --git a/Engine/Source/Programs/Horde/Horde.Agent/Utility/JsonRpcLogSink.cs b/Engine/Source/Programs/Horde/Horde.Agent/Utility/JsonRpcLogSink.cs
--- a/Engine/Source/Programs/Horde/Horde.Agent/Utility/JsonRpcLogSink.cs
+++ b/Engine/Source/Programs/Horde/Horde.Agent/Utility/JsonRpcLogSink.cs
@@ -167,15 +167,16 @@ namespace Horde.Agent.Utility
 					continue;
 				}
 
-				// Update the next tailing position
-				int numLines = CountLines(tailData.Span);
-				_logger.LogInformation("Setting log {LogId} tail = {TailNext}, data = {TailDataSize} bytes, {NumLines} lines", _logId, tailNext, tailData.Length, numLines);
-				if (tailData.Length > 0 && numLines == 0)
+				string start = "";
+				if (tailData.Length > 0)
 				{
-					string text = Encoding.UTF8.GetString(tailData.Slice(0, Math.Min(tailData.Length, 256)).Span);
-					_logger.LogInformation("Initial text: {Text}", text);
+					start = Encoding.UTF8.GetString(tailData.Slice(0, Math.Min(tailData.Length, 256)).Span);
 				}
 
+				// Update the next tailing position
+				int numLines = CountLines(tailData.Span);
+				_logger.LogInformation("Setting log {LogId} tail = {TailNext}, data = {TailDataSize} bytes, {NumLines} lines ('{Start}')", _logId, tailNext, tailData.Length, numLines, start);
+
 				int newTailNext = await UpdateLogTailAsync(tailNext, tailData);
 				_logger.LogInformation("Log {LogId} tail next = {TailNext}", _logId, newTailNext);
 
@@ -261,6 +262,7 @@ namespace Horde.Agent.Utility
 					request.TailNext = tailNext;
 					request.TailData = UnsafeByteOperations.UnsafeWrap(tailData);
 					await call.RequestStream.WriteAsync(request);
+					_logger.LogInformation("Writing log data: {LogId}, {TailNext}, {TailData} bytes", _logId, tailNext, tailData.Length);
 
 					// Wait until the server responds or we need to trigger a new update
 					Task<bool> moveNextAsync = call.ResponseStream.MoveNext();
@@ -268,11 +270,11 @@ namespace Horde.Agent.Utility
 					Task task = await Task.WhenAny(moveNextAsync, clientRef.DisposingTask, _tailTaskStop.Task, Task.Delay(TimeSpan.FromMinutes(1.0), CancellationToken.None));
 					if (task == clientRef.DisposingTask)
 					{
-						_logger.LogDebug("Cancelling long poll from client side (server migration)");
+						_logger.LogInformation("Cancelling long poll from client side (server migration)");
 					}
 					else if (task == _tailTaskStop.Task)
 					{
-						_logger.LogDebug("Cancelling long poll from client side (complete)");
+						_logger.LogInformation("Cancelling long poll from client side (complete)");
 					}
 
 					// Close the request stream to indicate that we're finished
