Horde: Additional logging for enabling tailing on an entry.

diff --git a/Engine/Source/Programs/Horde/Horde.Server/Logs/LogTailService.cs b/Engine/Source/Programs/Horde/Horde.Server/Logs/LogTailService.cs
--- a/Engine/Source/Programs/Horde/Horde.Server/Logs/LogTailService.cs
+++ b/Engine/Source/Programs/Horde/Horde.Server/Logs/LogTailService.cs
@@ -196,6 +196,8 @@ namespace Horde.Server.Logs
 
 			if (waitForTailData)
 			{
+				_logger.LogDebug("Waiting for tail data on log {LogId}", logId);
+
 				// Bit hacky; poll for a while until we get some tail data through
 				RedisSortedSetKey<ReadOnlyMemory<byte>> tailDataKey = TailDataKey(logId);
 				for (int idx = 0; idx < 20; idx++)
@@ -209,6 +211,8 @@ namespace Horde.Server.Logs
 						await Task.Delay(100);
 					}
 				}
+
+				_logger.LogDebug("Redis key for tail data on {LogId} is set", logId);
 			}
 		}
 
