Horde: Replace use of OpenTracing with OpenTelemetry in server

diff --git a/Engine/Source/Programs/Horde/Horde.Server.Tests/Fleet/AwsRecyclingFleetManagerTest.cs b/Engine/Source/Programs/Horde/Horde.Server.Tests/Fleet/AwsRecyclingFleetManagerTest.cs
--- a/Engine/Source/Programs/Horde/Horde.Server.Tests/Fleet/AwsRecyclingFleetManagerTest.cs
+++ b/Engine/Source/Programs/Horde/Horde.Server.Tests/Fleet/AwsRecyclingFleetManagerTest.cs
@@ -219,7 +219,7 @@ namespace Horde.Server.Tests.Fleet
 			ILogger<AwsRecyclingFleetManager> logger = loggerFactory.CreateLogger<AwsRecyclingFleetManager>();
 			IPool pool = await PoolService.CreatePoolAsync("testPool", null, true, 0, 0, sizeStrategy: PoolSizeStrategy.NoOp);
 			using NoOpDogStatsd dogStatsd = new ();
-			AwsRecyclingFleetManager manager = new (ec2, AgentCollection, dogStatsd, _clock, settings, logger);
+			AwsRecyclingFleetManager manager = new (ec2, AgentCollection, dogStatsd, _clock, settings, Tracer, logger);
 			return await manager.ExpandPoolAsync(pool, new List<IAgent>(), numRequestedInstances, CancellationToken.None);
 		}
 	}
diff --git a/Engine/Source/Programs/Horde/Horde.Server.Tests/Fleet/AwsReuseFleetManagerTest.cs b/Engine/Source/Programs/Horde/Horde.Server.Tests/Fleet/AwsReuseFleetManagerTest.cs
--- a/Engine/Source/Programs/Horde/Horde.Server.Tests/Fleet/AwsReuseFleetManagerTest.cs
+++ b/Engine/Source/Programs/Horde/Horde.Server.Tests/Fleet/AwsReuseFleetManagerTest.cs
@@ -46,7 +46,7 @@ namespace Horde.Server.Tests.Fleet
 		{
 			ILogger<AwsReuseFleetManager> logger = ServiceProvider.GetRequiredService<ILogger<AwsReuseFleetManager>>();
 			IPool pool = await PoolService.CreatePoolAsync("testPool", null, true, 0, 0, sizeStrategy: PoolSizeStrategy.NoOp);
-			AwsReuseFleetManager manager = new (ec2, AgentCollection, settings, logger);
+			AwsReuseFleetManager manager = new (ec2, AgentCollection, settings, Tracer, logger);
 			await manager.ExpandPoolAsync(pool, new List<IAgent>(), numRequestedInstances, CancellationToken.None);
 		}
 	}
diff --git a/Engine/Source/Programs/Horde/Horde.Server.Tests/Fleet/FleetServiceTest.cs b/Engine/Source/Programs/Horde/Horde.Server.Tests/Fleet/FleetServiceTest.cs
--- a/Engine/Source/Programs/Horde/Horde.Server.Tests/Fleet/FleetServiceTest.cs
+++ b/Engine/Source/Programs/Horde/Horde.Server.Tests/Fleet/FleetServiceTest.cs
@@ -16,6 +16,7 @@ using Microsoft.Extensions.DependencyInjection;
 using Microsoft.Extensions.Logging;
 using Microsoft.Extensions.Options;
 using Microsoft.VisualStudio.TestTools.UnitTesting;
+using OpenTelemetry.Trace;
 using StatsdClient;
 
 namespace Horde.Server.Tests.Fleet
@@ -240,10 +241,10 @@ namespace Horde.Server.Tests.Fleet
 			ILoggerFactory loggerFactory = ServiceProvider.GetRequiredService<ILoggerFactory>();
 			IOptions<ServerSettings> serverSettingsOpt = ServiceProvider.GetRequiredService<IOptions<ServerSettings>>();
 			serverSettingsOpt.Value.FleetManagerV2 = FleetManagerType.AwsReuse;
-			
+
 			FleetService service = new(
 				AgentCollection, GraphCollection, JobCollection, LeaseCollection, PoolCollection, new DowntimeServiceStub(isDowntimeActive), StreamCollection, _dogStatsD,
-				new StubFleetManagerFactory(fleetManager), Clock, Cache, serverSettingsOpt, GlobalConfig, loggerFactory.CreateLogger<FleetService>());
+				new StubFleetManagerFactory(fleetManager), Clock, Cache, serverSettingsOpt, GlobalConfig, Tracer, loggerFactory.CreateLogger<FleetService>());
 				
 			return service;
 		}
diff --git a/Engine/Source/Programs/Horde/Horde.Server.Tests/LogBuilderTests.cs b/Engine/Source/Programs/Horde/Horde.Server.Tests/LogBuilderTests.cs
--- a/Engine/Source/Programs/Horde/Horde.Server.Tests/LogBuilderTests.cs
+++ b/Engine/Source/Programs/Horde/Horde.Server.Tests/LogBuilderTests.cs
@@ -11,6 +11,7 @@ using Horde.Server.Server;
 using Horde.Server.Utilities;
 using Microsoft.Extensions.Logging.Abstractions;
 using Microsoft.VisualStudio.TestTools.UnitTesting;
+using OpenTelemetry.Trace;
 
 namespace Horde.Server.Tests
 {
@@ -28,7 +29,8 @@ namespace Horde.Server.Tests
 		public async Task TestRedisLogBuilder()
 		{
 			RedisService redisService = GetRedisServiceSingleton();
-			ILogBuilder builder = new RedisLogBuilder(redisService.ConnectionPool, NullLogger.Instance);
+			Tracer tracer = TracerProvider.Default.GetTracer("LogBuilderTests");
+			ILogBuilder builder = new RedisLogBuilder(redisService.ConnectionPool, tracer, NullLogger.Instance);
 			await TestBuilder(builder);
 		}
 
diff --git a/Engine/Source/Programs/Horde/Horde.Server.Tests/LogFileServiceTests.cs b/Engine/Source/Programs/Horde/Horde.Server.Tests/LogFileServiceTests.cs
--- a/Engine/Source/Programs/Horde/Horde.Server.Tests/LogFileServiceTests.cs
+++ b/Engine/Source/Programs/Horde/Horde.Server.Tests/LogFileServiceTests.cs
@@ -23,6 +23,7 @@ using Microsoft.Extensions.Logging;
 using Microsoft.Extensions.Logging.Abstractions;
 using Microsoft.VisualStudio.TestTools.UnitTesting;
 using MongoDB.Bson;
+using OpenTelemetry.Trace;
 
 namespace Horde.Server.Tests
 {
@@ -42,12 +43,13 @@ namespace Horde.Server.Tests
 			_loggerFactory = new LoggerFactory();
             ILogger<LogFileService> logger = _loggerFactory.CreateLogger<LogFileService>();
 
-			ILogBuilder logBuilder = new RedisLogBuilder(GetRedisServiceSingleton().ConnectionPool, NullLogger.Instance);
+            Tracer tracer = TracerProvider.Default.GetTracer("LogFileServiceTest");
+            ILogBuilder logBuilder = new RedisLogBuilder(GetRedisServiceSingleton().ConnectionPool, tracer, NullLogger.Instance);
 			_logStorageBackend = new MemoryStorageBackend();
 			_logStorage = new PersistentLogStorage(_logStorageBackend.ForType<PersistentLogStorage>(), NullLogger<PersistentLogStorage>.Instance);
 			_clock = new FakeClock();
 			TestOptions<ServerSettings> settingsOpts = new (new ServerSettings());
-			_logFileService = new LogFileService(logFileCollection, null!, logBuilder, _logStorage, _clock, null!, null!, settingsOpts, logger);
+			_logFileService = new LogFileService(logFileCollection, null!, logBuilder, _logStorage, _clock, null!, null!, settingsOpts, tracer, logger);
         }
 
 		protected override void Dispose(bool disposing)
diff --git a/Engine/Source/Programs/Horde/Horde.Server.Tests/LogIndexingTests.cs b/Engine/Source/Programs/Horde/Horde.Server.Tests/LogIndexingTests.cs
--- a/Engine/Source/Programs/Horde/Horde.Server.Tests/LogIndexingTests.cs
+++ b/Engine/Source/Programs/Horde/Horde.Server.Tests/LogIndexingTests.cs
@@ -20,6 +20,7 @@ using Microsoft.Extensions.Logging;
 using Microsoft.VisualStudio.TestTools.UnitTesting;
 using Horde.Server.Logs.Data;
 using EpicGames.Horde.Logs;
+using OpenTelemetry.Trace;
 
 namespace Horde.Server.Tests
 {
@@ -49,7 +50,8 @@ namespace Horde.Server.Tests
 			_nullLogStorage = new NullLogStorage();
 			_logStorage = new LocalLogStorage(20, _nullLogStorage);
 			TestOptions<ServerSettings> settingsOpts = new (new ServerSettings());
-			_logFileService = new LogFileService(logFileCollection, null!, logBuilder, _logStorage, new FakeClock(), null!, null!, settingsOpts, logger);
+			Tracer tracer = TracerProvider.Default.GetTracer("LogIndexingTests");
+			_logFileService = new LogFileService(logFileCollection, null!, logBuilder, _logStorage, new FakeClock(), null!, null!, settingsOpts, tracer, logger);
 		}
 
 		protected override void Dispose(bool disposing)
diff --git a/Engine/Source/Programs/Horde/Horde.Server.Tests/TestSetup.cs b/Engine/Source/Programs/Horde/Horde.Server.Tests/TestSetup.cs
--- a/Engine/Source/Programs/Horde/Horde.Server.Tests/TestSetup.cs
+++ b/Engine/Source/Programs/Horde/Horde.Server.Tests/TestSetup.cs
@@ -132,6 +132,7 @@ namespace Horde.Server.Tests
 		public DashboardController DashboardController => GetDashboardController();
 		public TestDataController TestDataController => GetTestDataController();
 
+		public OpenTelemetry.Trace.Tracer Tracer => ServiceProvider.GetRequiredService<OpenTelemetry.Trace.Tracer>();
 		public ConfigService ConfigService => ServiceProvider.GetRequiredService<ConfigService>();
 		public IOptionsMonitor<GlobalConfig> GlobalConfig => ServiceProvider.GetRequiredService<IOptionsMonitor<GlobalConfig>>();
 		public IOptionsSnapshot<GlobalConfig> GlobalConfigSnapshot => ServiceProvider.GetRequiredService<IOptionsSnapshot<GlobalConfig>>();
@@ -326,7 +327,7 @@ namespace Horde.Server.Tests
 		
 		private LeasesController GetLeasesController()
 		{
-			LeasesController controller = new LeasesController(AgentService, GlobalConfigSnapshot);
+			LeasesController controller = new LeasesController(AgentService, GlobalConfigSnapshot, Tracer);
 			controller.ControllerContext = GetControllerContext();
 			return controller;
 		}
diff --git a/Engine/Source/Programs/Horde/Horde.Server/Agents/Fleet/ComputeQueueAwsMetricStrategy.cs b/Engine/Source/Programs/Horde/Horde.Server/Agents/Fleet/ComputeQueueAwsMetricStrategy.cs
--- a/Engine/Source/Programs/Horde/Horde.Server/Agents/Fleet/ComputeQueueAwsMetricStrategy.cs
+++ b/Engine/Source/Programs/Horde/Horde.Server/Agents/Fleet/ComputeQueueAwsMetricStrategy.cs
@@ -10,9 +10,9 @@ using Amazon.CloudWatch.Model;
 using EpicGames.Horde.Compute;
 using Horde.Server.Agents.Pools;
 using Horde.Server.Compute;
+using Horde.Server.Utilities;
 using Microsoft.Extensions.Logging;
-using OpenTracing;
-using OpenTracing.Util;
+using OpenTelemetry.Trace;
 
 namespace Horde.Server.Agents.Fleet
 {
@@ -86,7 +86,9 @@ namespace Horde.Server.Agents.Fleet
 		/// <inheritdoc/>
 		public async Task<PoolSizeResult> CalculatePoolSizeAsync(IPool pool, List<IAgent> agents)
 		{
-			using IScope scope = GlobalTracer.Instance.BuildSpan("ComputeQueueAwsMetricStrategy.CalcDesiredPoolSizesAsync").StartActive();
+			using TelemetrySpan span = OpenTelemetryTracers.Horde.StartActiveSpan($"{nameof(ComputeQueueAwsMetricStrategy)}.{nameof(CalculatePoolSizeAsync)}");
+			span.SetAttribute(OpenTelemetryTracers.DatadogResourceAttribute, pool.Id.ToString());
+			span.SetAttribute("currentAgentCount", agents.Count);
 
 			Dictionary<string, List<MetricDatum>> metricsPerCloudWatchNamespace = new();
 			int numAgents = agents.Count;
@@ -128,12 +130,12 @@ namespace Horde.Server.Agents.Fleet
 
 			foreach ((string ns, List<MetricDatum> metricDatumsNs) in metricsPerCloudWatchNamespace)
 			{
-				using IScope cwScope = GlobalTracer.Instance.BuildSpan("Putting CloudWatch metrics").StartActive();
-				cwScope.Span.SetTag("namespace", ns);
+				using TelemetrySpan cwSpan = OpenTelemetryTracers.Horde.StartActiveSpan($"{nameof(ComputeQueueAwsMetricStrategy)}.{nameof(CalculatePoolSizeAsync)}.PutCloudWatchMetrics");
+				cwSpan.SetAttribute("namespace", ns);
 
 				PutMetricDataRequest request = new() { Namespace = ns, MetricData = metricDatumsNs };
 				PutMetricDataResponse response = await _cloudWatch.PutMetricDataAsync(request);
-				scope.Span.SetTag("res.statusCode", (int)response.HttpStatusCode);
+				cwSpan.SetAttribute("res.statusCode", (int)response.HttpStatusCode);
 				if (response.HttpStatusCode != HttpStatusCode.OK)
 				{
 					_logger.LogError("Unable to put CloudWatch metrics");
diff --git a/Engine/Source/Programs/Horde/Horde.Server/Agents/Fleet/FleetManagerFactory.cs b/Engine/Source/Programs/Horde/Horde.Server/Agents/Fleet/FleetManagerFactory.cs
--- a/Engine/Source/Programs/Horde/Horde.Server/Agents/Fleet/FleetManagerFactory.cs
+++ b/Engine/Source/Programs/Horde/Horde.Server/Agents/Fleet/FleetManagerFactory.cs
@@ -9,6 +9,7 @@ using HordeCommon;
 using Microsoft.Extensions.DependencyInjection;
 using Microsoft.Extensions.Logging;
 using Microsoft.Extensions.Options;
+using OpenTelemetry.Trace;
 using StatsdClient;
 
 namespace Horde.Server.Agents.Fleet;
@@ -38,6 +39,7 @@ public sealed class FleetManagerFactory : IFleetManagerFactory
 	private readonly IDogStatsd _dogStatsd;
 	private readonly IServiceProvider _provider;
 	private readonly IOptionsMonitor<ServerSettings> _settings;
+	private readonly Tracer _tracer;
 	private readonly ILoggerFactory _loggerFactory;
 	
 	private IAmazonAutoScaling? _awsAutoScaling;
@@ -46,13 +48,14 @@ public sealed class FleetManagerFactory : IFleetManagerFactory
 	/// <summary>
 	/// Constructor
 	/// </summary>
-	public FleetManagerFactory(IAgentCollection agentCollection, IClock clock, IDogStatsd dogStatsd, IServiceProvider provider, IOptionsMonitor<ServerSettings> settings, ILoggerFactory loggerFactory)
+	public FleetManagerFactory(IAgentCollection agentCollection, IClock clock, IDogStatsd dogStatsd, IServiceProvider provider, IOptionsMonitor<ServerSettings> settings, Tracer tracer, ILoggerFactory loggerFactory)
 	{
 		_agentCollection = agentCollection;
 		_clock = clock;
 		_dogStatsd = dogStatsd;
 		_provider = provider;
 		_settings = settings;
+		_tracer = tracer;
 		_loggerFactory = loggerFactory;
 	}
 	
@@ -66,13 +69,13 @@ public sealed class FleetManagerFactory : IFleetManagerFactory
 			FleetManagerType.NoOp =>
 				new NoOpFleetManager(_loggerFactory.CreateLogger<NoOpFleetManager>()),
 			FleetManagerType.Aws =>
-				new AwsFleetManager(GetAwsEc2(type), _agentCollection, DeserializeSettings<AwsFleetManagerSettings>(config), _loggerFactory.CreateLogger<AwsFleetManager>()),
+				new AwsFleetManager(GetAwsEc2(type), _agentCollection, DeserializeSettings<AwsFleetManagerSettings>(config), _tracer, _loggerFactory.CreateLogger<AwsFleetManager>()),
 			FleetManagerType.AwsReuse =>
-				new AwsReuseFleetManager(GetAwsEc2(type), _agentCollection, DeserializeSettings<AwsReuseFleetManagerSettings>(config), _loggerFactory.CreateLogger<AwsReuseFleetManager>()),
+				new AwsReuseFleetManager(GetAwsEc2(type), _agentCollection, DeserializeSettings<AwsReuseFleetManagerSettings>(config), _tracer, _loggerFactory.CreateLogger<AwsReuseFleetManager>()),
 			FleetManagerType.AwsRecycle =>
-				new AwsRecyclingFleetManager(GetAwsEc2(type), _agentCollection, _dogStatsd, _clock, DeserializeSettings<AwsRecyclingFleetManagerSettings>(config), _loggerFactory.CreateLogger<AwsRecyclingFleetManager>()),
+				new AwsRecyclingFleetManager(GetAwsEc2(type), _agentCollection, _dogStatsd, _clock, DeserializeSettings<AwsRecyclingFleetManagerSettings>(config), _tracer, _loggerFactory.CreateLogger<AwsRecyclingFleetManager>()),
 			FleetManagerType.AwsAsg =>
-				new AwsAsgFleetManager(GetAwsAutoScaling(type), DeserializeSettings<AwsAsgSettings>(config), _loggerFactory.CreateLogger<AwsAsgFleetManager>()),
+				new AwsAsgFleetManager(GetAwsAutoScaling(type), DeserializeSettings<AwsAsgSettings>(config), _tracer, _loggerFactory.CreateLogger<AwsAsgFleetManager>()),
 			_ => throw new ArgumentException("Unknown fleet manager type " + type)
 		};
 	}
diff --git a/Engine/Source/Programs/Horde/Horde.Server/Agents/Fleet/FleetService.cs b/Engine/Source/Programs/Horde/Horde.Server/Agents/Fleet/FleetService.cs
--- a/Engine/Source/Programs/Horde/Horde.Server/Agents/Fleet/FleetService.cs
+++ b/Engine/Source/Programs/Horde/Horde.Server/Agents/Fleet/FleetService.cs
@@ -20,8 +20,7 @@ using Microsoft.Extensions.Caching.Memory;
 using Microsoft.Extensions.Hosting;
 using Microsoft.Extensions.Logging;
 using Microsoft.Extensions.Options;
-using OpenTracing;
-using OpenTracing.Util;
+using OpenTelemetry.Trace;
 using StatsdClient;
 
 namespace Horde.Server.Agents.Fleet
@@ -80,6 +79,7 @@ namespace Horde.Server.Agents.Fleet
 		private readonly TimeSpan _defaultScaleInCooldown;
 		private readonly IOptions<ServerSettings> _settings;
 		private readonly IOptionsMonitor<GlobalConfig> _globalConfig;
+		private readonly Tracer _tracer;
 		private readonly ILogger<FleetService> _logger;
 		
 		/// <summary>
@@ -99,6 +99,7 @@ namespace Horde.Server.Agents.Fleet
 			IMemoryCache cache,
 			IOptions<ServerSettings> settings,
 			IOptionsMonitor<GlobalConfig> globalConfig,
+			Tracer tracer,
 			ILogger<FleetService> logger)
 		{
 			_agentCollection = agentCollection;
@@ -113,6 +114,7 @@ namespace Horde.Server.Agents.Fleet
 			_clock = clock;
 			_cache = cache;
 			_globalConfig = globalConfig;
+			_tracer = tracer;
 			_logger = logger;
 			_ticker = clock.AddSharedTicker<FleetService>(TimeSpan.FromSeconds(30), TickLeaderAsync, _logger);
 			_tickerHighFrequency = clock.AddSharedTicker("FleetService.TickHighFrequency", TimeSpan.FromSeconds(30), TickHighFrequencyAsync, _logger);
@@ -142,7 +144,8 @@ namespace Horde.Server.Agents.Fleet
 
 		internal async ValueTask TickLeaderAsync(CancellationToken stoppingToken)
 		{
-			ISpan span = GlobalTracer.Instance.BuildSpan("FleetService.Tick").Start();
+			TelemetrySpan span = _tracer.StartSpan($"{nameof(FleetService)}.{nameof(TickLeaderAsync)}");
+
 			try
 			{
 				List<PoolWithAgents> poolsWithAgents = await GetPoolsWithAgentsAsync();
@@ -162,13 +165,13 @@ namespace Horde.Server.Agents.Fleet
 			}
 			finally
 			{
-				span.Finish();
+				span.End();
 			}
 		}
 
 		internal async ValueTask TickHighFrequencyAsync(CancellationToken stoppingToken)
 		{
-			ISpan span = GlobalTracer.Instance.BuildSpan("FleetService.TickHighFrequency").Start();
+			TelemetrySpan span = _tracer.StartSpan($"{nameof(FleetService)}.{nameof(TickHighFrequencyAsync)}");
 			try
 			{
 				// TODO: Re-enable high frequency scaling (only used for experimental scaling of remote execution agents at the moment)
@@ -176,7 +179,7 @@ namespace Horde.Server.Agents.Fleet
 			}
 			finally
 			{
-				span.Finish();
+				span.End();
 			}
 		}
 
@@ -209,14 +212,12 @@ namespace Horde.Server.Agents.Fleet
 
 			IFleetManager fleetManager = CreateFleetManager(pool);
 			
-			using IScope scope = GlobalTracer.Instance
-				.BuildSpan("FleetService.ScalePool")
-				.WithTag(Datadog.Trace.OpenTracing.DatadogTags.ResourceName, pool.Id.ToString())
-				.WithTag("CurrentAgentCount", currentAgentCount)
-				.WithTag("DesiredAgentCount", desiredAgentCount)
-				.WithTag("DeltaAgentCount", deltaAgentCount)
-				.WithTag("FleetManager", fleetManager.GetType().Name)
-				.StartActive();
+			using TelemetrySpan span = _tracer.StartSpan($"{nameof(FleetService)}.{nameof(ScalePoolAsync)}");
+			span.SetAttribute(OpenTelemetryTracers.DatadogResourceAttribute, pool.Id.ToString());
+			span.SetAttribute("currentAgentCount", currentAgentCount);
+			span.SetAttribute("desiredAgentCount", desiredAgentCount);
+			span.SetAttribute("deltaAgentCount", deltaAgentCount);
+			span.SetAttribute("fleetManager", fleetManager.GetType().Name);
 
 			Dictionary<string, object> logScopeMetadata = new()
 			{
@@ -241,14 +242,14 @@ namespace Horde.Server.Agents.Fleet
 			TimeSpan scaleOutCooldown = pool.ScaleOutCooldown ?? _defaultScaleOutCooldown;
 			bool isScaleOutCoolingDown = pool.LastScaleUpTime != null && pool.LastScaleUpTime + scaleOutCooldown > _clock.UtcNow;
 			TimeSpan? scaleOutCooldownTimeLeft = pool.LastScaleUpTime + _defaultScaleOutCooldown - _clock.UtcNow;
-			scope.Span.SetTag("ScaleOutCooldownTimeLeftSecs", scaleOutCooldownTimeLeft?.TotalSeconds ?? -1);
+			span.SetAttribute("scaleOutCooldownTimeLeftSecs", scaleOutCooldownTimeLeft?.TotalSeconds ?? -1);
 			
 			TimeSpan scaleInCooldown = pool.ScaleInCooldown ?? _defaultScaleInCooldown;
 			bool isScaleInCoolingDown = pool.LastScaleDownTime != null && pool.LastScaleDownTime + scaleInCooldown > _clock.UtcNow;
 			TimeSpan? scaleInCooldownTimeLeft = pool.LastScaleDownTime + _defaultScaleInCooldown - _clock.UtcNow;
-			scope.Span.SetTag("ScaleOutCooldownTimeLeftSecs", scaleInCooldownTimeLeft?.TotalSeconds ?? -1);
+			span.SetAttribute("scaleInCooldownTimeLeftSecs", scaleInCooldownTimeLeft?.TotalSeconds ?? -1);
 			
-			scope.Span.SetTag("IsDowntimeActive", _downtimeService.IsDowntimeActive);
+			span.SetAttribute("isDowntimeActive", _downtimeService.IsDowntimeActive);
 			
 			ScaleResult cooldownActiveResult = new(FleetManagerOutcome.NoOp, 0, 0, "Cooldown active");
 			ScaleResult result = new (FleetManagerOutcome.NoOp, 0, 0);
@@ -299,10 +300,10 @@ namespace Horde.Server.Agents.Fleet
 					result.Outcome, result.AgentsAddedCount, result.AgentsAddedCount, result.Message);
 			}
 			
-			scope.Span.SetTag("ResultOutcome", result.Outcome.ToString());
-			scope.Span.SetTag("ResultAgentsAdded", result.AgentsAddedCount);
-			scope.Span.SetTag("ResultAgentsRemoved", result.AgentsRemovedCount);
-			scope.Span.SetTag("ResultOutcome", result.Message);
+			span.SetAttribute("resultOutcome", result.Outcome.ToString());
+			span.SetAttribute("resultAgentsAdded", result.AgentsAddedCount);
+			span.SetAttribute("resultAgentsRemoved", result.AgentsRemovedCount);
+			span.SetAttribute("resultOutcome", result.Message);
 
 			await _poolCollection.TryUpdateAsync(
 				pool,
@@ -445,7 +446,8 @@ namespace Horde.Server.Agents.Fleet
 		{
 			int numShutdownsCancelled = await CancelPendingShutdownsAsync(pool, agentsToAdd);
 			agentsToAdd -= numShutdownsCancelled;
-			GlobalTracer.Instance.ActiveSpan?.SetTag("NumShutdownsCancelled", numShutdownsCancelled);
+
+			Activity.Current?.SetTag("numShutdownsCancelled", numShutdownsCancelled);
 
 			if (agentsToAdd <= 0)
 			{
diff --git a/Engine/Source/Programs/Horde/Horde.Server/Agents/Fleet/JobQueueStrategy.cs b/Engine/Source/Programs/Horde/Horde.Server/Agents/Fleet/JobQueueStrategy.cs
--- a/Engine/Source/Programs/Horde/Horde.Server/Agents/Fleet/JobQueueStrategy.cs
+++ b/Engine/Source/Programs/Horde/Horde.Server/Agents/Fleet/JobQueueStrategy.cs
@@ -15,8 +15,7 @@ using Horde.Server.Utilities;
 using HordeCommon;
 using Microsoft.Extensions.Caching.Memory;
 using Microsoft.Extensions.Options;
-using OpenTracing;
-using OpenTracing.Util;
+using OpenTelemetry.Trace;
 
 namespace Horde.Server.Agents.Fleet
 {
@@ -172,7 +171,7 @@ namespace Horde.Server.Agents.Fleet
 
 		internal async Task<Dictionary<PoolId, int>> GetPoolQueueSizesAsync(DateTimeOffset jobsCreatedAfter)
 		{
-			using IScope scope = GlobalTracer.Instance.BuildSpan("JobQueueStrategy.GetPoolQueueSizes").StartActive();
+			using TelemetrySpan span = OpenTelemetryTracers.Horde.StartActiveSpan($"{nameof(JobQueueStrategy)}.{nameof(GetPoolQueueSizesAsync)}");
 
 			Dictionary<StreamId, StreamConfig> streams = _globalConfig.CurrentValue.Streams.ToDictionary(x => x.Id, x => (StreamConfig)x);
 			List<IJob> recentJobs = await _jobs.FindAsync(minCreateTime: jobsCreatedAfter);
@@ -185,18 +184,16 @@ namespace Horde.Server.Agents.Fleet
 
 			List<(PoolId PoolId, int QueueSize)> poolsWithQueueSize = jobBatches.GroupBy(t => t.PoolId).Select(t => (t.Key, t.Count())).ToList();
 
-			scope.Span.SetTag("NumPools", poolsWithQueueSize.Count);
+			span.SetAttribute("numPools", poolsWithQueueSize.Count);
 			return poolsWithQueueSize.ToDictionary(x => x.PoolId, x => x.QueueSize);
 		}
 
 		/// <inheritdoc/>
 		public async Task<PoolSizeResult> CalculatePoolSizeAsync(IPool pool, List<IAgent> agents)
 		{
-			using IScope scope = GlobalTracer.Instance
-				.BuildSpan("JobQueueStrategy.CalculatePoolSize")
-				.WithTag(Datadog.Trace.OpenTracing.DatadogTags.ResourceName, pool.Id.ToString())
-				.WithTag("CurrentAgentCount", agents.Count)
-				.StartActive();
+			using TelemetrySpan span = OpenTelemetryTracers.Horde.StartActiveSpan($"{nameof(JobQueueStrategy)}.{nameof(CalculatePoolSizeAsync)}");
+			span.SetAttribute(OpenTelemetryTracers.DatadogResourceAttribute, pool.Id.ToString());
+			span.SetAttribute("currentAgentCount", agents.Count);
 			
 			DateTimeOffset minCreateTime = _clock.UtcNow - TimeSpan.FromMinutes(Settings.SamplePeriodMin);
 
diff --git a/Engine/Source/Programs/Horde/Horde.Server/Agents/Fleet/LeaseUtilizationStrategy.cs b/Engine/Source/Programs/Horde/Horde.Server/Agents/Fleet/LeaseUtilizationStrategy.cs
--- a/Engine/Source/Programs/Horde/Horde.Server/Agents/Fleet/LeaseUtilizationStrategy.cs
+++ b/Engine/Source/Programs/Horde/Horde.Server/Agents/Fleet/LeaseUtilizationStrategy.cs
@@ -12,8 +12,7 @@ using Horde.Server.Utilities;
 using HordeCommon;
 using HordeCommon.Rpc.Tasks;
 using Microsoft.Extensions.Caching.Memory;
-using OpenTracing;
-using OpenTracing.Util;
+using OpenTelemetry.Trace;
 
 namespace Horde.Server.Agents.Fleet
 {
@@ -152,7 +151,7 @@ namespace Horde.Server.Agents.Fleet
 
 		private async Task<Dictionary<AgentId, AgentData>> GetAgentDataAsync()
 		{
-			using IScope scope = GlobalTracer.Instance.BuildSpan("LeaseUtilizationStrategy.GetAgentDataAsync").StartActive();
+			using TelemetrySpan span = OpenTelemetryTracers.Horde.StartActiveSpan($"{nameof(LeaseUtilizationStrategy)}.{nameof(GetAgentDataAsync)}");
 			
 			// Find all the current agents
 			List<IAgent> agents = await _agentCollection.FindAsync(status: AgentStatus.Ok);
@@ -195,13 +194,13 @@ namespace Horde.Server.Agents.Fleet
 				}
 			}
 
-			scope.Span.SetTag("AgentDataCount", agentIdToData.Count);
+			span.SetAttribute("agentDataCount", agentIdToData.Count);
 			return agentIdToData;
 		}
 		
 		private async Task<Dictionary<PoolId, PoolData>> GetPoolDataAsync()
 		{
-			using IScope scope = GlobalTracer.Instance.BuildSpan("LeaseUtilizationStrategy.GetPoolDataAsync").StartActive();
+			using TelemetrySpan span = OpenTelemetryTracers.Horde.StartActiveSpan($"{nameof(LeaseUtilizationStrategy)}.{nameof(GetPoolDataAsync)}");
 
 			Dictionary<AgentId, AgentData> agentIdToData = await GetAgentDataAsync();
 			
@@ -222,7 +221,7 @@ namespace Horde.Server.Agents.Fleet
 				}
 			}
 
-			scope.Span.SetTag("PoolDataCount", agentIdToData.Count);
+			span.SetAttribute("poolDataCount", agentIdToData.Count);
 			return poolToData;
 		}
 
@@ -232,11 +231,9 @@ namespace Horde.Server.Agents.Fleet
 		/// <inheritdoc/>
 		public async Task<PoolSizeResult> CalculatePoolSizeAsync(IPool pool, List<IAgent> agents)
 		{
-			using IScope scope = GlobalTracer.Instance
-				.BuildSpan("LeaseUtilizationStrategy.CalculatePoolSize")
-				.WithTag(Datadog.Trace.OpenTracing.DatadogTags.ResourceName, pool.Id.ToString())
-				.WithTag("CurrentAgentCount", agents.Count)
-				.StartActive();
+			using TelemetrySpan span = OpenTelemetryTracers.Horde.StartActiveSpan($"{nameof(LeaseUtilizationStrategy)}.{nameof(CalculatePoolSizeAsync)}");
+			span.SetAttribute(OpenTelemetryTracers.DatadogResourceAttribute, pool.Id.ToString());
+			span.SetAttribute("currentAgentCount", agents.Count);
 			
 			Dictionary<PoolId, PoolData> poolToData;
 			
diff --git a/Engine/Source/Programs/Horde/Horde.Server/Agents/Fleet/Providers/AwsAsgFleetManager.cs b/Engine/Source/Programs/Horde/Horde.Server/Agents/Fleet/Providers/AwsAsgFleetManager.cs
--- a/Engine/Source/Programs/Horde/Horde.Server/Agents/Fleet/Providers/AwsAsgFleetManager.cs
+++ b/Engine/Source/Programs/Horde/Horde.Server/Agents/Fleet/Providers/AwsAsgFleetManager.cs
@@ -9,10 +9,8 @@ using System.Threading.Tasks;
 using Amazon.AutoScaling;
 using Amazon.AutoScaling.Model;
 using Horde.Server.Agents.Pools;
-using Horde.Server.Utilities;
 using Microsoft.Extensions.Logging;
-using OpenTracing;
-using OpenTracing.Util;
+using OpenTelemetry.Trace;
 
 namespace Horde.Server.Agents.Fleet.Providers
 {
@@ -46,15 +44,17 @@ namespace Horde.Server.Agents.Fleet.Providers
 	{
 		internal AwsAsgSettings Settings { get; }
 		private readonly IAmazonAutoScaling _awsAutoScaling;
+		private readonly Tracer _tracer;
 		private readonly ILogger _logger;
 
 		/// <summary>
 		/// Constructor
 		/// </summary>
-		public AwsAsgFleetManager(IAmazonAutoScaling awsAutoScaling, AwsAsgSettings asgSettings, ILogger<AwsAsgFleetManager> logger)
+		public AwsAsgFleetManager(IAmazonAutoScaling awsAutoScaling, AwsAsgSettings asgSettings, Tracer tracer, ILogger<AwsAsgFleetManager> logger)
 		{
 			_awsAutoScaling = awsAutoScaling;
 			Settings = asgSettings;
+			_tracer = tracer;
 			_logger = logger;
 		}
 
@@ -63,13 +63,13 @@ namespace Horde.Server.Agents.Fleet.Providers
 		{
 			int desiredCapacity = agents.Count + count;
 			
-			using IScope scope = GlobalTracer.Instance.BuildSpan("AwsAsgFleetManager.ExpandPool").StartActive();
-			scope.Span.SetTag("poolName", pool.Name);
-			scope.Span.SetTag("numAgents", agents.Count);
-			scope.Span.SetTag("count", count);
-			scope.Span.SetTag("desiredCapacity", desiredCapacity);
+			using TelemetrySpan span = _tracer.StartActiveSpan($"{nameof(AwsAsgFleetManager)}.{nameof(ExpandPoolAsync)}");
+			span.SetAttribute("poolName", pool.Name);
+			span.SetAttribute("numAgents", agents.Count);
+			span.SetAttribute("count", count);
+			span.SetAttribute("desiredCapacity", desiredCapacity);
 			
-			await UpdateAsg(pool.Id, desiredCapacity, scope, cancellationToken);
+			await UpdateAsg(pool.Id, desiredCapacity, span, cancellationToken);
 			return new ScaleResult(FleetManagerOutcome.Success, count, 0);
 		}
 
@@ -78,22 +78,23 @@ namespace Horde.Server.Agents.Fleet.Providers
 		{
 			int desiredCapacity = Math.Max(0, agents.Count - count);
 			
-			using IScope scope = GlobalTracer.Instance.BuildSpan("AwsAsgFleetManager.ShrinkPool").StartActive();
-			scope.Span.SetTag("poolName", pool.Name);
-			scope.Span.SetTag("numAgents", agents.Count);
-			scope.Span.SetTag("count", count);
-			scope.Span.SetTag("desiredCapacity", desiredCapacity);
 			
-			await UpdateAsg(pool.Id, desiredCapacity, scope, cancellationToken);
+			using TelemetrySpan span = _tracer.StartActiveSpan($"{nameof(AwsAsgFleetManager)}.{nameof(ShrinkPoolAsync)}");
+			span.SetAttribute("poolName", pool.Name);
+			span.SetAttribute("numAgents", agents.Count);
+			span.SetAttribute("count", count);
+			span.SetAttribute("desiredCapacity", desiredCapacity);
+			
+			await UpdateAsg(pool.Id, desiredCapacity, span, cancellationToken);
 			return new ScaleResult(FleetManagerOutcome.Success, 0, count);
 		}
 		
-		private async Task UpdateAsg(PoolId poolId, int desiredCapacity, IScope scope, CancellationToken cancellationToken)
+		private async Task UpdateAsg(PoolId poolId, int desiredCapacity, TelemetrySpan span, CancellationToken cancellationToken)
 		{
 			UpdateAutoScalingGroupRequest request = new () { AutoScalingGroupName = Settings.Name, DesiredCapacity = desiredCapacity };
 			UpdateAutoScalingGroupResponse response = await _awsAutoScaling.UpdateAutoScalingGroupAsync(request, cancellationToken);
 			
-			scope.Span.SetTag("res.statusCode", (int)response.HttpStatusCode);
+			span.SetAttribute("res.statusCode", (int)response.HttpStatusCode);
 			if (response.HttpStatusCode != HttpStatusCode.OK)
 			{
 				_logger.LogError("Unable to update auto-scaling group {AutoScalingGroup} for pool {PoolId}.", Settings.Name, poolId);
diff --git a/Engine/Source/Programs/Horde/Horde.Server/Agents/Fleet/Providers/AwsFleetManager.cs b/Engine/Source/Programs/Horde/Horde.Server/Agents/Fleet/Providers/AwsFleetManager.cs
--- a/Engine/Source/Programs/Horde/Horde.Server/Agents/Fleet/Providers/AwsFleetManager.cs
+++ b/Engine/Source/Programs/Horde/Horde.Server/Agents/Fleet/Providers/AwsFleetManager.cs
@@ -9,9 +9,9 @@ using Amazon.EC2;
 using Amazon.EC2.Model;
 using Horde.Server.Agents.Pools;
 using Horde.Server.Auditing;
+using Horde.Server.Utilities;
 using Microsoft.Extensions.Logging;
-using OpenTracing;
-using OpenTracing.Util;
+using OpenTelemetry.Trace;
 
 namespace Horde.Server.Agents.Fleet.Providers
 {
@@ -139,26 +139,28 @@ namespace Horde.Server.Agents.Fleet.Providers
 		
 		private readonly IAmazonEC2 _ec2;
 		private readonly IAgentCollection _agentCollection;
+		private readonly Tracer _tracer;
 		private readonly ILogger _logger;
 
 		/// <summary>
 		/// Constructor
 		/// </summary>
-		public AwsFleetManager(IAmazonEC2 ec2, IAgentCollection agentCollection, AwsFleetManagerSettings settings, ILogger<AwsFleetManager> logger)
+		public AwsFleetManager(IAmazonEC2 ec2, IAgentCollection agentCollection, AwsFleetManagerSettings settings, Tracer tracer, ILogger<AwsFleetManager> logger)
 		{
 			_ec2 = ec2;
 			_agentCollection = agentCollection;
 			Settings = settings;
+			_tracer = tracer;
 			_logger = logger;
 		}
 
 		/// <inheritdoc/>
 		public async Task<ScaleResult> ExpandPoolAsync(IPool pool, IReadOnlyList<IAgent> agents, int count, CancellationToken cancellationToken)
 		{
-			using IScope scope = GlobalTracer.Instance.BuildSpan("ExpandPool").StartActive();
-			scope.Span.SetTag("poolName", pool.Name);
-			scope.Span.SetTag("numAgents", agents.Count);
-			scope.Span.SetTag("count", count);
+			using TelemetrySpan span = _tracer.StartActiveSpan($"{nameof(AwsFleetManager)}.{nameof(ExpandPoolAsync)}");
+			span.SetAttribute("poolName", pool.Name);
+			span.SetAttribute("numAgents", agents.Count);
+			span.SetAttribute("count", count);
 
 			List<Tag> tags = Settings.Tags.Select(x => new Tag(x.Key, x.Value)).ToList();
 			tags.Add(new Tag("Horde:RequestedPools", pool.Name));
@@ -197,8 +199,8 @@ namespace Horde.Server.Agents.Fleet.Providers
 
 			RunInstancesResponse response = await _ec2.RunInstancesAsync(request, cancellationToken);
 			int numStartedInstances = response.Reservation.Instances.Count;
-			scope.Span.SetTag("res.statusCode", (int)response.HttpStatusCode);
-			scope.Span.SetTag("res.numInstances", numStartedInstances);
+			span.SetAttribute("res.statusCode", (int)response.HttpStatusCode);
+			span.SetAttribute("res.numInstances", numStartedInstances);
 
 			foreach (Instance i in response.Reservation.Instances)
 			{
@@ -238,9 +240,10 @@ namespace Horde.Server.Agents.Fleet.Providers
 		{
 			_ = cancellationToken;
 
-			using IScope scope = GlobalTracer.Instance.BuildSpan("ShrinkPool").StartActive();
-			scope.Span.SetTag("poolName", pool.Name);
-			scope.Span.SetTag("count", count);
+			using TelemetrySpan span = OpenTelemetryTracers.Horde.StartActiveSpan($"{nameof(AwsFleetManager)}.{nameof(ShrinkPoolViaAgentShutdownRequestAsync)}");
+			span.SetAttribute("poolName", pool.Name);
+			span.SetAttribute("numAgents", agents.Count);
+			span.SetAttribute("count", count);
 			
 			string awsTagProperty = $"{AwsTagPropertyName}={PoolTagName}:{pool.Name}";
 			
@@ -249,10 +252,10 @@ namespace Horde.Server.Agents.Fleet.Providers
 			List<IAgent> agentsWithAwsTags = filteredAgents.Where(x => x.HasProperty(awsTagProperty)).ToList(); 
 			List<IAgent> agentsLimitedByCount = agentsWithAwsTags.Take(count).ToList();
 			
-			scope.Span.SetTag("agents.num", agents.Count);
-			scope.Span.SetTag("agents.filtered.num", filteredAgents.Count);
-			scope.Span.SetTag("agents.withAwsTags.num", agentsWithAwsTags.Count);
-			scope.Span.SetTag("agents.limitedByCount.num", agentsLimitedByCount.Count);
+			span.SetAttribute("agents.num", agents.Count);
+			span.SetAttribute("agents.filtered.num", filteredAgents.Count);
+			span.SetAttribute("agents.withAwsTags.num", agentsWithAwsTags.Count);
+			span.SetAttribute("agents.limitedByCount.num", agentsLimitedByCount.Count);
 
 			foreach (IAgent agent in agentsLimitedByCount)
 			{
diff --git a/Engine/Source/Programs/Horde/Horde.Server/Agents/Fleet/Providers/AwsRecyclingFleetManager.cs b/Engine/Source/Programs/Horde/Horde.Server/Agents/Fleet/Providers/AwsRecyclingFleetManager.cs
--- a/Engine/Source/Programs/Horde/Horde.Server/Agents/Fleet/Providers/AwsRecyclingFleetManager.cs
+++ b/Engine/Source/Programs/Horde/Horde.Server/Agents/Fleet/Providers/AwsRecyclingFleetManager.cs
@@ -13,8 +13,7 @@ using EpicGames.Core;
 using Horde.Server.Agents.Pools;
 using HordeCommon;
 using Microsoft.Extensions.Logging;
-using OpenTracing;
-using OpenTracing.Util;
+using OpenTelemetry.Trace;
 using StatsdClient;
 
 namespace Horde.Server.Agents.Fleet.Providers;
@@ -93,17 +92,19 @@ public sealed class AwsRecyclingFleetManager : IFleetManager
 	private readonly IAgentCollection _agentCollection;
 	private readonly IDogStatsd _dogStatsd;
 	private readonly IClock _clock;
+	private readonly Tracer _tracer;
 	private readonly ILogger _logger;
 
 	/// <summary>
 	/// Constructor
 	/// </summary>
-	public AwsRecyclingFleetManager(IAmazonEC2 ec2, IAgentCollection agentCollection, IDogStatsd dogStatsd, IClock clock, AwsRecyclingFleetManagerSettings settings, ILogger<AwsRecyclingFleetManager> logger)
+	public AwsRecyclingFleetManager(IAmazonEC2 ec2, IAgentCollection agentCollection, IDogStatsd dogStatsd, IClock clock, AwsRecyclingFleetManagerSettings settings, Tracer tracer, ILogger<AwsRecyclingFleetManager> logger)
 	{
 		_ec2 = ec2;
 		_agentCollection = agentCollection;
 		_dogStatsd = dogStatsd;
 		_clock = clock;
+		_tracer = tracer;
 		_logger = logger;
 		Settings = settings;
 	}
@@ -111,11 +112,11 @@ public sealed class AwsRecyclingFleetManager : IFleetManager
 	/// <inheritdoc/>
 	public async Task<ScaleResult> ExpandPoolAsync(IPool pool, IReadOnlyList<IAgent> agents, int requestedInstancesCount, CancellationToken cancellationToken)
 	{
-		using IScope scope = GlobalTracer.Instance.BuildSpan("ExpandPool").StartActive();
-		scope.Span.SetTag("PoolId", pool.Id.ToString());
-		scope.Span.SetTag("PoolName", pool.Name);
-		scope.Span.SetTag("NumAgents", agents.Count);
-		scope.Span.SetTag("Count", requestedInstancesCount);
+		using TelemetrySpan span = _tracer.StartActiveSpan($"{nameof(AwsRecyclingFleetManager)}.{nameof(ExpandPoolAsync)}");
+		span.SetAttribute("poolId", pool.Id.ToString());
+		span.SetAttribute("poolName", pool.Name);
+		span.SetAttribute("numAgents", agents.Count);
+		span.SetAttribute("count", requestedInstancesCount);
 
 		using IDisposable logScope = _logger.BeginScope(new Dictionary<string, object> { ["PoolId"] = pool.Id });
 
@@ -127,12 +128,12 @@ public sealed class AwsRecyclingFleetManager : IFleetManager
 		List<InstanceType>? instanceTypePriority = Settings.InstanceTypes?.Select(InstanceType.FindValue).ToList();
 		int instancesToStartCount = requestCountPerAz.Values.Sum(x => x);
 		
-		scope.Span.SetTag("InstancesToStartCount", instancesToStartCount);
-		scope.Span.SetTag("RequestedInstancesCount", requestedInstancesCount);
-		scope.Span.SetTag("CurrentAgentCount", agents.Count);
-		scope.Span.SetTag("StoppedInstancesMissingCount", stoppedInstancesMissingCount);
-		scope.Span.SetTag("InstancesToStartCount", instancesToStartCount);
-		scope.Span.SetTag("StoppedInstancesMissingCount", stoppedInstancesMissingCount);
+		span.SetAttribute("InstancesToStartCount", instancesToStartCount);
+		span.SetAttribute("RequestedInstancesCount", requestedInstancesCount);
+		span.SetAttribute("CurrentAgentCount", agents.Count);
+		span.SetAttribute("StoppedInstancesMissingCount", stoppedInstancesMissingCount);
+		span.SetAttribute("InstancesToStartCount", instancesToStartCount);
+		span.SetAttribute("StoppedInstancesMissingCount", stoppedInstancesMissingCount);
 
 		if (instancesToStartCount == 0)
 		{
@@ -238,10 +239,9 @@ public sealed class AwsRecyclingFleetManager : IFleetManager
 	{
 		List<string> instanceIds = instances.Select(x => x.InstanceId).ToList();
 		string instanceTypePriorityString = instanceTypePriority == null ? "" : String.Join(',', instanceTypePriority.Select(x => x.Value));
-		using IScope scope = GlobalTracer.Instance.BuildSpan("StartInstancesWithRetries")
-			.WithTag("Instances", String.Join(',', instanceIds))
-			.WithTag("InstanceTypePriority", instanceTypePriorityString)
-			.StartActive();
+		using TelemetrySpan span = _tracer.StartActiveSpan($"{nameof(AwsRecyclingFleetManager)}.{nameof(StartInstancesWithRetriesAsync)}");
+		span.SetAttribute("instances", String.Join(',', instanceIds));
+		span.SetAttribute("instanceTypePriority", instanceTypePriorityString);
 		
 		if (instances.Count == 0) return;
 		
@@ -284,7 +284,7 @@ public sealed class AwsRecyclingFleetManager : IFleetManager
 			}
 		}
 
-		scope.Span.SetTag("Success", success);
+		span.SetAttribute("success", success);
 		if (!success)
 		{
 			_logger.LogError("Unable to start instances. Insufficient capacity for all instance types tried");
@@ -294,10 +294,9 @@ public sealed class AwsRecyclingFleetManager : IFleetManager
 	internal async Task StartInstancesAsync(List<Instance> instances, InstanceType? instanceType, CancellationToken cancellationToken)
 	{
 		List<string> instanceIds = instances.Select(x => x.InstanceId).ToList();
-		using IScope scope = GlobalTracer.Instance.BuildSpan("StartInstances")
-			.WithTag("Instances", String.Join(',', instanceIds))
-			.WithTag("InstanceType", instanceType)
-			.StartActive();
+		using TelemetrySpan span = _tracer.StartActiveSpan($"{nameof(AwsRecyclingFleetManager)}.{nameof(StartInstancesAsync)}");
+		span.SetAttribute("instances", String.Join(',', instanceIds));
+		span.SetAttribute("instanceType", instanceType);
 		
 		if (instances.Count == 0) return;
 		if (instanceType != null)
@@ -311,11 +310,11 @@ public sealed class AwsRecyclingFleetManager : IFleetManager
 		StartInstancesResponse startResponse;
 		try
 		{
-			using IScope ec2Scope = GlobalTracer.Instance.BuildSpan("StartInstancesEc2").StartActive();
-			ec2Scope.Span.SetTag("Request.InstanceIds", String.Join(",", startRequest.InstanceIds));
+			using TelemetrySpan ec2Span = _tracer.StartActiveSpan("StartInstancesEc2");
+			span.SetAttribute("req.instanceIds", String.Join(',', startRequest.InstanceIds));
 			startResponse = await _ec2.StartInstancesAsync(startRequest, cancellationToken);
-			ec2Scope.Span.SetTag("Response.StatusCode", (int)startResponse.HttpStatusCode);
-			ec2Scope.Span.SetTag("Response.NumInstances", startResponse.StartingInstances.Count);
+			span.SetAttribute("res.statusCode", (int)startResponse.HttpStatusCode);
+			span.SetAttribute("res.numInstances", startResponse.StartingInstances.Count);
 		}
 		catch (AmazonEC2Exception e)
 		{
@@ -355,16 +354,15 @@ public sealed class AwsRecyclingFleetManager : IFleetManager
 	{
 		foreach (Instance instance in instances)
 		{
-			using IScope scope = GlobalTracer.Instance.BuildSpan("ChangeInstanceType")
-				.WithTag("FromInstanceType", instance.InstanceType)
-				.WithTag("ToInstanceType", newInstanceType)
-				.StartActive();
+			using TelemetrySpan span = _tracer.StartActiveSpan($"{nameof(AwsRecyclingFleetManager)}.{nameof(ChangeInstanceTypeAsync)}");
+			span.SetAttribute("fromInstanceType", instance.InstanceType);
+			span.SetAttribute("toInstanceType", newInstanceType);
 			
 			if (instance.InstanceType == newInstanceType) { continue; }
 
 			ModifyInstanceAttributeRequest request = new () { InstanceId = instance.InstanceId, InstanceType = newInstanceType };
 			ModifyInstanceAttributeResponse response = await _ec2.ModifyInstanceAttributeAsync(request, cancellationToken);
-			scope.Span.SetTag("Response.StatusCode", (int)response.HttpStatusCode);
+			span.SetAttribute("Response.StatusCode", (int)response.HttpStatusCode);
 			
 			if (IsResponseSuccessful(response))
 			{
@@ -387,8 +385,8 @@ public sealed class AwsRecyclingFleetManager : IFleetManager
 	{
 		TimeSpan pendingStateTimeout = TimeSpan.FromMinutes(5);
 		
-		using IScope scope = GlobalTracer.Instance.BuildSpan("StopStuckPendingInstances")
-			.WithTag("PoolId", pool.Id.ToString()).StartActive();
+		using TelemetrySpan span = _tracer.StartActiveSpan($"{nameof(AwsRecyclingFleetManager)}.{nameof(StopStuckPendingInstancesAsync)}");
+		span.SetAttribute("poolId", pool.Id.ToString());
 		
 		// Find pending instances in the correct pool
 		DescribeInstancesRequest request = new()
@@ -400,8 +398,8 @@ public sealed class AwsRecyclingFleetManager : IFleetManager
 			}
 		};
 		DescribeInstancesResponse response = await _ec2.DescribeInstancesAsync(request, cancellationToken);
-		scope.Span.SetTag("res.statusCode", (int)response.HttpStatusCode);
-		scope.Span.SetTag("res.numReservations", response.Reservations.Count);
+		span.SetAttribute("res.statusCode", (int)response.HttpStatusCode);
+		span.SetAttribute("res.numReservations", response.Reservations.Count);
 
 		List<string> stuckInstanceIds = response.Reservations
 			.SelectMany(x => x.Instances)
@@ -413,8 +411,8 @@ public sealed class AwsRecyclingFleetManager : IFleetManager
 		{
 			StopInstancesRequest stopRequest = new() { InstanceIds = stuckInstanceIds };
 			StopInstancesResponse stopResponse = await _ec2.StopInstancesAsync(stopRequest, cancellationToken);
-			scope.Span.SetTag("stopRes.statusCode", (int)stopResponse.HttpStatusCode);
-			scope.Span.SetTag("stopRes.numStoppingInstances", stopResponse.StoppingInstances.Count);
+			span.SetAttribute("stopRes.statusCode", (int)stopResponse.HttpStatusCode);
+			span.SetAttribute("stopRes.numStoppingInstances", stopResponse.StoppingInstances.Count);
 		}
 	}
 	
@@ -426,8 +424,8 @@ public sealed class AwsRecyclingFleetManager : IFleetManager
 	/// <returns>List of instances grouped by availability zone</returns>
 	private async Task<Dictionary<string, List<Instance>>> GetCandidateInstancesAsync(IPool pool, CancellationToken cancellationToken)
 	{
-		using IScope scope = GlobalTracer.Instance.BuildSpan("GetCandidateInstances")
-			.WithTag("PoolId", pool.Id.ToString()).StartActive();
+		using TelemetrySpan span = _tracer.StartActiveSpan($"{nameof(AwsRecyclingFleetManager)}.{nameof(GetCandidateInstancesAsync)}");
+		span.SetAttribute("poolId", pool.Id.ToString());
 		
 		// Find stopped instances in the correct pool
 		DescribeInstancesRequest request = new()
@@ -439,8 +437,8 @@ public sealed class AwsRecyclingFleetManager : IFleetManager
 			}
 		};
 		DescribeInstancesResponse response = await _ec2.DescribeInstancesAsync(request, cancellationToken);
-		scope.Span.SetTag("res.statusCode", (int)response.HttpStatusCode);
-		scope.Span.SetTag("res.numReservations", response.Reservations.Count);
+		span.SetAttribute("res.statusCode", (int)response.HttpStatusCode);
+		span.SetAttribute("res.numReservations", response.Reservations.Count);
 
 		// Group instance objects by availability zone
 		List<IGrouping<string, Instance>> azInstancePairs = response.Reservations
diff --git a/Engine/Source/Programs/Horde/Horde.Server/Agents/Fleet/Providers/AwsReuseFleetManager.cs b/Engine/Source/Programs/Horde/Horde.Server/Agents/Fleet/Providers/AwsReuseFleetManager.cs
--- a/Engine/Source/Programs/Horde/Horde.Server/Agents/Fleet/Providers/AwsReuseFleetManager.cs
+++ b/Engine/Source/Programs/Horde/Horde.Server/Agents/Fleet/Providers/AwsReuseFleetManager.cs
@@ -9,8 +9,7 @@ using Amazon.EC2;
 using Amazon.EC2.Model;
 using Horde.Server.Agents.Pools;
 using Microsoft.Extensions.Logging;
-using OpenTracing;
-using OpenTracing.Util;
+using OpenTelemetry.Trace;
 
 namespace Horde.Server.Agents.Fleet.Providers
 {
@@ -49,15 +48,17 @@ namespace Horde.Server.Agents.Fleet.Providers
 		
 		private readonly IAmazonEC2 _ec2;
 		private readonly IAgentCollection _agentCollection;
+		private readonly Tracer _tracer;
 		private readonly ILogger _logger;
 
 		/// <summary>
 		/// Constructor
 		/// </summary>
-		public AwsReuseFleetManager(IAmazonEC2 ec2, IAgentCollection agentCollection, AwsReuseFleetManagerSettings settings, ILogger<AwsReuseFleetManager> logger)
+		public AwsReuseFleetManager(IAmazonEC2 ec2, IAgentCollection agentCollection, AwsReuseFleetManagerSettings settings, Tracer tracer, ILogger<AwsReuseFleetManager> logger)
 		{
 			_ec2 = ec2;
 			_agentCollection = agentCollection;
+			_tracer = tracer;
 			_logger = logger;
 			Settings = settings;
 		}
@@ -65,13 +66,13 @@ namespace Horde.Server.Agents.Fleet.Providers
 		/// <inheritdoc/>
 		public async Task<ScaleResult> ExpandPoolAsync(IPool pool, IReadOnlyList<IAgent> agents, int requestedInstancesCount, CancellationToken cancellationToken)
 		{
-			using IScope scope = GlobalTracer.Instance.BuildSpan("ExpandPool").StartActive();
-			scope.Span.SetTag("poolName", pool.Name);
-			scope.Span.SetTag("numAgents", agents.Count);
-			scope.Span.SetTag("count", requestedInstancesCount);
+			using TelemetrySpan span = _tracer.StartActiveSpan($"{nameof(AwsReuseFleetManager)}.{nameof(ExpandPoolAsync)}");
+			span.SetAttribute("poolName", pool.Name);
+			span.SetAttribute("numAgents", agents.Count);
+			span.SetAttribute("count", requestedInstancesCount);
 
 			DescribeInstancesResponse describeResponse;
-			using (IScope describeScope = GlobalTracer.Instance.BuildSpan("DescribeInstances").StartActive())
+			using (TelemetrySpan describeSpan = _tracer.StartActiveSpan("DescribeInstances"))
 			{
 				// Find stopped instances in the correct pool
 				DescribeInstancesRequest describeRequest = new ();
@@ -79,8 +80,8 @@ namespace Horde.Server.Agents.Fleet.Providers
 				describeRequest.Filters.Add(new Filter("instance-state-name", new List<string> { InstanceStateName.Stopped.Value }));
 				describeRequest.Filters.Add(new Filter("tag:" + AwsFleetManager.PoolTagName, new List<string> { pool.Name }));
 				describeResponse = await _ec2.DescribeInstancesAsync(describeRequest, cancellationToken);
-				describeScope.Span.SetTag("res.statusCode", (int)describeResponse.HttpStatusCode);
-				describeScope.Span.SetTag("res.numReservations", describeResponse.Reservations.Count);
+				describeSpan.SetAttribute("res.statusCode", (int)describeResponse.HttpStatusCode);
+				describeSpan.SetAttribute("res.numReservations", describeResponse.Reservations.Count);
 			}
 
 			IEnumerable<Instance> stoppedInstances = describeResponse.Reservations.SelectMany(x => x.Instances);
@@ -91,7 +92,7 @@ namespace Horde.Server.Agents.Fleet.Providers
 				foreach (Instance stoppedInstance in stoppedInstances)
 				{
 					if (stoppedInstance.InstanceType == newInstanceType) { continue; }
-					using IScope modifyScope = GlobalTracer.Instance.BuildSpan("ModifyInstanceAttribute").StartActive();
+					using TelemetrySpan modifySpan = _tracer.StartActiveSpan("ModifyInstanceAttribute");
 
 					ModifyInstanceAttributeRequest request = new () { InstanceId = stoppedInstance.InstanceId, InstanceType = newInstanceType };
 					ModifyInstanceAttributeResponse response = await _ec2.ModifyInstanceAttributeAsync(request, cancellationToken);
@@ -106,20 +107,20 @@ namespace Horde.Server.Agents.Fleet.Providers
 				}
 			}
 
-			using (IScope startScope = GlobalTracer.Instance.BuildSpan("StartInstances").StartActive())
+			using (TelemetrySpan startSpan = _tracer.StartActiveSpan("StartInstances"))
 			{
 				// Try to start the given instances
 				StartInstancesRequest startRequest = new ();
 				startRequest.InstanceIds.AddRange(describeResponse.Reservations.SelectMany(x => x.Instances).Select(x => x.InstanceId).Take(requestedInstancesCount));
 				int stoppedInstancesCount = startRequest.InstanceIds.Count;
 				
-				startScope.Span.SetTag("req.instanceIds", String.Join(",", startRequest.InstanceIds));
+				startSpan.SetAttribute("req.instanceIds", String.Join(",", startRequest.InstanceIds));
 				int startedInstancesCount = 0;
 				if (startRequest.InstanceIds.Count > 0)
 				{
 					StartInstancesResponse startResponse = await _ec2.StartInstancesAsync(startRequest, cancellationToken);
-					startScope.Span.SetTag("res.statusCode", (int)startResponse.HttpStatusCode);
-					startScope.Span.SetTag("res.numInstances", startResponse.StartingInstances.Count);
+					startSpan.SetAttribute("res.statusCode", (int)startResponse.HttpStatusCode);
+					startSpan.SetAttribute("res.numInstances", startResponse.StartingInstances.Count);
 					if ((int)startResponse.HttpStatusCode >= 200 && (int)startResponse.HttpStatusCode <= 299)
 					{
 						foreach (InstanceStateChange instanceChange in startResponse.StartingInstances)
diff --git a/Engine/Source/Programs/Horde/Horde.Server/Agents/Leases/LeasesController.cs b/Engine/Source/Programs/Horde/Horde.Server/Agents/Leases/LeasesController.cs
--- a/Engine/Source/Programs/Horde/Horde.Server/Agents/Leases/LeasesController.cs
+++ b/Engine/Source/Programs/Horde/Horde.Server/Agents/Leases/LeasesController.cs
@@ -11,8 +11,7 @@ using Horde.Server.Utilities;
 using Microsoft.AspNetCore.Authorization;
 using Microsoft.AspNetCore.Mvc;
 using Microsoft.Extensions.Options;
-using OpenTracing;
-using OpenTracing.Util;
+using OpenTelemetry.Trace;
 
 namespace Horde.Server.Agents.Leases
 {
@@ -26,14 +25,16 @@ namespace Horde.Server.Agents.Leases
 	{
 		readonly AgentService _agentService;
 		readonly IOptionsSnapshot<GlobalConfig> _globalConfig;
+		readonly Tracer _tracer;
 
 		/// <summary>
 		/// Constructor
 		/// </summary>
-		public LeasesController(AgentService agentService, IOptionsSnapshot<GlobalConfig> globalConfig)
+		public LeasesController(AgentService agentService, IOptionsSnapshot<GlobalConfig> globalConfig, Tracer tracer)
 		{
 			_agentService = agentService;
 			_globalConfig = globalConfig;
+			_tracer = tracer;
 		}
 
 		/// <summary>
@@ -83,7 +84,7 @@ namespace Horde.Server.Agents.Leases
 
 			List<object> responses = new List<object>();
 
-			using (IScope _ = GlobalTracer.Instance.BuildSpan($"GenerateResponses").StartActive())
+			using TelemetrySpan span = _tracer.StartActiveSpan($"{nameof(LeasesController)}.{nameof(FindLeasesAsync)}");
 			{
 				Dictionary<AgentId, double?> cachedAgentRates = new Dictionary<AgentId, double?>();
 				foreach (ILease lease in leases)
diff --git a/Engine/Source/Programs/Horde/Horde.Server/Devices/DeviceService.cs b/Engine/Source/Programs/Horde/Horde.Server/Devices/DeviceService.cs
--- a/Engine/Source/Programs/Horde/Horde.Server/Devices/DeviceService.cs
+++ b/Engine/Source/Programs/Horde/Horde.Server/Devices/DeviceService.cs
@@ -20,8 +20,7 @@ using Microsoft.Extensions.Hosting;
 using Microsoft.Extensions.Logging;
 using Microsoft.Extensions.Options;
 using MongoDB.Bson;
-using OpenTracing;
-using OpenTracing.Util;
+using OpenTelemetry.Trace;
 
 namespace Horde.Server.Devices
 {
@@ -86,6 +85,7 @@ namespace Horde.Server.Devices
 		readonly JobService _jobService;
 		readonly IStreamCollection _streamCollection;
 		readonly IUserCollection _userCollection;
+		readonly Tracer _tracer;
 		readonly ILogger<DeviceService> _logger;
 		readonly IDeviceCollection _devices;
 		readonly ITicker _ticker;
@@ -108,7 +108,7 @@ namespace Horde.Server.Devices
 		/// <summary>
 		/// Device service constructor
 		/// </summary>
-		public DeviceService(IDeviceCollection devices, ISingletonDocument<DevicePlatformMapV1> platformMapSingleton, IUserCollection userCollection, JobService jobService, IStreamCollection streamCollection, IOptionsMonitor<ServerSettings> settings, IOptionsMonitor<GlobalConfig> globalConfig, INotificationService notificationService, IClock clock, ILogger<DeviceService> logger)
+		public DeviceService(IDeviceCollection devices, ISingletonDocument<DevicePlatformMapV1> platformMapSingleton, IUserCollection userCollection, JobService jobService, IStreamCollection streamCollection, IOptionsMonitor<ServerSettings> settings, IOptionsMonitor<GlobalConfig> globalConfig, INotificationService notificationService, IClock clock, Tracer tracer, ILogger<DeviceService> logger)
 		{
 			_userCollection = userCollection;
 			_devices = devices;
@@ -117,6 +117,7 @@ namespace Horde.Server.Devices
 			_notificationService = notificationService;
 			_ticker = clock.AddSharedTicker<DeviceService>(TimeSpan.FromMinutes(1.0), TickAsync, logger);
 			_telemetryTicker = clock.AddSharedTicker("DeviceService.Telemetry", TimeSpan.FromMinutes(10.0), TickTelemetryAsync, logger);
+			_tracer = tracer;
 			_logger = logger;
 			_settings = settings;
 			_globalConfig = globalConfig;
@@ -182,8 +183,7 @@ namespace Horde.Server.Devices
 
 			if (!stoppingToken.IsCancellationRequested)
 			{
-				using IScope scope = GlobalTracer.Instance.BuildSpan("DeviceService.TickTelemetryAsync").StartActive();
-
+				using TelemetrySpan span = _tracer.StartActiveSpan($"{nameof(DeviceService)}.{nameof(TickTelemetryAsync)}");
 				_logger.LogInformation("Updating pool telemetry");
 				await _devices.CreatePoolTelemetrySnapshot();
 			}
@@ -196,7 +196,7 @@ namespace Horde.Server.Devices
 		{
 			if (!stoppingToken.IsCancellationRequested)
 			{
-				using IScope scope = GlobalTracer.Instance.BuildSpan("DeviceService.TickAsync").StartActive();
+				using TelemetrySpan span = _tracer.StartActiveSpan($"{nameof(DeviceService)}.{nameof(TickAsync)}");
 
 				GlobalConfig globalConfig = _globalConfig.CurrentValue;
 
diff --git a/Engine/Source/Programs/Horde/Horde.Server/Horde.Server.csproj b/Engine/Source/Programs/Horde/Horde.Server/Horde.Server.csproj
--- a/Engine/Source/Programs/Horde/Horde.Server/Horde.Server.csproj
+++ b/Engine/Source/Programs/Horde/Horde.Server/Horde.Server.csproj
@@ -44,8 +44,6 @@
     <PackageReference Include="Microsoft.Extensions.Caching.Memory" Version="6.0.1" />
     <PackageReference Include="Microsoft.Extensions.Http" Version="6.0.0" />
     <PackageReference Include="MongoDB.Driver" Version="2.18.0" />
-    <PackageReference Include="OpenTracing" Version="0.12.1" />
-    <PackageReference Include="OpenTracing.Contrib.Grpc" Version="0.4.0" />
     <PackageReference Include="OpenTelemetry" Version="1.4.0" />
     <PackageReference Include="OpenTelemetry.Exporter.Console" Version="1.4.0" />
     <PackageReference Include="OpenTelemetry.Exporter.OpenTelemetryProtocol" Version="1.4.0" />
diff --git a/Engine/Source/Programs/Horde/Horde.Server/Issues/IssueService.cs b/Engine/Source/Programs/Horde/Horde.Server/Issues/IssueService.cs
--- a/Engine/Source/Programs/Horde/Horde.Server/Issues/IssueService.cs
+++ b/Engine/Source/Programs/Horde/Horde.Server/Issues/IssueService.cs
@@ -18,16 +18,13 @@ using Horde.Server.Perforce;
 using Horde.Server.Server;
 using Horde.Server.Streams;
 using Horde.Server.Users;
-using Horde.Server.Utilities;
 using HordeCommon;
 using Microsoft.Extensions.Hosting;
 using Microsoft.Extensions.Logging;
 using Microsoft.Extensions.Options;
 using MongoDB.Bson;
 using MongoDB.Bson.Serialization;
-using MongoDB.Driver;
-using OpenTracing;
-using OpenTracing.Util;
+using OpenTelemetry.Trace;
 
 namespace Horde.Server.Issues
 {
@@ -193,6 +190,12 @@ namespace Horde.Server.Issues
 		/// </summary>
 		Dictionary<StreamId, HashSet<TemplateId>> _cachedDesktopAlerts = new Dictionary<StreamId, HashSet<TemplateId>>();
 
+		/// <summary>
+		/// Tracer
+		/// </summary>
+		/// 
+		readonly Tracer _tracer;
+
 		/// <summary>
 		/// Logger for tracing
 		/// </summary>
@@ -208,7 +211,7 @@ namespace Horde.Server.Issues
 		/// <summary>
 		/// Constructor
 		/// </summary>
-		public IssueService(IIssueCollection issueCollection, ICommitService commitService, IJobStepRefCollection jobStepRefs, IStreamCollection streams, IUserCollection userCollection, ILogFileService logFileService, IClock clock, IOptionsMonitor<GlobalConfig> globalConfig, ILogger<IssueService> logger)
+		public IssueService(IIssueCollection issueCollection, ICommitService commitService, IJobStepRefCollection jobStepRefs, IStreamCollection streams, IUserCollection userCollection, ILogFileService logFileService, IClock clock, IOptionsMonitor<GlobalConfig> globalConfig, Tracer tracer, ILogger<IssueService> logger)
 		{
 			Type[] issueTypes = Assembly.GetExecutingAssembly().GetTypes().Where(x => !x.IsAbstract && typeof(IIssue).IsAssignableFrom(x)).ToArray();
 			foreach (Type issueType in issueTypes)
@@ -226,6 +229,7 @@ namespace Horde.Server.Issues
 			_clock = clock;
 			_ticker = clock.AddTicker<IssueService>(TimeSpan.FromMinutes(1.0), TickAsync, logger);
 			_globalConfig = globalConfig;
+			_tracer = tracer;
 			_logger = logger;
 
 			// Create all the issue handlers
@@ -499,10 +503,10 @@ namespace Horde.Server.Issues
 		/// <returns>Async task</returns>
 		public async Task UpdateCompleteStep(IJob job, IGraph graph, SubResourceId batchId, SubResourceId stepId)
 		{
-			using IScope scope = GlobalTracer.Instance.BuildSpan("UpdateCompleteStep").StartActive();
-			scope.Span.SetTag("JobId", job.Id.ToString());
-			scope.Span.SetTag("BatchId", batchId.ToString());
-			scope.Span.SetTag("StepId", stepId.ToString());
+			using TelemetrySpan span = _tracer.StartActiveSpan($"{nameof(IssueService)}.{nameof(UpdateCompleteStep)}");
+			span.SetAttribute("jobId", job.Id.ToString());
+			span.SetAttribute("batchId", batchId.ToString());
+			span.SetAttribute("stepId", stepId.ToString());
 
 			_logger.LogInformation("Updating issues for {JobId}:{BatchId}:{StepId}", job.Id, batchId, stepId);
 
@@ -526,7 +530,7 @@ namespace Horde.Server.Issues
 				throw new ArgumentException($"Invalid step id {stepId}");
 			}
 
-			scope.Span.SetTag("LogId", step.LogId.ToString());
+			span.SetAttribute("logId", step.LogId.ToString());
 
 			// Get all the annotations for this template
 			TemplateRefConfig? templateRef;
diff --git a/Engine/Source/Programs/Horde/Horde.Server/Jobs/JobCollection.cs b/Engine/Source/Programs/Horde/Horde.Server/Jobs/JobCollection.cs
--- a/Engine/Source/Programs/Horde/Horde.Server/Jobs/JobCollection.cs
+++ b/Engine/Source/Programs/Horde/Horde.Server/Jobs/JobCollection.cs
@@ -19,7 +19,6 @@ using Horde.Server.Streams;
 using Horde.Server.Telemetry;
 using Horde.Server.Users;
 using Horde.Server.Utilities;
-using Horde.Common;
 using HordeCommon;
 using HordeCommon.Rpc.Tasks;
 using Microsoft.Extensions.Logging;
@@ -29,8 +28,7 @@ using MongoDB.Bson.Serialization;
 using MongoDB.Bson.Serialization.Attributes;
 using MongoDB.Bson.Serialization.Options;
 using MongoDB.Driver;
-using OpenTracing;
-using OpenTracing.Util;
+using OpenTelemetry.Trace;
 
 namespace Horde.Server.Jobs
 {
@@ -348,6 +346,7 @@ namespace Horde.Server.Jobs
 		readonly MongoIndex<JobDocument> _streamThenTemplateThenCreationTimeIndex;
 		readonly ITelemetrySink _telemetrySink;
 		readonly IClock _clock;
+		readonly Tracer _tracer;
 		readonly ILogger<JobCollection> _logger;
 
 		/// <summary>
@@ -356,11 +355,13 @@ namespace Horde.Server.Jobs
 		/// <param name="mongoService">The database service singleton</param>
 		/// <param name="clock"></param>
 		/// <param name="telemetrySink">Telemetry sink for data</param>
+		/// <param name="tracer">Tracer</param>
 		/// <param name="logger">The logger instance</param>
-		public JobCollection(MongoService mongoService, IClock clock, ITelemetrySink telemetrySink, ILogger<JobCollection> logger)
+		public JobCollection(MongoService mongoService, IClock clock, ITelemetrySink telemetrySink, Tracer tracer, ILogger<JobCollection> logger)
 		{
 			_clock = clock;
 			_telemetrySink = telemetrySink;
+			_tracer = tracer;
 			_logger = logger;
 
 			List<MongoIndex<JobDocument>> indexes = new List<MongoIndex<JobDocument>>();
@@ -533,7 +534,7 @@ namespace Horde.Server.Jobs
 			}
 
 			List<JobDocument> results;
-			using (IScope scope = GlobalTracer.Instance.BuildSpan("Jobs.Find").StartActive())
+			using (TelemetrySpan _ = _tracer.StartActiveSpan($"{nameof(JobCollection)}.{nameof(FindAsync)}"))
 			{
 				IMongoCollection<JobDocument> collection = consistentRead ? _jobs : _jobs.WithReadPreference(ReadPreference.SecondaryPreferred);
 				results = await collection.FindWithHint(filter, indexHint, x => x.SortByDescending(x => x.CreateTimeUtc!).Range(index, count).ToListAsync());
diff --git a/Engine/Source/Programs/Horde/Horde.Server/Jobs/JobService.cs b/Engine/Source/Programs/Horde/Horde.Server/Jobs/JobService.cs
--- a/Engine/Source/Programs/Horde/Horde.Server/Jobs/JobService.cs
+++ b/Engine/Source/Programs/Horde/Horde.Server/Jobs/JobService.cs
@@ -3,14 +3,12 @@
 using System;
 using System.Collections.Generic;
 using System.Globalization;
-using System.IO;
 using System.Linq;
 using System.Security.Claims;
 using System.Text.RegularExpressions;
 using System.Threading;
 using System.Threading.Tasks;
 using EpicGames.Core;
-using EpicGames.Perforce;
 using Horde.Server.Acls;
 using Horde.Server.Agents;
 using Horde.Server.Agents.Leases;
@@ -29,10 +27,7 @@ using Horde.Server.Utilities;
 using HordeCommon;
 using Microsoft.Extensions.Logging;
 using MongoDB.Bson;
-using MongoDB.Driver;
 using OpenTelemetry.Trace;
-using OpenTracing;
-using OpenTracing.Util;
 
 namespace Horde.Server.Jobs
 {
diff --git a/Engine/Source/Programs/Horde/Horde.Server/Jobs/JobsController.cs b/Engine/Source/Programs/Horde/Horde.Server/Jobs/JobsController.cs
--- a/Engine/Source/Programs/Horde/Horde.Server/Jobs/JobsController.cs
+++ b/Engine/Source/Programs/Horde/Horde.Server/Jobs/JobsController.cs
@@ -25,8 +25,6 @@ using Microsoft.AspNetCore.Mvc;
 using Microsoft.Extensions.Logging;
 using Microsoft.Extensions.Options;
 using MongoDB.Bson;
-using OpenTracing;
-using OpenTracing.Util;
 
 namespace Horde.Server.Jobs
 {
@@ -776,12 +774,9 @@ namespace Horde.Server.Jobs
 			}
 
 			List<IJob> jobs;
-			using (IScope _ = GlobalTracer.Instance.BuildSpan("FindJobs").StartActive())
-			{
-				jobs = await _jobService.FindJobsAsync(jobIdValues, streamIdValue, name, templateRefIds, minChange,
-					maxChange, preflightChange, preflightOnly, preflightStartedByUserIdValue, startedByUserIdValue, minCreateTime?.UtcDateTime, maxCreateTime?.UtcDateTime, target, state, outcome,
-					modifiedBefore, modifiedAfter, index, count, false);
-			}
+			jobs = await _jobService.FindJobsAsync(jobIdValues, streamIdValue, name, templateRefIds, minChange,
+				maxChange, preflightChange, preflightOnly, preflightStartedByUserIdValue, startedByUserIdValue, minCreateTime?.UtcDateTime, maxCreateTime?.UtcDateTime, target, state, outcome,
+				modifiedBefore, modifiedAfter, index, count, false);
 
 			return await CreateAuthorizedJobResponsesAsync(jobs, filter);
 		}
diff --git a/Engine/Source/Programs/Horde/Horde.Server/Jobs/Schedules/ScheduleService.cs b/Engine/Source/Programs/Horde/Horde.Server/Jobs/Schedules/ScheduleService.cs
--- a/Engine/Source/Programs/Horde/Horde.Server/Jobs/Schedules/ScheduleService.cs
+++ b/Engine/Source/Programs/Horde/Horde.Server/Jobs/Schedules/ScheduleService.cs
@@ -21,9 +21,7 @@ using HordeCommon;
 using Microsoft.Extensions.Hosting;
 using Microsoft.Extensions.Logging;
 using Microsoft.Extensions.Options;
-using MongoDB.Driver;
-using OpenTracing;
-using OpenTracing.Util;
+using OpenTelemetry.Trace;
 using StackExchange.Redis;
 
 namespace Horde.Server.Jobs.Schedules
@@ -63,6 +61,7 @@ namespace Horde.Server.Jobs.Schedules
 		readonly IStreamCollection _streamCollection;
 		readonly ITemplateCollection _templateCollection;
 		readonly IClock _clock;
+		readonly Tracer _tracer;
 		readonly ILogger _logger;
 		readonly RedisService _redis;
 		static readonly RedisKey s_baseLockKey = "scheduler/locks";
@@ -74,7 +73,7 @@ namespace Horde.Server.Jobs.Schedules
 		/// <summary>
 		/// Constructor
 		/// </summary>
-		public ScheduleService(RedisService redis, IGraphCollection graphs, ICommitService commitService, IJobCollection jobCollection, JobService jobService, IStreamCollection streamCollection, ITemplateCollection templateCollection, MongoService mongoService, IClock clock, IOptionsMonitor<GlobalConfig> globalConfig, ILogger<ScheduleService> logger)
+		public ScheduleService(RedisService redis, IGraphCollection graphs, ICommitService commitService, IJobCollection jobCollection, JobService jobService, IStreamCollection streamCollection, ITemplateCollection templateCollection, MongoService mongoService, IClock clock, IOptionsMonitor<GlobalConfig> globalConfig, Tracer tracer, ILogger<ScheduleService> logger)
 		{
 			_graphs = graphs;
 			_commitService = commitService;
@@ -84,6 +83,7 @@ namespace Horde.Server.Jobs.Schedules
 			_templateCollection = templateCollection;
 			_clock = clock;
 			_globalConfig = globalConfig;
+			_tracer = tracer;
 			_logger = logger;
 
 			_redis = redis;
@@ -251,9 +251,9 @@ namespace Horde.Server.Jobs.Schedules
 				return false;
 			}
 
-			using IScope scope = GlobalTracer.Instance.BuildSpan("ScheduleService.TriggerAsync").StartActive();
-			scope.Span.SetTag("StreamId", stream.Id);
-			scope.Span.SetTag("TemplateId", templateId);
+			using TelemetrySpan span = _tracer.StartActiveSpan($"{nameof(ScheduleService)}.{nameof(TriggerAsync)}");
+			span.SetAttribute("streamId", stream.Id);
+			span.SetAttribute("templateId", templateId);
 
 			Stopwatch stopwatch = Stopwatch.StartNew();
 			_logger.LogInformation("Updating schedule for {StreamId} template {TemplateId}", stream.Id, templateId);
diff --git a/Engine/Source/Programs/Horde/Horde.Server/Jobs/TestData/TestDataCollection.cs b/Engine/Source/Programs/Horde/Horde.Server/Jobs/TestData/TestDataCollection.cs
--- a/Engine/Source/Programs/Horde/Horde.Server/Jobs/TestData/TestDataCollection.cs
+++ b/Engine/Source/Programs/Horde/Horde.Server/Jobs/TestData/TestDataCollection.cs
@@ -7,16 +7,13 @@ using System.Threading.Tasks;
 using Horde.Server.Server;
 using Horde.Server.Streams;
 using Horde.Server.Utilities;
-using Microsoft.Extensions.Configuration;
 using Microsoft.Extensions.Logging;
 using MongoDB.Bson;
 using MongoDB.Bson.Serialization;
 using MongoDB.Bson.Serialization.Attributes;
 using MongoDB.Driver;
-using OpenTracing.Util;
-using OpenTracing;
-using Microsoft.CodeAnalysis;
 using Microsoft.Extensions.Options;
+using OpenTelemetry.Trace;
 
 namespace Horde.Server.Jobs.TestData
 {
@@ -397,6 +394,7 @@ namespace Horde.Server.Jobs.TestData
 		/// </summary>
 		readonly IMongoCollection<TestStreamDocument> _testStreams;
 
+		readonly Tracer _tracer;
 		readonly ILogger _logger;
 
 
@@ -404,10 +402,12 @@ namespace Horde.Server.Jobs.TestData
 		/// Constructor
 		/// </summary>
 		/// <param name="mongoService"></param>
+		/// <param name="tracer"></param>
 		/// <param name="logger"></param>
 		/// <param name="settings"></param>
-		public TestDataCollection(MongoService mongoService, ILogger<TestDataCollection> logger, IOptionsMonitor<ServerSettings> settings)
+		public TestDataCollection(MongoService mongoService, Tracer tracer, ILogger<TestDataCollection> logger, IOptionsMonitor<ServerSettings> settings)
 		{
+			_tracer = tracer;
 			_logger = logger;
 
 			List<MongoIndex<TestDataDocument>> indexes = new List<MongoIndex<TestDataDocument>>();
@@ -504,7 +504,8 @@ namespace Horde.Server.Jobs.TestData
 			}
 
 			List<TestDataRefDocument> results;
-			using (IScope scope = GlobalTracer.Instance.BuildSpan("TestData.FindTestRefs").StartActive())
+			
+			using (TelemetrySpan _ = _tracer.StartActiveSpan($"{nameof(TestDataCollection)}.{nameof(FindTestRefs)}"))
 			{
 				results = await _testRefs.Find(filter).ToListAsync();
 			}
diff --git a/Engine/Source/Programs/Horde/Horde.Server/Logs/Builder/RedisLogBuilder.cs b/Engine/Source/Programs/Horde/Horde.Server/Logs/Builder/RedisLogBuilder.cs
--- a/Engine/Source/Programs/Horde/Horde.Server/Logs/Builder/RedisLogBuilder.cs
+++ b/Engine/Source/Programs/Horde/Horde.Server/Logs/Builder/RedisLogBuilder.cs
@@ -2,7 +2,6 @@
 
 using System;
 using System.Collections.Generic;
-using System.Globalization;
 using System.Text.RegularExpressions;
 using System.Threading.Tasks;
 using EpicGames.Core;
@@ -10,8 +9,7 @@ using EpicGames.Redis;
 using Horde.Server.Logs.Data;
 using Horde.Server.Utilities;
 using Microsoft.Extensions.Logging;
-using OpenTracing;
-using OpenTracing.Util;
+using OpenTelemetry.Trace;
 using StackExchange.Redis;
 
 namespace Horde.Server.Logs.Builder
@@ -61,6 +59,11 @@ namespace Horde.Server.Logs.Builder
 		/// </summary>
 		readonly RedisConnectionPool _redisConnectionPool;
 
+		/// <summary>
+		/// Tracer
+		/// </summary>
+		readonly Tracer _tracer;
+		
 		/// <summary>
 		/// Logger for debug output
 		/// </summary>
@@ -70,10 +73,12 @@ namespace Horde.Server.Logs.Builder
 		/// Constructor
 		/// </summary>
 		/// <param name="redisConnectionPool">The redis database singleton</param>
+		/// <param name="tracer">Tracer</param>
 		/// <param name="logger">Logger for debug output</param>
-		public RedisLogBuilder(RedisConnectionPool redisConnectionPool, ILogger logger)
+		public RedisLogBuilder(RedisConnectionPool redisConnectionPool, Tracer tracer, ILogger logger)
 		{
 			_redisConnectionPool = redisConnectionPool;
+			_tracer = tracer;
 			_logger = logger;
 		}
 
@@ -85,10 +90,10 @@ namespace Horde.Server.Logs.Builder
 
 			if(chunkOffset == writeOffset)
 			{
-				using IScope scope = GlobalTracer.Instance.BuildSpan("Redis.CreateChunk").StartActive();
-				scope.Span.SetTag("LogId", logId.ToString());
-				scope.Span.SetTag("Offset", chunkOffset.ToString(CultureInfo.InvariantCulture));
-				scope.Span.SetTag("WriteOffset", writeOffset.ToString(CultureInfo.InvariantCulture));
+				using TelemetrySpan span = _tracer.StartActiveSpan($"{nameof(RedisLogBuilder)}.{nameof(AppendAsync)}.CreateChunk");
+				span.SetAttribute("logId", logId.ToString());
+				span.SetAttribute("offset", chunkOffset);
+				span.SetAttribute("writeOffset", writeOffset);
 
 				ITransaction createTransaction = redisDb.CreateTransaction();
 				createTransaction.AddCondition(Condition.SortedSetNotContains(ItemsKey, keys.Prefix));
@@ -105,10 +110,10 @@ namespace Horde.Server.Logs.Builder
 			}
 
 			{
-				using IScope scope = GlobalTracer.Instance.BuildSpan("Redis.AppendChunk").StartActive();
-				scope.Span.SetTag("LogId", logId.ToString());
-				scope.Span.SetTag("Offset", chunkOffset.ToString(CultureInfo.InvariantCulture));
-				scope.Span.SetTag("WriteOffset", writeOffset.ToString(CultureInfo.InvariantCulture));
+				using TelemetrySpan span = _tracer.StartActiveSpan($"{nameof(RedisLogBuilder)}.{nameof(AppendAsync)}.AppendChunk");
+				span.SetAttribute("logId", logId.ToString());
+				span.SetAttribute("offset", chunkOffset);
+				span.SetAttribute("writeOffset", writeOffset);
 
 				ITransaction appendTransaction = redisDb.CreateTransaction();
 				appendTransaction.AddCondition(Condition.SortedSetContains(ItemsKey, keys.Prefix));
@@ -138,9 +143,9 @@ namespace Horde.Server.Logs.Builder
 			int numTries;
 			for (numTries = 0; numTries < MaxRetries; numTries++)
 			{
-				using IScope scope = GlobalTracer.Instance.BuildSpan("Redis.CompleteSubChunk").StartActive();
-				scope.Span.SetTag("LogId", logId.ToString());
-				scope.Span.SetTag("Offset", offset.ToString(CultureInfo.InvariantCulture));
+				using TelemetrySpan span = _tracer.StartActiveSpan($"{nameof(RedisLogBuilder)}.{nameof(CompleteSubChunkAsync)}");
+				span.SetAttribute("logId", logId.ToString());
+				span.SetAttribute("offset", offset);
 
 				LogType type = (LogType)(int)await redisDb.StringGetAsync(keys.Type);
 				int length = (int)await redisDb.StringGetAsync(keys.Length);
@@ -194,9 +199,10 @@ namespace Horde.Server.Logs.Builder
 		public async Task CompleteChunkAsync(LogId logId, long offset)
 		{
 			IDatabase redisDb = _redisConnectionPool.GetDatabase();
-			using IScope scope = GlobalTracer.Instance.BuildSpan("Redis.CompleteChunk").StartActive();
-			scope.Span.SetTag("LogId", logId.ToString());
-			scope.Span.SetTag("Offset", offset.ToString(CultureInfo.InvariantCulture));
+			
+			using TelemetrySpan span = _tracer.StartActiveSpan($"{nameof(RedisLogBuilder)}.{nameof(CompleteChunkAsync)}");
+			span.SetAttribute("logId", logId.ToString());
+			span.SetAttribute("offset", offset);
 
 			ChunkKeys keys = new ChunkKeys(logId, offset);
 
@@ -216,9 +222,9 @@ namespace Horde.Server.Logs.Builder
 		public async Task RemoveChunkAsync(LogId logId, long offset)
 		{
 			IDatabase redisDb = _redisConnectionPool.GetDatabase();
-			using IScope scope = GlobalTracer.Instance.BuildSpan("Redis.RemoveChunk").StartActive();
-			scope.Span.SetTag("LogId", logId.ToString());
-			scope.Span.SetTag("Offset", offset.ToString(CultureInfo.InvariantCulture));
+			using TelemetrySpan span = _tracer.StartActiveSpan($"{nameof(RedisLogBuilder)}.{nameof(RemoveChunkAsync)}");
+			span.SetAttribute("logId", logId.ToString());
+			span.SetAttribute("offset", offset);
 
 			ChunkKeys keys = new ChunkKeys(logId, offset);
 
@@ -243,9 +249,9 @@ namespace Horde.Server.Logs.Builder
 			int numTries;
 			for (numTries = 0; numTries < MaxRetries; numTries++)
 			{
-				using IScope scope = GlobalTracer.Instance.BuildSpan("Redis.GetChunk").StartActive();
-				scope.Span.SetTag("LogId", logId.ToString());
-				scope.Span.SetTag("Offset", offset.ToString(CultureInfo.InvariantCulture));
+				using TelemetrySpan span = _tracer.StartActiveSpan($"{nameof(RedisLogBuilder)}.{nameof(GetChunkAsync)}");
+				span.SetAttribute("logId", logId.ToString());
+				span.SetAttribute("offset", offset);
 
 				RedisValue chunkDataValue = await redisDb.StringGetAsync(keys.ChunkData);
 
diff --git a/Engine/Source/Programs/Horde/Horde.Server/Logs/Data/LogIndexData.cs b/Engine/Source/Programs/Horde/Horde.Server/Logs/Data/LogIndexData.cs
--- a/Engine/Source/Programs/Horde/Horde.Server/Logs/Data/LogIndexData.cs
+++ b/Engine/Source/Programs/Horde/Horde.Server/Logs/Data/LogIndexData.cs
@@ -7,8 +7,7 @@ using System.Linq;
 using EpicGames.Core;
 using EpicGames.Horde.Logs;
 using Horde.Server.Utilities;
-using OpenTracing;
-using OpenTracing.Util;
+using OpenTelemetry.Trace;
 
 namespace Horde.Server.Logs.Data
 {
@@ -185,7 +184,7 @@ namespace Horde.Server.Logs.Data
 		/// <returns>Index data</returns>
 		public static LogIndexData Merge(IEnumerable<LogIndexData> indexes)
 		{
-			using IScope scope = GlobalTracer.Instance.BuildSpan("LogIndexData.Merge").StartActive();
+			using TelemetrySpan span = OpenTelemetryTracers.Horde.StartActiveSpan($"{nameof(LogIndexData)}.{nameof(Merge)}");
 
 			// Create the combined block list
 			LogIndexBlock[] newBlocks = indexes.SelectMany(x => x._blocks).ToArray();
diff --git a/Engine/Source/Programs/Horde/Horde.Server/Logs/LogFileService.cs b/Engine/Source/Programs/Horde/Horde.Server/Logs/LogFileService.cs
--- a/Engine/Source/Programs/Horde/Horde.Server/Logs/LogFileService.cs
+++ b/Engine/Source/Programs/Horde/Horde.Server/Logs/LogFileService.cs
@@ -2,9 +2,7 @@
 
 using System;
 using System.Collections.Generic;
-using System.Collections.Immutable;
 using System.Diagnostics;
-using System.Globalization;
 using System.IO;
 using System.Linq;
 using System.Runtime.CompilerServices;
@@ -29,8 +27,7 @@ using Microsoft.Extensions.Hosting;
 using Microsoft.Extensions.Logging;
 using Microsoft.Extensions.Options;
 using MongoDB.Bson;
-using OpenTracing;
-using OpenTracing.Util;
+using OpenTelemetry.Trace;
 using Stream = System.IO.Stream;
 
 namespace Horde.Server.Logs
@@ -231,7 +228,8 @@ namespace Horde.Server.Logs
 	public sealed class LogFileService : IHostedService, ILogFileService, IDisposable
 	{
 		private const int MaxConcurrentChunkWrites = 10;
-		
+
+		private readonly Tracer _tracer;
 		private readonly ILogger<LogFileService> _logger;
 		private readonly ILogFileCollection _logFiles;
 		private readonly ILogEventCollection _logEvents;
@@ -562,7 +560,7 @@ namespace Horde.Server.Logs
 		/// <summary>
 		/// Constructor
 		/// </summary>
-		public LogFileService(ILogFileCollection logFiles, ILogEventCollection logEvents, ILogBuilder builder, ILogStorage storage, IClock clock, LogTailService logTailService, StorageService storageService, IOptions<ServerSettings> settings, ILogger<LogFileService> logger)
+		public LogFileService(ILogFileCollection logFiles, ILogEventCollection logEvents, ILogBuilder builder, ILogStorage storage, IClock clock, LogTailService logTailService, StorageService storageService, IOptions<ServerSettings> settings, Tracer tracer, ILogger<LogFileService> logger)
 		{
 			_logFiles = logFiles;
 			_logEvents = logEvents;
@@ -573,6 +571,7 @@ namespace Horde.Server.Logs
 			_logTailService = logTailService;
 			_storageService = storageService;
 			_settings = settings;
+			_tracer = tracer;
 			_logger = logger;
 		}
 
@@ -728,11 +727,11 @@ namespace Horde.Server.Logs
 		/// <inheritdoc/>
 		public async Task<ILogFile?> WriteLogDataAsync(ILogFile logFile, long offset, int lineIndex, ReadOnlyMemory<byte> data, bool flush, int maxChunkLength, int maxSubChunkLineCount, CancellationToken cancellationToken)
 		{
-			using IScope scope = GlobalTracer.Instance.BuildSpan("WriteLogDataAsync").StartActive();
-			scope.Span.SetTag("LogId", logFile.Id.ToString());
-			scope.Span.SetTag("Offset", offset.ToString(CultureInfo.InvariantCulture));
-			scope.Span.SetTag("Length", data.Length.ToString(CultureInfo.InvariantCulture));
-			scope.Span.SetTag("LineIndex", lineIndex.ToString(CultureInfo.InvariantCulture));
+			using TelemetrySpan span = _tracer.StartActiveSpan($"{nameof(LogFileService)}.{nameof(WriteLogDataAsync)}");
+			span.SetAttribute("logId", logFile.Id.ToString());
+			span.SetAttribute("offset", offset);
+			span.SetAttribute("length", data.Length);
+			span.SetAttribute("lineIndex", lineIndex);
 
 			// Make sure the data ends in a newline
 			if (data.Length > 0 && data.Span[data.Length - 1] != '\n')
@@ -1045,10 +1044,10 @@ namespace Horde.Server.Logs
 		/// <inheritdoc/>
 		public async Task<ILogEventData> GetEventDataAsync(ILogFile logFile, int lineIndex, int lineCount, CancellationToken cancellationToken)
 		{
-			using IScope scope = GlobalTracer.Instance.BuildSpan("GetEventDataAsync").StartActive();
-			scope.Span.SetTag("LogId", logFile.Id.ToString());
-			scope.Span.SetTag("LineIndex", lineIndex.ToString(CultureInfo.InvariantCulture));
-			scope.Span.SetTag("LineCount", lineCount.ToString(CultureInfo.InvariantCulture));
+			using TelemetrySpan span = _tracer.StartActiveSpan($"{nameof(LogFileService)}.{nameof(GetEventDataAsync)}");
+			span.SetAttribute("logId", logFile.Id.ToString());
+			span.SetAttribute("lineIndex", lineIndex);
+			span.SetAttribute("lineCount", lineCount);
 
 			List<Utf8String> lines = await ReadLinesAsync(logFile, lineIndex, lineCount, cancellationToken);
 			List<LogEventLine> eventLines = new List<LogEventLine>(lines.Count);
@@ -1284,7 +1283,7 @@ namespace Horde.Server.Logs
 		/// <param name="stoppingToken">Cancellation token</param>
 		async ValueTask TickAsync(CancellationToken stoppingToken)
 		{
-			using IScope scope = GlobalTracer.Instance.BuildSpan("LogFileService.TickAsync").StartActive();
+			using TelemetrySpan span = _tracer.StartActiveSpan($"{nameof(LogFileService)}.{nameof(TickAsync)}");
 			
 			lock (_writeLock)
 			{
@@ -1306,12 +1305,12 @@ namespace Horde.Server.Logs
 		/// <returns>Async task</returns>
 		private async Task IncrementalFlushAsync(CancellationToken cancellationToken)
 		{
-			using IScope scope = GlobalTracer.Instance.BuildSpan("LogFileService.IncrementalFlush").StartActive();
+			using TelemetrySpan span = _tracer.StartActiveSpan($"{nameof(LogFileService)}.{nameof(IncrementalFlushAsync)}");
 			
 			// Get all the chunks older than 20 minutes
 			List<(LogId, long)> flushChunks = await _builder.TouchChunksAsync(TimeSpan.FromMinutes(10.0));
 
-			scope.Span.SetTag("numChunks", flushChunks.Count);
+			span.SetAttribute("numChunks", flushChunks.Count);
 
 			// Mark them all as complete
 			foreach ((LogId logId, long offset) in flushChunks)
@@ -1329,7 +1328,7 @@ namespace Horde.Server.Logs
 		/// <returns>Async task</returns>
 		public async Task FlushAsync()
 		{
-			using IScope scope = GlobalTracer.Instance.BuildSpan("LogFileService.FlushAsync").StartActive();
+			using TelemetrySpan span = _tracer.StartActiveSpan($"{nameof(LogFileService)}.{nameof(FlushAsync)}");
 			_logger.LogInformation("Forcing flush of pending log chunks...");
 
 			// Mark everything in the cache as complete
@@ -1373,7 +1372,8 @@ namespace Horde.Server.Logs
 		/// <param name="createIndex">Create an index for the log</param>
 		private void WriteCompleteChunks(List<(LogId, long)> chunksToWrite, bool createIndex)
 		{
-			using IScope scope = GlobalTracer.Instance.BuildSpan("LogFileService.WriteCompleteChunks").StartActive();
+			using TelemetrySpan span = _tracer.StartActiveSpan($"{nameof(LogFileService)}.{nameof(WriteCompleteChunks)}");
+
 			int numTasksCreated = 0;
 			
 			foreach (IGrouping<LogId, long> group in chunksToWrite.GroupBy(x => x.Item1, x => x.Item2))
@@ -1405,7 +1405,7 @@ namespace Horde.Server.Logs
 				}
 			}
 
-			scope.Span.SetTag("numWriteTasksCreated", numTasksCreated);
+			span.SetAttribute("numWriteTasksCreated", numTasksCreated);
 			_logger.LogInformation("{NumWriteTasksCreated} write tasks created", numTasksCreated);
 		}
 		
@@ -1417,7 +1417,7 @@ namespace Horde.Server.Logs
 		/// <param name="cancellationToken">Cancellation token</param>
 		private async Task WriteCompleteChunksV2Async(List<(LogId, long)> chunksToWrite, bool createIndex, CancellationToken cancellationToken)
 		{
-			using IScope scope = GlobalTracer.Instance.BuildSpan("LogFileService.WriteCompleteChunksV2Async").StartActive();
+			using TelemetrySpan span = _tracer.StartActiveSpan($"{nameof(LogFileService)}.{nameof(WriteCompleteChunksV2Async)}");
 			
 			HashSet<(LogId, long)> writeChunks = new ();
 			List<(LogId, List<long>)> offsetsToWrite = new();
@@ -1443,7 +1443,7 @@ namespace Horde.Server.Logs
 				}
 			}
 
-			scope.Span.SetTag("numOffsetsToWrite", offsetsToWrite.Count);
+			span.SetAttribute("numOffsetsToWrite", offsetsToWrite.Count);
 			ParallelOptions opts = new() { MaxDegreeOfParallelism = MaxConcurrentChunkWrites, CancellationToken = cancellationToken };
 			await Parallel.ForEachAsync(offsetsToWrite, opts, async (x, innerCt) =>
 			{
@@ -1480,10 +1480,10 @@ namespace Horde.Server.Logs
 		/// <returns>Async task</returns>
 		private async Task<ILogFile?> WriteCompleteChunksForLogAsync(ILogFile logFileInterface, List<long> offsets, bool createIndex, CancellationToken cancellationToken = default)
 		{
-			using IScope scope = GlobalTracer.Instance.BuildSpan("WriteCompleteChunksForLogAsync").StartActive();
-			scope.Span.SetTag("LogId", logFileInterface.Id.ToString());
-			scope.Span.SetTag("NumOffsets", offsets.Count);
-			scope.Span.SetTag("CreateIndex", createIndex);
+			using TelemetrySpan span = _tracer.StartActiveSpan($"{nameof(LogFileService)}.{nameof(WriteCompleteChunksForLogAsync)}");
+			span.SetAttribute("logId", logFileInterface.Id.ToString());
+			span.SetAttribute("numOffsets", offsets.Count);
+			span.SetAttribute("createIndex", createIndex);
 			
 			// Write the data to the storage provider
 			List<Task<LogChunkData?>> chunkWriteTasks = new List<Task<LogChunkData?>>();
@@ -1498,7 +1498,7 @@ namespace Horde.Server.Logs
 				}
 			}
 			
-			scope.Span.SetTag("NumWriteTasks", chunkWriteTasks.Count);
+			span.SetAttribute("numWriteTasks", chunkWriteTasks.Count);
 
 			// Wait for the tasks to complete, periodically updating the log file object
 			ILogFile? logFile = logFileInterface;
@@ -1582,9 +1582,9 @@ namespace Horde.Server.Logs
 			}
 
 			// Save stats for the index creation
-			using IScope scope = GlobalTracer.Instance.BuildSpan("CreateIndexAsync").StartActive();
-			scope.Span.SetTag("LogId", logFile.Id.ToString());
-			scope.Span.SetTag("Length", (lastChunk.Offset + lastChunk.Length).ToString(CultureInfo.InvariantCulture));
+			using TelemetrySpan span = _tracer.StartActiveSpan($"{nameof(LogFileService)}.{nameof(CreateIndexAsync)}");
+			span.SetAttribute("logId", logFile.Id.ToString());
+			span.SetAttribute("length", lastChunk.Offset + lastChunk.Length);
 
 			long newLength = 0;
 			int newLineCount = 0;
@@ -1863,10 +1863,10 @@ namespace Horde.Server.Logs
 		{
 			Stopwatch timer = Stopwatch.StartNew();
 
-			using IScope scope = GlobalTracer.Instance.BuildSpan("SearchLogDataAsync").StartActive();
-			scope.Span.SetTag("LogId", logFile.Id.ToString());
-			scope.Span.SetTag("Text", text);
-			scope.Span.SetTag("Count", count.ToString(CultureInfo.InvariantCulture));
+			using TelemetrySpan span = _tracer.StartActiveSpan($"{nameof(LogFileService)}.{nameof(SearchLogDataAsync)}");
+			span.SetAttribute("logId", logFile.Id.ToString());
+			span.SetAttribute("text", text);
+			span.SetAttribute("count", count);
 
 			List<int> results = new List<int>();
 			if (count > 0)
@@ -1984,8 +1984,8 @@ namespace Horde.Server.Logs
 				LogIndexData? indexData = await ReadIndexAsync(logFile, logFile.IndexLength.Value);
 				if(indexData != null && firstLine < indexData.LineCount)
 				{
-					using IScope indexScope = GlobalTracer.Instance.BuildSpan("Indexed").StartActive();
-					indexScope.Span.SetTag("LineCount", indexData.LineCount.ToString(CultureInfo.InvariantCulture));
+					using TelemetrySpan span = _tracer.StartActiveSpan($"{nameof(LogFileService)}.{nameof(SearchLogDataInternalAsync)}.Indexed");
+					span.SetAttribute("lineCount", indexData.LineCount);
 
 					foreach(int lineIndex in indexData.Search(firstLine, searchText, searchStats))
 					{
diff --git a/Engine/Source/Programs/Horde/Horde.Server/Perforce/PerforceService.cs b/Engine/Source/Programs/Horde/Horde.Server/Perforce/PerforceService.cs
--- a/Engine/Source/Programs/Horde/Horde.Server/Perforce/PerforceService.cs
+++ b/Engine/Source/Programs/Horde/Horde.Server/Perforce/PerforceService.cs
@@ -22,8 +22,7 @@ using Horde.Server.Utilities;
 using Microsoft.Extensions.Caching.Memory;
 using Microsoft.Extensions.Logging;
 using Microsoft.Extensions.Options;
-using OpenTracing;
-using OpenTracing.Util;
+using OpenTelemetry.Trace;
 
 namespace Horde.Server.Perforce
 {
@@ -231,6 +230,7 @@ namespace Horde.Server.Perforce
 
 		readonly PerforceLoadBalancer _loadBalancer;
 		readonly IOptionsMonitor<GlobalConfig> _globalConfig;
+		readonly Tracer _tracer;
 		readonly ILogger _logger;
 
 		// Useful overrides for local debugging with read-only data
@@ -248,12 +248,13 @@ namespace Horde.Server.Perforce
 		/// <summary>
 		/// Constructor
 		/// </summary>
-		public PerforceService(PerforceLoadBalancer loadBalancer, IUserCollection userCollection, IOptions<ServerSettings> settings, IOptionsMonitor<GlobalConfig> globalConfig, ILogger<PerforceService> logger)
+		public PerforceService(PerforceLoadBalancer loadBalancer, IUserCollection userCollection, IOptions<ServerSettings> settings, IOptionsMonitor<GlobalConfig> globalConfig, Tracer tracer, ILogger<PerforceService> logger)
 		{
 			_loadBalancer = loadBalancer;
 			_userCollection = userCollection;
 			_settings = settings.Value;
 			_globalConfig = globalConfig;
+			_tracer = tracer;
 			_logger = logger;
 
 			if(settings.Value.UseLocalPerforceEnv)
@@ -521,9 +522,9 @@ namespace Horde.Server.Perforce
 
 		async Task<Credentials> GetTicketAsync(PerforceCluster cluster, string userName, CancellationToken cancellationToken)
 		{
-			using IScope scope = GlobalTracer.Instance.BuildSpan("PerforceService.TryGetImpersonatedConnectionSettings").StartActive();
-			scope.Span.SetTag("ClusterName", cluster.Name);
-			scope.Span.SetTag("UserName", userName);
+			using TelemetrySpan span = _tracer.StartActiveSpan($"{nameof(PerforceService)}.{nameof(GetTicketAsync)}");
+			span.SetAttribute("clusterName", cluster.Name);
+			span.SetAttribute("userName", userName);
 
 			Credentials? ticketInfo = null;
 
@@ -579,9 +580,9 @@ namespace Horde.Server.Perforce
 
 		public async ValueTask<IUser> FindOrAddUserAsync(PerforceCluster cluster, string userName, CancellationToken cancellationToken = default)
 		{
-			using IScope scope = GlobalTracer.Instance.BuildSpan("PerforceService.FindOrAddUserAsync").StartActive();
-			scope.Span.SetTag("ClusterName", cluster.Name);
-			scope.Span.SetTag("UserName", userName);
+			using TelemetrySpan span = _tracer.StartActiveSpan($"{nameof(PerforceService)}.{nameof(FindOrAddUserAsync)}");
+			span.SetAttribute("clusterName", cluster.Name);
+			span.SetAttribute("userName", userName);
 			
 			IUser? user;
 			if (!_userCache.TryGetValue((cluster.Name, userName), out user))
@@ -617,8 +618,8 @@ namespace Horde.Server.Perforce
 
 		async Task<IPerforceServer> GetServerAsync(PerforceCluster cluster)
 		{
-			using IScope scope = GlobalTracer.Instance.BuildSpan("PerforceService.SelectServer").StartActive();
-			scope.Span.SetTag("ClusterName", cluster.Name);
+			using TelemetrySpan span = _tracer.StartActiveSpan($"{nameof(PerforceService)}.{nameof(GetServerAsync)}");
+			span.SetAttribute("clusterName", cluster.Name);
 
 			IPerforceServer? server = await _loadBalancer.SelectServerAsync(cluster);
 			if (server == null)
@@ -743,10 +744,10 @@ namespace Horde.Server.Perforce
 		/// <inheritdoc/>
 		public async Task<ICommit> GetChangeDetailsAsync(StreamConfig streamConfig, int changeNumber, CancellationToken cancellationToken)
 		{
-			using IScope scope = GlobalTracer.Instance.BuildSpan("PerforceService.GetChangeDetailsAsync").StartActive();
-			scope.Span.SetTag("ClusterName", streamConfig.ClusterName);
-			scope.Span.SetTag("StreamName", streamConfig.Name);
-			scope.Span.SetTag("ChangeNumber", changeNumber);
+			using TelemetrySpan span = _tracer.StartActiveSpan($"{nameof(PerforceService)}.{nameof(GetChangeDetailsAsync)}");
+			span.SetAttribute("clusterName", streamConfig.ClusterName);
+			span.SetAttribute("streamName", streamConfig.Name);
+			span.SetAttribute("changeNumber", changeNumber);
 
 			PerforceCluster cluster = _globalConfig.CurrentValue.GetPerforceCluster(streamConfig.ClusterName);
 
@@ -767,10 +768,10 @@ namespace Horde.Server.Perforce
 		/// <inheritdoc/>
 		public async Task<(CheckShelfResult, ShelfInfo?)> CheckShelfAsync(StreamConfig streamConfig, int changeNumber, CancellationToken cancellationToken)
 		{
-			using IScope scope = GlobalTracer.Instance.BuildSpan("PerforceService.CheckPreflightAsync").StartActive();
-			scope.Span.SetTag("ClusterName", streamConfig.ClusterName);
-			scope.Span.SetTag("StreamName", streamConfig.Name);
-			scope.Span.SetTag("ChangeNumber", changeNumber);
+			using TelemetrySpan span = _tracer.StartActiveSpan($"{nameof(PerforceService)}.{nameof(CheckShelfAsync)}");
+			span.SetAttribute("clusterName", streamConfig.ClusterName);
+			span.SetAttribute("streamName", streamConfig.Name);
+			span.SetAttribute("changeNumber", changeNumber);
 			
 			PerforceCluster cluster = _globalConfig.CurrentValue.GetPerforceCluster(streamConfig.ClusterName);
 
@@ -840,9 +841,9 @@ namespace Horde.Server.Perforce
 		/// <inheritdoc/>
 		public async Task DeleteShelvedChangeAsync(string clusterName, int shelvedChange, CancellationToken cancellationToken)
 		{
-			using IScope scope = GlobalTracer.Instance.BuildSpan("PerforceService.DeleteShelvedChangeAsync").StartActive();
-			scope.Span.SetTag("ClusterName", clusterName);
-			scope.Span.SetTag("ShelvedChange", shelvedChange);
+			using TelemetrySpan span = _tracer.StartActiveSpan($"{nameof(PerforceService)}.{nameof(DeleteShelvedChangeAsync)}");
+			span.SetAttribute("clusterName", clusterName);
+			span.SetAttribute("shelvedChange", shelvedChange);
 			
 			PerforceCluster cluster = _globalConfig.CurrentValue.GetPerforceCluster(clusterName);
 
@@ -856,9 +857,9 @@ namespace Horde.Server.Perforce
 		/// <inheritdoc/>
 		public async Task UpdateChangelistDescriptionAsync(string clusterName, int change, Func<string, string> updateFunc, CancellationToken cancellationToken)
 		{
-			using IScope scope = GlobalTracer.Instance.BuildSpan("PerforceService.UpdateChangelistDescription").StartActive();
-			scope.Span.SetTag("ClusterName", clusterName);
-			scope.Span.SetTag("Change", change);
+			using TelemetrySpan span = _tracer.StartActiveSpan($"{nameof(PerforceService)}.{nameof(UpdateChangelistDescriptionAsync)}");
+			span.SetAttribute("clusterName", clusterName);
+			span.SetAttribute("change", change);
 
 			try
 			{
@@ -899,10 +900,10 @@ namespace Horde.Server.Perforce
 		/// <inheritdoc/>
 		public async Task<int> CreateNewChangeAsync(string clusterName, string streamName, string filePath, string description, CancellationToken cancellationToken)
 		{
-			using IScope scope = GlobalTracer.Instance.BuildSpan("PerforceService.CreateNewChangeAsync").StartActive();
-			scope.Span.SetTag("ClusterName", clusterName);
-			scope.Span.SetTag("StreamName", streamName);
-			scope.Span.SetTag("FilePath", filePath);
+			using TelemetrySpan span = _tracer.StartActiveSpan($"{nameof(PerforceService)}.{nameof(CreateNewChangeAsync)}");
+			span.SetAttribute("clusterName", clusterName);
+			span.SetAttribute("streamName", streamName);
+			span.SetAttribute("filePath", filePath);
 			
 			using (PooledConnectionHandle perforce = await ConnectWithStreamClientAsync(clusterName, null, streamName, cancellationToken))
 			{
@@ -1034,10 +1035,10 @@ namespace Horde.Server.Perforce
 		/// <inheritdoc/>
 		public async Task<(int? Change, string Message)> SubmitShelvedChangeAsync(StreamConfig streamConfig, int change, int originalChange, CancellationToken cancellationToken)
 		{
-			using IScope scope = GlobalTracer.Instance.BuildSpan("PerforceService.SubmitShelvedChangeAsync").StartActive();
-			scope.Span.SetTag("Stream", streamConfig.Id);
-			scope.Span.SetTag("Change", change);
-			scope.Span.SetTag("OriginalChange", originalChange);
+			using TelemetrySpan span = _tracer.StartActiveSpan($"{nameof(PerforceService)}.{nameof(SubmitShelvedChangeAsync)}");
+			span.SetAttribute("stream", streamConfig.Id);
+			span.SetAttribute("change", change);
+			span.SetAttribute("originalChange", originalChange);
 
 			PerforceCluster cluster = _globalConfig.CurrentValue.GetPerforceCluster(streamConfig.ClusterName);
 
@@ -1173,11 +1174,11 @@ namespace Horde.Server.Perforce
 
 			public virtual async IAsyncEnumerable<ICommit> FindAsync(int? minChange, int? maxChange, int? maxResults, IReadOnlyList<CommitTag>? tags, [EnumeratorCancellation] CancellationToken cancellationToken = default)
 			{
-				using IScope scope = GlobalTracer.Instance.BuildSpan("PerforceService.CommitSource.FindAsync").StartActive();
-				scope.Span.SetTag("ClusterName", _streamConfig.ClusterName);
-				scope.Span.SetTag("MinChange", minChange ?? -2);
-				scope.Span.SetTag("MaxChange", maxChange ?? -2);
-				scope.Span.SetTag("MaxResults", maxResults ?? -1);
+				using TelemetrySpan span = _perforceService._tracer.StartActiveSpan($"{nameof(PerforceService)}.{nameof(CommitSource)}.{nameof(FindAsync)}");
+				span.SetAttribute("stream", _streamConfig.ClusterName);
+				span.SetAttribute("minChange", minChange ?? -2);
+				span.SetAttribute("maxChange", maxChange ?? -2);
+				span.SetAttribute("maxResults", maxResults ?? -1);
 
 				using (PooledConnectionHandle perforce = await _perforceService.ConnectWithStreamClientAsync(_streamConfig, null, cancellationToken))
 				{
diff --git a/Engine/Source/Programs/Horde/Horde.Server/Perforce/PerforceServiceCache.cs b/Engine/Source/Programs/Horde/Horde.Server/Perforce/PerforceServiceCache.cs
--- a/Engine/Source/Programs/Horde/Horde.Server/Perforce/PerforceServiceCache.cs
+++ b/Engine/Source/Programs/Horde/Horde.Server/Perforce/PerforceServiceCache.cs
@@ -22,6 +22,7 @@ using MongoDB.Bson;
 using MongoDB.Bson.Serialization.Attributes;
 using MongoDB.Bson.Serialization.Options;
 using MongoDB.Driver;
+using OpenTelemetry.Trace;
 using StackExchange.Redis;
 
 namespace Horde.Server.Perforce
@@ -192,8 +193,8 @@ namespace Horde.Server.Perforce
 		/// <summary>
 		/// Constructor
 		/// </summary>
-		public PerforceServiceCache(PerforceLoadBalancer loadBalancer, MongoService mongoService, RedisService redisService, IDowntimeService downtimeService, IUserCollection userCollection, IClock clock, IOptions<ServerSettings> settings, IOptionsMonitor<GlobalConfig> globalConfig, ILogger<PerforceService> logger)
-			: base(loadBalancer, userCollection, settings, globalConfig, logger)
+		public PerforceServiceCache(PerforceLoadBalancer loadBalancer, MongoService mongoService, RedisService redisService, IDowntimeService downtimeService, IUserCollection userCollection, IClock clock, IOptions<ServerSettings> settings, IOptionsMonitor<GlobalConfig> globalConfig, Tracer tracer, ILogger<PerforceService> logger)
+			: base(loadBalancer, userCollection, settings, globalConfig, tracer, logger)
 		{
 			_mongoService = mongoService;
 			_redisService = redisService;
diff --git a/Engine/Source/Programs/Horde/Horde.Server/Program.cs b/Engine/Source/Programs/Horde/Horde.Server/Program.cs
--- a/Engine/Source/Programs/Horde/Horde.Server/Program.cs
+++ b/Engine/Source/Programs/Horde/Horde.Server/Program.cs
@@ -13,13 +13,8 @@ using Horde.Server.Utilities;
 using Microsoft.Extensions.Configuration;
 using Microsoft.Extensions.DependencyInjection;
 using Microsoft.Extensions.Hosting;
-using OpenTracing;
-using OpenTracing.Propagation;
-using OpenTracing.Util;
 using Serilog;
 using Serilog.Configuration;
-using Serilog.Core;
-using Serilog.Events;
 using Serilog.Exceptions;
 using Serilog.Exceptions.Core;
 using Serilog.Exceptions.Grpc.Destructurers;
@@ -57,59 +52,11 @@ namespace Horde.Server
 			{
 				configuration = configuration.Enrich.With<OpenTelemetryDatadogLogEnricher>();
 			}
-			else if (settings.WithDatadog)
-			{
-				// OpenTracing is replaced with OpenTelemetry. To be removed.
-				configuration = configuration.Enrich.With<OpenTracingDatadogLogEnricher>();
-			}
 
 			return configuration;
 		}
 	}
 
-	class OpenTracingDatadogLogEnricher : ILogEventEnricher
-	{
-		public void Enrich(Serilog.Events.LogEvent logEvent, ILogEventPropertyFactory propertyFactory)
-		{
-			ISpan? span = GlobalTracer.Instance?.ActiveSpan;
-			if (span != null)
-			{
-				logEvent.AddPropertyIfAbsent(propertyFactory.CreateProperty("dd.trace_id", span.Context.TraceId));
-				logEvent.AddPropertyIfAbsent(propertyFactory.CreateProperty("dd.span_id", span.Context.SpanId));
-			}
-		}
-	}
-
-	class TestTracer : ITracer
-	{
-		readonly ITracer _inner;
-
-		public TestTracer(ITracer inner)
-		{
-			_inner = inner;
-		}
-
-		public IScopeManager ScopeManager => _inner.ScopeManager;
-
-		public ISpan ActiveSpan => _inner.ActiveSpan;
-
-		public ISpanBuilder BuildSpan(string operationName)
-		{
-			Serilog.Log.Debug("Creating span {Name}", operationName);
-			return _inner.BuildSpan(operationName);
-		}
-
-		public ISpanContext Extract<TCarrier>(IFormat<TCarrier> format, TCarrier carrier)
-		{
-			return _inner.Extract<TCarrier>(format, carrier);
-		}
-
-		public void Inject<TCarrier>(ISpanContext spanContext, IFormat<TCarrier> format, TCarrier carrier)
-		{
-			_inner.Inject<TCarrier>(spanContext, format, carrier);
-		}
-	}
-
 	class Program
 	{
 		public static SemVer Version => s_version;
@@ -178,12 +125,6 @@ namespace Horde.Server
 
 			Serilog.Log.Logger.Information("Server version: {Version}", Version);
 
-			if (hordeSettings.WithDatadog)
-			{
-				GlobalTracer.Register(Datadog.Trace.OpenTracing.OpenTracingTracerFactory.WrapTracer(Datadog.Trace.Tracer.Instance));
-				Serilog.Log.Logger.Information("Enabling datadog tracing (OpenTrace)");
-			}
-
 			ServiceCollection services = new ServiceCollection();
 			services.AddCommandsFromAssembly(Assembly.GetExecutingAssembly());
 			services.AddLogging(builder => builder.AddSerilog());
diff --git a/Engine/Source/Programs/Horde/Horde.Server/Startup.cs b/Engine/Source/Programs/Horde/Horde.Server/Startup.cs
--- a/Engine/Source/Programs/Horde/Horde.Server/Startup.cs
+++ b/Engine/Source/Programs/Horde/Horde.Server/Startup.cs
@@ -80,8 +80,6 @@ using MongoDB.Bson;
 using MongoDB.Bson.Serialization;
 using MongoDB.Bson.Serialization.Conventions;
 using MongoDB.Bson.Serialization.Serializers;
-using OpenTracing.Contrib.Grpc.Interceptors;
-using OpenTracing.Util;
 using Serilog;
 using Serilog.Events;
 using StatsdClient;
@@ -98,6 +96,7 @@ using Horde.Server.Notifications.Sinks;
 using StatusCode = Grpc.Core.StatusCode;
 using Horde.Server.Artifacts;
 using Horde.Server.Compute;
+using OpenTelemetry.Trace;
 
 namespace Horde.Server
 {
@@ -378,7 +377,6 @@ namespace Horde.Server
 				options.EnableDetailedErrors = true;
 				options.MaxReceiveMessageSize = 200 * 1024 * 1024; // 100 MB (packaged builds of Horde agent can be large) 
 				options.Interceptors.Add(typeof(GrpcExceptionInterceptor));
-				options.Interceptors.Add<ServerTracingInterceptor>(GlobalTracer.Instance);
 			});
 			services.AddGrpcReflection();
 
@@ -848,7 +846,10 @@ namespace Horde.Server
 				}
 				else
 				{
-					return new RedisLogBuilder(redisService.ConnectionPool, provider.GetRequiredService<ILogger<RedisLogBuilder>>());
+					return new RedisLogBuilder(
+						redisService.ConnectionPool,
+						provider.GetRequiredService<Tracer>(),
+						provider.GetRequiredService<ILogger<RedisLogBuilder>>());
 				}
 			});
 
diff --git a/Engine/Source/Programs/Horde/Horde.Server/Storage/Backends/AwsStorageBackend.cs b/Engine/Source/Programs/Horde/Horde.Server/Storage/Backends/AwsStorageBackend.cs
--- a/Engine/Source/Programs/Horde/Horde.Server/Storage/Backends/AwsStorageBackend.cs
+++ b/Engine/Source/Programs/Horde/Horde.Server/Storage/Backends/AwsStorageBackend.cs
@@ -17,8 +17,7 @@ using Google.Protobuf.WellKnownTypes;
 using Horde.Server.Utilities;
 using Microsoft.Extensions.Configuration;
 using Microsoft.Extensions.Logging;
-using OpenTracing;
-using OpenTracing.Util;
+using OpenTelemetry.Trace;
 
 namespace Horde.Server.Storage.Backends
 {
@@ -271,8 +270,8 @@ namespace Horde.Server.Storage.Backends
 
 		async Task<Stream?> TryReadAsync(string path, ByteRange? byteRange, CancellationToken cancellationToken)
 		{
-			using IScope scope = GlobalTracer.Instance.BuildSpan("AwsStorageBackend.ReadAsync").StartActive();
-			scope.Span.SetTag("Path", path);
+			using TelemetrySpan span = OpenTelemetryTracers.Horde.StartActiveSpan($"{nameof(AwsStorageBackend)}.{nameof(TryReadAsync)}");
+			span.SetAttribute("path", path);
 			
 			string fullPath = GetFullPath(path);
 
@@ -313,8 +312,8 @@ namespace Horde.Server.Storage.Backends
 		/// </summary>
 		Uri? GetPresignedUrl(string path, HttpVerb verb)
 		{
-			using IScope scope = GlobalTracer.Instance.BuildSpan("AwsStorageBackend.TryGetHttpRedirectUrl").StartActive();
-			scope.Span.SetTag("Path", path);
+			using TelemetrySpan span = OpenTelemetryTracers.Horde.StartActiveSpan($"{nameof(AwsStorageBackend)}.{nameof(GetPresignedUrl)}");
+			span.SetAttribute("path", path);
 
 			string fullPath = GetFullPath(path);
 
@@ -372,8 +371,8 @@ namespace Horde.Server.Storage.Backends
 		/// <inheritdoc/>
 		public async Task WriteInternalAsync(string fullPath, Stream stream, CancellationToken cancellationToken)
 		{
-			using IScope scope = GlobalTracer.Instance.BuildSpan("AwsStorageBackend.WriteInternalAsync").StartActive();
-			scope.Span.SetTag("Path", fullPath);
+			using TelemetrySpan span = OpenTelemetryTracers.Horde.StartActiveSpan($"{nameof(AwsStorageBackend)}.{nameof(WriteInternalAsync)}");
+			span.SetAttribute("path", fullPath);
 			
 			const int MinPartSize = 5 * 1024 * 1024;
 
@@ -481,8 +480,8 @@ namespace Horde.Server.Storage.Backends
 		/// <inheritdoc/>
 		public async Task DeleteAsync(string path, CancellationToken cancellationToken)
 		{
-			using IScope scope = GlobalTracer.Instance.BuildSpan("AwsStorageBackend.DeleteAsync").StartActive();
-			scope.Span.SetTag("Path", path);
+			using TelemetrySpan span = OpenTelemetryTracers.Horde.StartActiveSpan($"{nameof(AwsStorageBackend)}.{nameof(DeleteAsync)}");
+			span.SetAttribute("path", path);
 			
 			DeleteObjectRequest newDeleteRequest = new DeleteObjectRequest();
 			newDeleteRequest.BucketName = _options.AwsBucketName;
@@ -493,8 +492,8 @@ namespace Horde.Server.Storage.Backends
 		/// <inheritdoc/>
 		public async Task<bool> ExistsAsync(string path, CancellationToken cancellationToken)
 		{
-			using IScope scope = GlobalTracer.Instance.BuildSpan("AwsStorageBackend.ExistsAsync").StartActive();
-			scope.Span.SetTag("Path", path);
+			using TelemetrySpan span = OpenTelemetryTracers.Horde.StartActiveSpan($"{nameof(AwsStorageBackend)}.{nameof(ExistsAsync)}");
+			span.SetAttribute("path", path);
 			
 			try
 			{
diff --git a/Engine/Source/Programs/Horde/Horde.Server/Storage/Backends/FileStorageBackend.cs b/Engine/Source/Programs/Horde/Horde.Server/Storage/Backends/FileStorageBackend.cs
--- a/Engine/Source/Programs/Horde/Horde.Server/Storage/Backends/FileStorageBackend.cs
+++ b/Engine/Source/Programs/Horde/Horde.Server/Storage/Backends/FileStorageBackend.cs
@@ -7,9 +7,8 @@ using System.Runtime.CompilerServices;
 using System.Threading;
 using System.Threading.Tasks;
 using EpicGames.Core;
-using EpicGames.Horde.Storage;
-using OpenTracing;
-using OpenTracing.Util;
+using Horde.Server.Utilities;
+using OpenTelemetry.Trace;
 
 namespace Horde.Server.Storage.Backends
 {
@@ -66,8 +65,8 @@ namespace Horde.Server.Storage.Backends
 		/// <inheritdoc/>
 		public Task<Stream?> TryReadAsync(string path, CancellationToken cancellationToken)
 		{
-			using IScope scope = GlobalTracer.Instance.BuildSpan("FileSystemStorageBackend.ReadAsync").StartActive();
-			scope.Span.SetTag("Path", path);
+			using TelemetrySpan span = OpenTelemetryTracers.Horde.StartActiveSpan($"{nameof(FileSystemStorageBackend)}.{nameof(TryReadAsync)}");
+			span.SetAttribute("path", path);
 			
 			FileReference location = FileReference.Combine(_baseDir, path);
 			if (!FileReference.Exists(location))
@@ -103,8 +102,8 @@ namespace Horde.Server.Storage.Backends
 		/// <inheritdoc/>
 		public async Task WriteAsync(string path, Stream stream, CancellationToken cancellationToken)
 		{
-			using IScope scope = GlobalTracer.Instance.BuildSpan("FileSystemStorageBackend.WriteAsync").StartActive();
-			scope.Span.SetTag("Path", path);
+			using TelemetrySpan span = OpenTelemetryTracers.Horde.StartActiveSpan($"{nameof(FileSystemStorageBackend)}.{nameof(WriteAsync)}");
+			span.SetAttribute("path", path);
 			
 			FileReference finalLocation = FileReference.Combine(_baseDir, path);
 			if (!FileReference.Exists(finalLocation))
@@ -142,8 +141,8 @@ namespace Horde.Server.Storage.Backends
 		/// <inheritdoc/>
 		public Task<bool> ExistsAsync(string path, CancellationToken cancellationToken)
 		{
-			using IScope scope = GlobalTracer.Instance.BuildSpan("FileSystemStorageBackend.ExistsAsync").StartActive();
-			scope.Span.SetTag("Path", path);
+			using TelemetrySpan span = OpenTelemetryTracers.Horde.StartActiveSpan($"{nameof(FileSystemStorageBackend)}.{nameof(ExistsAsync)}");
+			span.SetAttribute("path", path);
 			
 			FileReference location = FileReference.Combine(_baseDir, path);
 			return Task.FromResult(FileReference.Exists(location));
@@ -152,8 +151,8 @@ namespace Horde.Server.Storage.Backends
 		/// <inheritdoc/>
 		public Task DeleteAsync(string path, CancellationToken cancellationToken)
 		{
-			using IScope scope = GlobalTracer.Instance.BuildSpan("FileSystemStorageBackend.DeleteAsync").StartActive();
-			scope.Span.SetTag("Path", path);
+			using TelemetrySpan span = OpenTelemetryTracers.Horde.StartActiveSpan($"{nameof(FileSystemStorageBackend)}.{nameof(DeleteAsync)}");
+			span.SetAttribute("path", path);
 			
 			FileReference location = FileReference.Combine(_baseDir, path);
 			FileReference.Delete(location);
diff --git a/Engine/Source/Programs/Horde/Horde.Server/Streams/StreamsController.cs b/Engine/Source/Programs/Horde/Horde.Server/Streams/StreamsController.cs
--- a/Engine/Source/Programs/Horde/Horde.Server/Streams/StreamsController.cs
+++ b/Engine/Source/Programs/Horde/Horde.Server/Streams/StreamsController.cs
@@ -16,8 +16,7 @@ using Horde.Server.Utilities;
 using Microsoft.AspNetCore.Authorization;
 using Microsoft.AspNetCore.Mvc;
 using Microsoft.Extensions.Options;
-using OpenTracing;
-using OpenTracing.Util;
+using OpenTelemetry.Trace;
 using TimeZoneConverter;
 
 namespace Horde.Server.Streams
@@ -35,19 +34,21 @@ namespace Horde.Server.Streams
 		private readonly ITemplateCollection _templateCollection;
 		private readonly IJobStepRefCollection _jobStepRefCollection;
 		private readonly IUserCollection _userCollection;
+		private readonly Tracer _tracer;
 		private readonly IOptionsSnapshot<GlobalConfig> _globalConfig;
 		private readonly TimeZoneInfo _timeZone;
 
 		/// <summary>
 		/// Constructor
 		/// </summary>
-		public StreamsController(IStreamCollection streamCollection, ICommitService commitService, ITemplateCollection templateCollection, IJobStepRefCollection jobStepRefCollection, IUserCollection userCollection, IOptionsSnapshot<GlobalConfig> globalConfig)
+		public StreamsController(IStreamCollection streamCollection, ICommitService commitService, ITemplateCollection templateCollection, IJobStepRefCollection jobStepRefCollection, IUserCollection userCollection, Tracer tracer, IOptionsSnapshot<GlobalConfig> globalConfig)
 		{
 			_streamCollection = streamCollection;
 			_commitService = commitService;
 			_templateCollection = templateCollection;
 			_jobStepRefCollection = jobStepRefCollection;
 			_userCollection = userCollection;
+			_tracer = tracer;
 			_globalConfig = globalConfig;
 
 			string? timeZoneName = _globalConfig.Value.ServerSettings.ScheduleTimeZone;
@@ -179,14 +180,14 @@ namespace Horde.Server.Streams
 		/// <returns>Response object</returns>
 		async Task<GetStreamResponse> CreateGetStreamResponseAsync(IStream stream)
 		{
-			using IScope scope = GlobalTracer.Instance.BuildSpan("CreateGetStreamResponse").StartActive();
-			scope.Span.SetTag("streamId", stream.Id);
+			using TelemetrySpan span = _tracer.StartActiveSpan($"{nameof(StreamsController)}.{nameof(CreateGetStreamResponseAsync)}");
+			span.SetAttribute("streamId", stream.Id);
 
 			List<GetTemplateRefResponse> apiTemplateRefs = new List<GetTemplateRefResponse>();
 			foreach (KeyValuePair<TemplateId, ITemplateRef> pair in stream.Templates)
 			{
-				using IScope templateScope = GlobalTracer.Instance.BuildSpan("CreateGetStreamResponse.Template").StartActive();
-				templateScope.Span.SetTag("templateName", pair.Value.Config.Name);
+				using TelemetrySpan templateScope = _tracer.StartActiveSpan($"{nameof(StreamsController)}.{nameof(CreateGetStreamResponseAsync)}.Template");
+				templateScope.SetAttribute("templateName", pair.Value.Config.Name);
 
 				ITemplateRef templateRef = pair.Value;
 				if (templateRef.Config.Authorize(StreamAclAction.ViewTemplate, User))
diff --git a/Engine/Source/Programs/Horde/Horde.Server/Utilities/CompressionExtensions.cs b/Engine/Source/Programs/Horde/Horde.Server/Utilities/CompressionExtensions.cs
--- a/Engine/Source/Programs/Horde/Horde.Server/Utilities/CompressionExtensions.cs
+++ b/Engine/Source/Programs/Horde/Horde.Server/Utilities/CompressionExtensions.cs
@@ -2,12 +2,10 @@
 
 using System;
 using System.Buffers.Binary;
-using System.Globalization;
 using System.IO;
 using EpicGames.Core;
 using ICSharpCode.SharpZipLib.BZip2;
-using OpenTracing;
-using OpenTracing.Util;
+using OpenTelemetry.Trace;
 
 namespace Horde.Server.Utilities
 {
@@ -22,8 +20,8 @@ namespace Horde.Server.Utilities
 		/// <returns>The compressed data</returns>
 		public static byte[] CompressBzip2(this ReadOnlyMemory<byte> memory)
 		{
-			using IScope scope = GlobalTracer.Instance.BuildSpan("CompressBzip2").StartActive();
-			scope.Span.SetTag("DecompressedSize", memory.Length.ToString(CultureInfo.InvariantCulture));
+			using TelemetrySpan span = OpenTelemetryTracers.Horde.StartActiveSpan($"{nameof(CompressionExtensions)}.{nameof(CompressBzip2)}");
+			span.SetAttribute("decompressedSize", memory.Length);
 
 			byte[] compressedData;
 			using (MemoryStream stream = new MemoryStream())
@@ -40,7 +38,7 @@ namespace Horde.Server.Utilities
 				compressedData = stream.ToArray();
 			}
 
-			scope.Span.SetTag("CompressedSize", compressedData.Length.ToString(CultureInfo.InvariantCulture));
+			span.SetAttribute("compressedSize", compressedData.Length);
 			return compressedData;
 		}
 
@@ -52,9 +50,9 @@ namespace Horde.Server.Utilities
 		{
 			int decompressedSize = BinaryPrimitives.ReadInt32LittleEndian(memory.Span);
 
-			using IScope scope = GlobalTracer.Instance.BuildSpan("DecompressBzip2").StartActive();
-			scope.Span.SetTag("CompressedSize", memory.Length.ToString(CultureInfo.InvariantCulture));
-			scope.Span.SetTag("DecompressedSize", decompressedSize.ToString(CultureInfo.InvariantCulture));
+			using TelemetrySpan span = OpenTelemetryTracers.Horde.StartActiveSpan($"{nameof(CompressionExtensions)}.{nameof(DecompressBzip2)}");
+			span.SetAttribute("compressedSize", memory.Length);
+			span.SetAttribute("decompressedSize", decompressedSize);
 
 			byte[] data = new byte[decompressedSize];
 			using (ReadOnlyMemoryStream stream = new ReadOnlyMemoryStream(memory.Slice(4)))
diff --git a/Engine/Source/Programs/Horde/Horde.Server/Utilities/OpenTelemetry.cs b/Engine/Source/Programs/Horde/Horde.Server/Utilities/OpenTelemetry.cs
--- a/Engine/Source/Programs/Horde/Horde.Server/Utilities/OpenTelemetry.cs
+++ b/Engine/Source/Programs/Horde/Horde.Server/Utilities/OpenTelemetry.cs
@@ -216,6 +216,12 @@ public static class OpenTelemetryHelper
 /// </summary>
 public static class OpenTelemetryTracers
 {
+	/// <summary>
+	/// Name of resource name attribute in Datadog
+	/// Some traces use this for prettier display inside their UI
+	/// </summary>
+	public const string DatadogResourceAttribute = "resource.name";
+	
 	/// <summary>
 	/// Name of default Horde tracer (aka activity source)
 	/// </summary>
@@ -234,11 +240,13 @@ public static class OpenTelemetryTracers
 	
 	/// <summary>
 	/// Default tracer used in Horde
+	/// Prefer dependency-injected tracer over this static member.
 	/// </summary>
 	public static readonly Tracer Horde = TracerProvider.Default.GetTracer(HordeName);
 	
 	/// <summary>
 	/// Tracer specific to MongoDB
+	/// Prefer StartMongoDbSpan static extension for Tracer.
 	/// </summary>
 	public static readonly Tracer MongoDb = TracerProvider.Default.GetTracer(MongoDbName);
 }
