Horde: Additional Include *.verse as a known source file extension.

diff --git a/Engine/Source/Programs/Horde/Horde.Server/Issues/SuspectChange.cs b/Engine/Source/Programs/Horde/Horde.Server/Issues/SuspectChange.cs
--- a/Engine/Source/Programs/Horde/Horde.Server/Issues/SuspectChange.cs
+++ b/Engine/Source/Programs/Horde/Horde.Server/Issues/SuspectChange.cs
@@ -33,7 +33,8 @@ namespace Horde.Server.Issues
 			".ush",
 			".uproject",
 			".uplugin",
-			".sln"
+			".sln",
+			".verse"
 		};
 
 		/// <summary>
