Horde: Reduce initial poill for live log
#rnx
#skipci

diff --git a/Engine/Source/Programs/Horde/HordeDashboard/src/components/LogSource.tsx b/Engine/Source/Programs/Horde/HordeDashboard/src/components/LogSource.tsx
--- a/Engine/Source/Programs/Horde/HordeDashboard/src/components/LogSource.tsx
+++ b/Engine/Source/Programs/Horde/HordeDashboard/src/components/LogSource.tsx
@@ -67,7 +67,7 @@ export abstract class LogSource {
 
    initComplete() {
       if (this.active) {
-         setTimeout(() => { this.poll(); }, 2000);
+         setTimeout(() => { this.poll(); }, 250);
       }
    }
 
