Horde: Fix MongoDB upgrade tasks not correctly starting.

diff --git a/Engine/Source/Programs/Horde/Horde.Server/Server/MongoService.cs b/Engine/Source/Programs/Horde/Horde.Server/Server/MongoService.cs
--- a/Engine/Source/Programs/Horde/Horde.Server/Server/MongoService.cs
+++ b/Engine/Source/Programs/Horde/Horde.Server/Server/MongoService.cs
@@ -959,16 +959,20 @@ namespace Horde.Server.Server
 		{
 			for (; ; )
 			{
-				Func<CancellationToken, Task> updateIndexTask = await _mongoService.ReadNextUpgradeTask(cancellationToken);
-				for (; ; )
+				try
 				{
-					using (RedisLock schemaLock = new (_redisService.GetDatabase(), s_schemaLockKey))
+					Func<CancellationToken, Task> updateIndexTask = await _mongoService.ReadNextUpgradeTask(cancellationToken);
+					using (RedisLock schemaLock = new(_redisService.GetDatabase(), s_schemaLockKey))
 					{
 						if (await schemaLock.AcquireAsync(TimeSpan.FromMinutes(5.0)) && await UpdateOneAsync(updateIndexTask, cancellationToken))
 						{
 							break;
 						}
 					}
+				}
+				catch (Exception ex)
+				{
+					_logger.LogError(ex, "Exception while running upgrade task: {Message}", ex.Message);
 					await Task.Delay(TimeSpan.FromMinutes(1.0), cancellationToken);
 				}
 			}
