Horde: Revised fix for indexes not being created.

diff --git a/Engine/Source/Programs/Horde/Horde.Server/Server/MongoService.cs b/Engine/Source/Programs/Horde/Horde.Server/Server/MongoService.cs
--- a/Engine/Source/Programs/Horde/Horde.Server/Server/MongoService.cs
+++ b/Engine/Source/Programs/Horde/Horde.Server/Server/MongoService.cs
@@ -959,26 +959,23 @@ namespace Horde.Server.Server
 		{
 			for (; ; )
 			{
-				try
+				Func<CancellationToken, Task> updateIndexTask = await _mongoService.ReadNextUpgradeTask(cancellationToken);
+				for (; ; )
 				{
-					Func<CancellationToken, Task> updateIndexTask = await _mongoService.ReadNextUpgradeTask(cancellationToken);
-					using (RedisLock schemaLock = new(_redisService.GetDatabase(), s_schemaLockKey))
+					using (RedisLock schemaLock = new (_redisService.GetDatabase(), s_schemaLockKey))
 					{
-						if (await schemaLock.AcquireAsync(TimeSpan.FromMinutes(5.0)) && await UpdateOneAsync(updateIndexTask, cancellationToken))
+						if (await schemaLock.AcquireAsync(TimeSpan.FromMinutes(5.0)))
 						{
+							await UpdateOneAsync(updateIndexTask, cancellationToken);
 							break;
 						}
 					}
-				}
-				catch (Exception ex)
-				{
-					_logger.LogError(ex, "Exception while running upgrade task: {Message}", ex.Message);
 					await Task.Delay(TimeSpan.FromMinutes(1.0), cancellationToken);
 				}
 			}
 		}
 
-		async ValueTask<bool> UpdateOneAsync(Func<CancellationToken, Task> taskAsync, CancellationToken cancellationToken)
+		async ValueTask UpdateOneAsync(Func<CancellationToken, Task> taskAsync, CancellationToken cancellationToken)
 		{
 			try
 			{
@@ -991,7 +988,7 @@ namespace Horde.Server.Server
 						SemVer currentVersion = SemVer.Parse(currentSchema.Version);
 						if (Program.Version < currentVersion)
 						{
-							return false;
+							return;
 						}
 						if (Program.Version == currentVersion)
 						{
@@ -1012,17 +1009,14 @@ namespace Horde.Server.Server
 
 				// Perform the upgrade
 				await taskAsync(cancellationToken);
-				return true;
 			}
 			catch (MongoCommandException ex)
 			{
 				_logger.LogWarning(ex, "Command exception while attempting to update indexes ({Code})", ex.Code);
-				return false;
 			}
 			catch (Exception ex)
 			{
 				_logger.LogWarning(ex, "Exception while attempting to update indexes");
-				return false;
 			}
 		}
 	}
