Horde: Better logging for schema upgrades.

diff --git a/Engine/Source/Programs/Horde/Horde.Server/Server/MongoService.cs b/Engine/Source/Programs/Horde/Horde.Server/Server/MongoService.cs
--- a/Engine/Source/Programs/Horde/Horde.Server/Server/MongoService.cs
+++ b/Engine/Source/Programs/Horde/Horde.Server/Server/MongoService.cs
@@ -686,6 +686,8 @@ namespace Horde.Server.Server
 
 		private async Task UpdateIndexesAsync<T>(string collectionName, IMongoCollection<T> collection, MongoIndex<T>[] newIndexes, CancellationToken cancellationToken)
 		{
+			_logger.LogDebug("Updating indexes for collection {CollectionName}", collectionName);
+
 			// Find all the current indexes, excluding the default
 			Dictionary<string, MongoIndex> nameToExistingIndex = new Dictionary<string, MongoIndex>(StringComparer.Ordinal);
 			using (IAsyncCursor<BsonDocument> cursor = await collection.Indexes.ListAsync(cancellationToken))
@@ -753,6 +755,8 @@ namespace Horde.Server.Server
 					await collection.Indexes.CreateOneAsync(model, cancellationToken: cancellationToken);
 				}
 			}
+
+			_logger.LogDebug("Finished updating indexes for collection {CollectionName}", collectionName);
 		}
 
 		class SingletonInfo<T> where T : SingletonBase
@@ -979,36 +983,10 @@ namespace Horde.Server.Server
 		{
 			try
 			{
-				// Check we're not downgrading the data
-				for (; ; )
+				if (await SetSchemaVersion(Program.Version))
 				{
-					MongoSchemaDocument currentSchema = await _schemaDocument.GetAsync();
-					if (!String.IsNullOrEmpty(currentSchema.Version))
-					{
-						SemVer currentVersion = SemVer.Parse(currentSchema.Version);
-						if (Program.Version < currentVersion)
-						{
-							return;
-						}
-						if (Program.Version == currentVersion)
-						{
-							break;
-						}
-					}
-					currentSchema.Version = Program.Version.ToString();
-
-					if (_mongoService.ReadOnlyMode) // Pretend like we've upgraded the document so that the task can run and log its changes; it will guard writing while this flag is set itself.
-					{
-						break;
-					}
-					if (await _schemaDocument.TryUpdateAsync(currentSchema))
-					{
-						break;
-					}
+					await taskAsync(cancellationToken);
 				}
-
-				// Perform the upgrade
-				await taskAsync(cancellationToken);
 			}
 			catch (MongoCommandException ex)
 			{
@@ -1019,5 +997,37 @@ namespace Horde.Server.Server
 				_logger.LogWarning(ex, "Exception while attempting to update indexes");
 			}
 		}
+
+		async Task<bool> SetSchemaVersion(SemVer schemaVersion)
+		{
+			// Check we're not downgrading the data
+			for (; ; )
+			{
+				MongoSchemaDocument currentSchema = await _schemaDocument.GetAsync();
+				if (!String.IsNullOrEmpty(currentSchema.Version))
+				{
+					SemVer currentVersion = SemVer.Parse(currentSchema.Version);
+					if (schemaVersion < currentVersion)
+					{
+						_logger.LogDebug("Ignoring upgrade command; server is older than current schema version ({ProgramVer} < {CurrentVer})", Program.Version, currentVersion);
+						return false;
+					}
+					if (schemaVersion == currentVersion)
+					{
+						return true;
+					}
+				}
+				currentSchema.Version = schemaVersion.ToString();
+
+				if (_mongoService.ReadOnlyMode)
+				{
+					return false;
+				}
+				if (await _schemaDocument.TryUpdateAsync(currentSchema))
+				{
+					return true;
+				}
+			}
+		}
 	}
 }
