Horde: Add log output for when we can't acquire the upgrade schema lock.

diff --git a/Engine/Source/Programs/Horde/Horde.Server/Server/MongoService.cs b/Engine/Source/Programs/Horde/Horde.Server/Server/MongoService.cs
--- a/Engine/Source/Programs/Horde/Horde.Server/Server/MongoService.cs
+++ b/Engine/Source/Programs/Horde/Horde.Server/Server/MongoService.cs
@@ -973,8 +973,12 @@ namespace Horde.Server.Server
 							await UpdateOneAsync(updateIndexTask, cancellationToken);
 							break;
 						}
+						else
+						{
+							_logger.LogDebug("Unable to acquire lock for upgrade task; pausing for 1m");
+							await Task.Delay(TimeSpan.FromMinutes(1.0), cancellationToken);
+						}
 					}
-					await Task.Delay(TimeSpan.FromMinutes(1.0), cancellationToken);
 				}
 			}
 		}
