Horde: More logging for MongoDB upgrades.

diff --git a/Engine/Source/Programs/Horde/Horde.Server/Server/MongoService.cs b/Engine/Source/Programs/Horde/Horde.Server/Server/MongoService.cs
--- a/Engine/Source/Programs/Horde/Horde.Server/Server/MongoService.cs
+++ b/Engine/Source/Programs/Horde/Horde.Server/Server/MongoService.cs
@@ -663,6 +663,8 @@ namespace Horde.Server.Server
 			IMongoCollection<T> collection = GetCollection<T>(name);
 			lock (_collectionNames)
 			{
+				_logger.LogDebug("Queuing update for collection {Name}", name);
+
 				if (!_collectionNames.Add(name))
 				{
 					throw new NotImplementedException();
@@ -1021,6 +1023,8 @@ namespace Horde.Server.Server
 						return true;
 					}
 				}
+
+				_logger.LogDebug("Upgrading schema version {OldVersion} -> {NewVersion}", currentSchema.Version, schemaVersion.ToString());
 				currentSchema.Version = schemaVersion.ToString();
 
 				if (_mongoService.ReadOnlyMode)
