Horde: More logging for index updates.

diff --git a/Engine/Source/Programs/Horde/Horde.Server/Server/MongoService.cs b/Engine/Source/Programs/Horde/Horde.Server/Server/MongoService.cs
--- a/Engine/Source/Programs/Horde/Horde.Server/Server/MongoService.cs
+++ b/Engine/Source/Programs/Horde/Horde.Server/Server/MongoService.cs
@@ -965,23 +965,31 @@ namespace Horde.Server.Server
 		{
 			for (; ; )
 			{
-				Func<CancellationToken, Task> updateIndexTask = await _mongoService.ReadNextUpgradeTask(cancellationToken);
-				for (; ; )
+				try
 				{
-					using (RedisLock schemaLock = new (_redisService.GetDatabase(), s_schemaLockKey))
+					Func<CancellationToken, Task> updateIndexTask = await _mongoService.ReadNextUpgradeTask(cancellationToken);
+					for (; ; )
 					{
-						if (await schemaLock.AcquireAsync(TimeSpan.FromMinutes(5.0)))
+						using (RedisLock schemaLock = new(_redisService.GetDatabase(), s_schemaLockKey))
 						{
-							await UpdateOneAsync(updateIndexTask, cancellationToken);
-							break;
-						}
-						else
-						{
-							_logger.LogDebug("Unable to acquire lock for upgrade task; pausing for 1m");
-							await Task.Delay(TimeSpan.FromMinutes(1.0), cancellationToken);
+							if (await schemaLock.AcquireAsync(TimeSpan.FromMinutes(5.0)))
+							{
+								await UpdateOneAsync(updateIndexTask, cancellationToken);
+								break;
+							}
+							else
+							{
+								_logger.LogDebug("Unable to acquire lock for upgrade task; pausing for 1m");
+								await Task.Delay(TimeSpan.FromMinutes(1.0), cancellationToken);
+							}
 						}
 					}
 				}
+				catch (Exception ex)
+				{
+					_logger.LogWarning(ex, "Exception updating indexes: {Message}", ex.Message);
+					await Task.Delay(TimeSpan.FromSeconds(30.0), cancellationToken);
+				}
 			}
 		}
 
