Horde: Run the Mongo upgrade service from non-worker nodes.

Some collections are only instantiated when required for a controller; this may not happen on worker nodes.

diff --git a/Engine/Source/Programs/Horde/Horde.Server/Startup.cs b/Engine/Source/Programs/Horde/Horde.Server/Startup.cs
--- a/Engine/Source/Programs/Horde/Horde.Server/Startup.cs
+++ b/Engine/Source/Programs/Horde/Horde.Server/Startup.cs
@@ -687,9 +687,11 @@ namespace Horde.Server
 			// Always run agent service too; need to be able to listen to Redis for events on any server.
 			services.AddHostedService(provider => provider.GetRequiredService<AgentService>());
 
+			// We need to run the Mongo upgrade service too; some collections only get instantiated when referenced at the moment, and references may come solely from HTTP requests.
+			services.AddHostedService<MongoUpgradeService>();
+
 			if (settings.IsRunModeActive(RunMode.Worker) && !settings.DatabaseReadOnlyMode)
 			{
-				services.AddHostedService<MongoUpgradeService>();
 				services.AddHostedService(provider => provider.GetRequiredService<FleetService>());
 				services.AddHostedService(provider => provider.GetRequiredService<ConsistencyService>());
 				services.AddHostedService(provider => provider.GetRequiredService<IssueService>());
