Horde: Add custom ACLs for modifying notices.

diff --git a/Engine/Source/Programs/Horde/Horde.Server/Server/Notices/NoticeAclAction.cs b/Engine/Source/Programs/Horde/Horde.Server/Server/Notices/NoticeAclAction.cs
new file mode 100644
--- /dev/null
+++ b/Engine/Source/Programs/Horde/Horde.Server/Server/Notices/NoticeAclAction.cs
@@ -0,0 +1,27 @@
+// Copyright Epic Games, Inc. All Rights Reserved.
+
+using Horde.Server.Acls;
+
+namespace Horde.Server.Server.Notices
+{
+	/// <summary>
+	/// ACL actions relating to notices
+	/// </summary>
+	public static class NoticeAclAction
+	{
+		/// <summary>
+		/// Ability to create new notices
+		/// </summary>
+		public static AclAction CreateNotice { get; } = new AclAction("CreateNotice");
+
+		/// <summary>
+		/// Ability to modify notices on the server
+		/// </summary>
+		public static AclAction UpdateNotice { get; } = new AclAction("UpdateNotice");
+
+		/// <summary>
+		/// Ability to delete notices
+		/// </summary>
+		public static AclAction DeleteNotice { get; } = new AclAction("DeleteNotice");
+	}
+}
diff --git a/Engine/Source/Programs/Horde/Horde.Server/Server/Notices/NoticesController.cs b/Engine/Source/Programs/Horde/Horde.Server/Server/Notices/NoticesController.cs
--- a/Engine/Source/Programs/Horde/Horde.Server/Server/Notices/NoticesController.cs
+++ b/Engine/Source/Programs/Horde/Horde.Server/Server/Notices/NoticesController.cs
@@ -52,7 +52,7 @@ namespace Horde.Server.Server.Notices
 		[HttpPost("/api/v1/notices")]
 		public async Task<ActionResult> AddNoticeAsync(CreateNoticeRequest request)
 		{
-			if (!_globalConfig.Value.Authorize(AdminAclAction.AdminWrite, User))
+			if (!_globalConfig.Value.Authorize(NoticeAclAction.CreateNotice, User))
 			{
 				return Forbid();
 			}
@@ -71,7 +71,7 @@ namespace Horde.Server.Server.Notices
 		[HttpPut("/api/v1/notices")]
 		public async Task<ActionResult> UpdateNoticeAsync(UpdateNoticeRequest request)
 		{
-			if (!_globalConfig.Value.Authorize(AdminAclAction.AdminWrite, User))
+			if (!_globalConfig.Value.Authorize(NoticeAclAction.UpdateNotice, User))
 			{
 				return Forbid();
 			}
@@ -86,7 +86,7 @@ namespace Horde.Server.Server.Notices
 		[HttpDelete("/api/v1/notices/{id}")]
 		public async Task<ActionResult> DeleteNoticeAsync(string id)
 		{
-			if (!_globalConfig.Value.Authorize(AdminAclAction.AdminWrite, User))
+			if (!_globalConfig.Value.Authorize(NoticeAclAction.DeleteNotice, User))
 			{
 				return Forbid();
 			}
