Horde: Return flags indicating whether various admin-related features should be enabled in the dashboard for the current user.

diff --git a/Engine/Source/Programs/Horde/Horde.Server/Users/UserController.cs b/Engine/Source/Programs/Horde/Horde.Server/Users/UserController.cs
--- a/Engine/Source/Programs/Horde/Horde.Server/Users/UserController.cs
+++ b/Engine/Source/Programs/Horde/Horde.Server/Users/UserController.cs
@@ -4,9 +4,11 @@ using System.Collections.Generic;
 using System.Linq;
 using System.Threading.Tasks;
 using Horde.Server.Jobs;
+using Horde.Server.Server;
 using Horde.Server.Utilities;
 using Microsoft.AspNetCore.Authorization;
 using Microsoft.AspNetCore.Mvc;
+using Microsoft.Extensions.Options;
 
 namespace Horde.Server.Users
 {
@@ -18,25 +20,18 @@ namespace Horde.Server.Users
 	[Route("[controller]")]
 	public class UserController : ControllerBase
 	{
-		/// <summary>
-		/// The user collection instance
-		/// </summary>
-		IUserCollection UserCollection { get; set; }
-
-		/// <summary>
-		/// The avatar service
-		/// </summary>
-		IAvatarService? AvatarService { get; set; }
+		readonly IUserCollection _userCollection;
+		readonly IAvatarService? _avatarService;
+		readonly IOptionsSnapshot<GlobalConfig> _globalConfig;
 
 		/// <summary>
 		/// Constructor
 		/// </summary>
-		/// <param name="userCollection"></param>
-		/// <param name="avatarService"></param>
-		public UserController(IUserCollection userCollection, IAvatarService? avatarService)
+		public UserController(IUserCollection userCollection, IAvatarService? avatarService, IOptionsSnapshot<GlobalConfig> globalConfig)
 		{
-			UserCollection = userCollection;
-			AvatarService = avatarService;
+			_userCollection = userCollection;
+			_avatarService = avatarService;
+			_globalConfig = globalConfig;
 		}
 
 		/// <summary>
@@ -45,19 +40,23 @@ namespace Horde.Server.Users
 		/// <returns>Http result code</returns>
 		[HttpGet]
 		[Route("/api/v1/user")]
-		[ProducesResponseType(typeof(List<GetUserResponse>), 200)]
+		[ProducesResponseType(typeof(GetUserResponse), 200)]
 		public async Task<ActionResult<object>> GetUserAsync([FromQuery] PropertyFilter? filter = null)
 		{
-			IUser? internalUser = await UserCollection.GetUserAsync(User);
+			IUser? internalUser = await _userCollection.GetUserAsync(User);
 			if (internalUser == null)
 			{
 				return NotFound();
 			}
 
-			IAvatar? avatar = (AvatarService == null)? (IAvatar?)null : await AvatarService.GetAvatarAsync(internalUser);
-			IUserClaims claims = await UserCollection.GetClaimsAsync(internalUser.Id);
-			IUserSettings settings = await UserCollection.GetSettingsAsync(internalUser.Id);
-			return PropertyFilter.Apply(new GetUserResponse(internalUser, avatar, claims, settings), filter);
+			IAvatar? avatar = (_avatarService == null)? (IAvatar?)null : await _avatarService.GetAvatarAsync(internalUser);
+			IUserClaims claims = await _userCollection.GetClaimsAsync(internalUser.Id);
+			IUserSettings settings = await _userCollection.GetSettingsAsync(internalUser.Id);
+
+			GetUserResponse response = new GetUserResponse(internalUser, avatar, claims, settings);
+			response.DashboardFeatures = new GetDashboardFeaturesResponse(_globalConfig.Value, User);
+
+			return PropertyFilter.Apply(response, filter);
 		}
 
 		/// <summary>
@@ -73,7 +72,8 @@ namespace Horde.Server.Users
 			{
 				return BadRequest("Current user does not have a registered profile");
 			}
-			await UserCollection.UpdateSettingsAsync(userId.Value, request.EnableExperimentalFeatures, request.DashboardSettings?.ToBsonValue(), request.AddPinnedJobIds?.Select(x => JobId.Parse(x)), request.RemovePinnedJobIds?.Select(x => JobId.Parse(x)));
+
+			await _userCollection.UpdateSettingsAsync(userId.Value, request.EnableExperimentalFeatures, request.DashboardSettings?.ToBsonValue(), request.AddPinnedJobIds?.Select(x => JobId.Parse(x)), request.RemovePinnedJobIds?.Select(x => JobId.Parse(x)));
 			return Ok();
 		}
 	}
diff --git a/Engine/Source/Programs/Horde/Horde.Server/Users/UserMessages.cs b/Engine/Source/Programs/Horde/Horde.Server/Users/UserMessages.cs
--- a/Engine/Source/Programs/Horde/Horde.Server/Users/UserMessages.cs
+++ b/Engine/Source/Programs/Horde/Horde.Server/Users/UserMessages.cs
@@ -3,8 +3,14 @@
 using System;
 using System.Collections.Generic;
 using System.Linq;
+using System.Security.Claims;
 using System.Text.Json;
 using EpicGames.Core;
+using Horde.Server.Acls;
+using Horde.Server.Agents;
+using Horde.Server.Agents.Pools;
+using Horde.Server.Server;
+using Horde.Server.Server.Notices;
 using MongoDB.Bson;
 
 namespace Horde.Server.Users
@@ -64,6 +70,11 @@ namespace Horde.Server.Users
 		/// </summary>
 		public object? DashboardSettings { get; set; }
 
+		/// <summary>
+		/// Settings for whether various dashboard features should be shown for the current user
+		/// </summary>
+		public GetDashboardFeaturesResponse? DashboardFeatures { get; set; }
+
 		/// <summary>
 		/// List of pinned job ids
 		/// </summary>
@@ -95,6 +106,37 @@ namespace Horde.Server.Users
 		}
 	}
 
+	/// <summary>
+	/// Settings for whether various features should be enabled on the dashboard
+	/// </summary>
+	public class GetDashboardFeaturesResponse
+	{
+		/// <summary>
+		/// Whether the notice editor should be listed in the server menu
+		/// </summary>
+		public bool ShowNoticeEditor { get; set; }
+
+		/// <summary>
+		/// Whether controls for modifying pools should be shown
+		/// </summary>
+		public bool ShowPoolEditor { get; set; }
+
+		/// <summary>
+		/// Whether the remote desktop button should be shown on the agent modal
+		/// </summary>
+		public bool ShowRemoteDesktop { get; set; }
+
+		/// <summary>
+		/// Constructor
+		/// </summary>
+		public GetDashboardFeaturesResponse(GlobalConfig globalConfig, ClaimsPrincipal principal)
+		{
+			ShowNoticeEditor = globalConfig.Authorize(NoticeAclAction.CreateNotice, principal) || globalConfig.Authorize(NoticeAclAction.UpdateNotice, principal);
+			ShowPoolEditor = globalConfig.Authorize(PoolAclAction.CreatePool, principal) || globalConfig.Authorize(PoolAclAction.UpdatePool, principal);
+			ShowRemoteDesktop = globalConfig.Authorize(AgentAclAction.UpdateAgent, principal);
+		}
+	}
+
 	/// <summary>
 	/// Basic information about a user. May be embedded in other responses.
 	/// </summary>
