Horde: Remove OpenTracing use coming from Datadog and Serilog NuGet packages

diff --git a/Engine/Source/Programs/Horde/Horde.Server/Horde.Server.csproj b/Engine/Source/Programs/Horde/Horde.Server/Horde.Server.csproj
--- a/Engine/Source/Programs/Horde/Horde.Server/Horde.Server.csproj
+++ b/Engine/Source/Programs/Horde/Horde.Server/Horde.Server.csproj
@@ -26,7 +26,6 @@
     <PackageReference Include="AWSSDK.S3" Version="3.7.100" />
     <PackageReference Include="AWSSDK.SecurityToken" Version="3.7.100" />
     <PackageReference Include="Crc32.NET" Version="1.2.0" />
-    <PackageReference Include="Datadog.Trace.OpenTracing" Version="2.14.0" />
     <PackageReference Include="DogStatsD-CSharp-Client" Version="7.0.1" />
     <PackageReference Include="Grpc.AspNetCore" Version="2.49.0" />
     <PackageReference Include="Grpc.AspNetCore.Server.Reflection" Version="2.49.0" />
@@ -55,7 +54,6 @@
     <PackageReference Include="Polly" Version="7.2.3" />
     <PackageReference Include="Polly.Extensions.Http" Version="3.0.0" />
     <PackageReference Include="Serilog.AspNetCore" Version="6.0.1" />
-    <PackageReference Include="Serilog.Enrichers.OpenTracing" Version="0.0.2" />
     <PackageReference Include="Serilog.Exceptions.Grpc" Version="8.4.0" />
     <PackageReference Include="Serilog.Extensions.Hosting" Version="5.0.1" />
     <PackageReference Include="Serilog.Sinks.Console" Version="4.1.0" />
diff --git a/Engine/Source/Programs/Horde/Horde.Server/Streams/StreamId.cs b/Engine/Source/Programs/Horde/Horde.Server/Streams/StreamId.cs
--- a/Engine/Source/Programs/Horde/Horde.Server/Streams/StreamId.cs
+++ b/Engine/Source/Programs/Horde/Horde.Server/Streams/StreamId.cs
@@ -5,7 +5,6 @@ using EpicGames.Core;
 using EpicGames.Horde;
 using EpicGames.Serialization;
 using Horde.Server.Utilities;
-using OpenTracing;
 
 namespace Horde.Server.Streams
 {
@@ -44,16 +43,4 @@ namespace Horde.Server.Streams
 		/// <inheritdoc/>
 		public override StringId ToStringId(StreamId value) => value.Id;
 	}
-
-	/// <summary>
-	/// Extension methods for stream id values
-	/// </summary>
-	public static class StreamIdExtensions
-	{
-		/// <inheritdoc cref="ISpan.SetTag(System.String, System.String)"/>
-		public static ISpan SetTag(this ISpan span, string key, StreamId value) => span.SetTag(key, value.Id.ToString());
-
-		/// <inheritdoc cref="ISpan.SetTag(System.String, System.String)"/>
-		public static ISpan SetTag(this ISpan span, string key, StreamId? value) => span.SetTag(key, value.HasValue ? value.Value.ToString() : null);
-	}
 }
diff --git a/Engine/Source/Programs/Horde/Horde.Server/Streams/TemplateId.cs b/Engine/Source/Programs/Horde/Horde.Server/Streams/TemplateId.cs
--- a/Engine/Source/Programs/Horde/Horde.Server/Streams/TemplateId.cs
+++ b/Engine/Source/Programs/Horde/Horde.Server/Streams/TemplateId.cs
@@ -7,7 +7,6 @@ using EpicGames.Core;
 using EpicGames.Horde;
 using EpicGames.Serialization;
 using Horde.Server.Utilities;
-using OpenTracing;
 
 namespace Horde.Server.Streams
 {
@@ -46,16 +45,4 @@ namespace Horde.Server.Streams
 		/// <inheritdoc/>
 		public override StringId ToStringId(TemplateId value) => value.Id;
 	}
-
-	/// <summary>
-	/// Extension methods for stream id values
-	/// </summary>
-	public static class TemplateIdExtensions
-	{
-		/// <inheritdoc cref="ISpan.SetTag(System.String, System.String)"/>
-		public static ISpan SetTag(this ISpan span, string key, TemplateId value) => span.SetTag(key, value.Id.ToString());
-
-		/// <inheritdoc cref="ISpan.SetTag(System.String, System.String)"/>
-		public static ISpan SetTag(this ISpan span, string key, TemplateId[]? values) => span.SetTag(key, (values != null)? String.Join(',', values.Select(x => x.Id.ToString())) : null);
-	}
 }
diff --git a/Engine/Source/Programs/Horde/Horde.Server/Utilities/OpenTelemetry.cs b/Engine/Source/Programs/Horde/Horde.Server/Utilities/OpenTelemetry.cs
--- a/Engine/Source/Programs/Horde/Horde.Server/Utilities/OpenTelemetry.cs
+++ b/Engine/Source/Programs/Horde/Horde.Server/Utilities/OpenTelemetry.cs
@@ -19,7 +19,6 @@ using OpenTelemetry.Logs;
 using OpenTelemetry.Metrics;
 using OpenTelemetry.Resources;
 using OpenTelemetry.Trace;
-using OpenTracing;
 using Serilog.Core;
 
 namespace Horde.Server.Utilities;
@@ -294,13 +293,13 @@ public static class OpenTelemetrySpanExtensions
 		return span;
 	}
 	
-	/// <inheritdoc cref="ISpan.SetTag(System.String, System.String)"/>
+	/// <inheritdoc cref="TelemetrySpan.SetAttribute(System.String, System.String)"/>
 	public static TelemetrySpan SetAttribute(this TelemetrySpan span, string key, StreamId? value) => span.SetAttribute(key, value?.ToString());
 	
-	/// <inheritdoc cref="ISpan.SetTag(System.String, System.String)"/>
+	/// <inheritdoc cref="TelemetrySpan.SetAttribute(System.String, System.String)"/>
 	public static TelemetrySpan SetAttribute(this TelemetrySpan span, string key, TemplateId? value) => span.SetAttribute(key, value?.ToString());
 	
-	/// <inheritdoc cref="ISpan.SetTag(System.String, System.String)"/>
+	/// <inheritdoc cref="TelemetrySpan.SetAttribute(System.String, System.String)"/>
 	public static TelemetrySpan SetAttribute(this TelemetrySpan span, string key, TemplateId[]? values) => span.SetAttribute(key, values != null ? String.Join(',', values.Select(x => x.Id.ToString())) : null);
 
 	/// <summary>
diff --git a/Engine/Source/Programs/Horde/Horde.Server/Utilities/TracingExtensions.cs b/Engine/Source/Programs/Horde/Horde.Server/Utilities/TracingExtensions.cs
deleted file mode 100644
