Horde: Replace DogStatsD with OpenTelemetry metrics

diff --git a/Engine/Source/Programs/Horde/Horde.Server.Tests/Fleet/AwsRecyclingFleetManagerTest.cs b/Engine/Source/Programs/Horde/Horde.Server.Tests/Fleet/AwsRecyclingFleetManagerTest.cs
--- a/Engine/Source/Programs/Horde/Horde.Server.Tests/Fleet/AwsRecyclingFleetManagerTest.cs
+++ b/Engine/Source/Programs/Horde/Horde.Server.Tests/Fleet/AwsRecyclingFleetManagerTest.cs
@@ -218,8 +218,7 @@ namespace Horde.Server.Tests.Fleet
 
 			ILogger<AwsRecyclingFleetManager> logger = loggerFactory.CreateLogger<AwsRecyclingFleetManager>();
 			IPool pool = await PoolService.CreatePoolAsync("testPool", null, true, 0, 0, sizeStrategy: PoolSizeStrategy.NoOp);
-			using NoOpDogStatsd dogStatsd = new ();
-			AwsRecyclingFleetManager manager = new (ec2, AgentCollection, dogStatsd, _clock, settings, Tracer, logger);
+			AwsRecyclingFleetManager manager = new (ec2, AgentCollection, Meter, _clock, settings, Tracer, logger);
 			return await manager.ExpandPoolAsync(pool, new List<IAgent>(), numRequestedInstances, CancellationToken.None);
 		}
 	}
diff --git a/Engine/Source/Programs/Horde/Horde.Server.Tests/Fleet/FleetServiceTest.cs b/Engine/Source/Programs/Horde/Horde.Server.Tests/Fleet/FleetServiceTest.cs
--- a/Engine/Source/Programs/Horde/Horde.Server.Tests/Fleet/FleetServiceTest.cs
+++ b/Engine/Source/Programs/Horde/Horde.Server.Tests/Fleet/FleetServiceTest.cs
@@ -17,7 +17,6 @@ using Microsoft.Extensions.Logging;
 using Microsoft.Extensions.Options;
 using Microsoft.VisualStudio.TestTools.UnitTesting;
 using OpenTelemetry.Trace;
-using StatsdClient;
 
 namespace Horde.Server.Tests.Fleet
 {
@@ -100,17 +99,6 @@ namespace Horde.Server.Tests.Fleet
 	public class FleetServiceTest : TestSetup
 	{
 		readonly FleetManagerSpy _fleetManagerSpy = new();
-		readonly IDogStatsd _dogStatsD = new NoOpDogStatsd();
-
-		protected override void Dispose(bool disposing)
-		{
-			base.Dispose(disposing);
-
-			if (disposing)
-			{
-				_dogStatsD.Dispose();
-			}
-		}
 
 		[TestMethod]
 		public async Task OnlyEnabledAgentsAreAutoScaled()
@@ -243,7 +231,7 @@ namespace Horde.Server.Tests.Fleet
 			serverSettingsOpt.Value.FleetManagerV2 = FleetManagerType.AwsReuse;
 
 			FleetService service = new(
-				AgentCollection, GraphCollection, JobCollection, LeaseCollection, PoolCollection, new DowntimeServiceStub(isDowntimeActive), StreamCollection, _dogStatsD,
+				AgentCollection, GraphCollection, JobCollection, LeaseCollection, PoolCollection, new DowntimeServiceStub(isDowntimeActive), StreamCollection, Meter,
 				new StubFleetManagerFactory(fleetManager), Clock, Cache, serverSettingsOpt, GlobalConfig, Tracer, loggerFactory.CreateLogger<FleetService>());
 				
 			return service;
diff --git a/Engine/Source/Programs/Horde/Horde.Server.Tests/TestSetup.cs b/Engine/Source/Programs/Horde/Horde.Server.Tests/TestSetup.cs
--- a/Engine/Source/Programs/Horde/Horde.Server.Tests/TestSetup.cs
+++ b/Engine/Source/Programs/Horde/Horde.Server.Tests/TestSetup.cs
@@ -39,7 +39,6 @@ using Microsoft.Extensions.DependencyInjection;
 using Microsoft.Extensions.Hosting;
 using Microsoft.Extensions.Logging;
 using Microsoft.Extensions.Options;
-using StatsdClient;
 using Horde.Server.Users;
 using Horde.Server.Agents.Software;
 using Horde.Server.Ugs;
@@ -133,6 +132,7 @@ namespace Horde.Server.Tests
 		public TestDataController TestDataController => GetTestDataController();
 
 		public OpenTelemetry.Trace.Tracer Tracer => ServiceProvider.GetRequiredService<OpenTelemetry.Trace.Tracer>();
+		public Meter Meter => ServiceProvider.GetRequiredService<Meter>();
 		public ConfigService ConfigService => ServiceProvider.GetRequiredService<ConfigService>();
 		public IOptionsMonitor<GlobalConfig> GlobalConfig => ServiceProvider.GetRequiredService<IOptionsMonitor<GlobalConfig>>();
 		public IOptionsSnapshot<GlobalConfig> GlobalConfigSnapshot => ServiceProvider.GetRequiredService<IOptionsSnapshot<GlobalConfig>>();
@@ -244,7 +244,6 @@ namespace Horde.Server.Tests
 			services.AddSingleton<CredentialService>();
 			services.AddSingleton<JobTaskSource>();
 			services.AddSingleton<IDowntimeService, DowntimeServiceStub>();
-			services.AddSingleton<IDogStatsd, NoOpDogStatsd>();
 			services.AddSingleton<IssueService>();
 			services.AddSingleton<JobService>();
 			services.AddSingleton<LifetimeService>();
diff --git a/Engine/Source/Programs/Horde/Horde.Server/Agents/AgentService.cs b/Engine/Source/Programs/Horde/Horde.Server/Agents/AgentService.cs
--- a/Engine/Source/Programs/Horde/Horde.Server/Agents/AgentService.cs
+++ b/Engine/Source/Programs/Horde/Horde.Server/Agents/AgentService.cs
@@ -3,6 +3,7 @@
 using System;
 using System.Collections.Generic;
 using System.Diagnostics;
+using System.Diagnostics.Metrics;
 using System.Linq;
 using System.Security.Claims;
 using System.Threading;
@@ -22,12 +23,8 @@ using Horde.Server.Utilities;
 using HordeCommon;
 using Microsoft.Extensions.Hosting;
 using Microsoft.Extensions.Logging;
-using MongoDB.Driver;
 using OpenTelemetry.Trace;
-using OpenTracing;
-using OpenTracing.Util;
 using StackExchange.Redis;
-using StatsdClient;
 
 namespace Horde.Server.Agents
 {
@@ -73,11 +70,11 @@ namespace Horde.Server.Agents
 
 		readonly ILeaseCollection _leases;
 		readonly ISessionCollection _sessions;
-		readonly IDogStatsd _dogStatsd;
 		readonly ITaskSource[] _taskSources;
 		readonly IHostApplicationLifetime _applicationLifetime;
 		readonly IClock _clock;
 		readonly Tracer _tracer;
+		readonly Meter _meter;
 		readonly ILogger _logger;
 		readonly ITicker _ticker;
 		
@@ -93,13 +90,15 @@ namespace Horde.Server.Agents
 		// All the agents currently performing a long poll for work on this server
 		readonly Dictionary<AgentId, CancellationTokenSource> _waitingAgents = new Dictionary<AgentId, CancellationTokenSource>();
 
+		readonly Gauge<int> _agentTotalCount;
+
 		// Subscription for update events
 		IAsyncDisposable? _subscription;
 
 		/// <summary>
 		/// Constructor
 		/// </summary>
-		public AgentService(IAgentCollection agents, ILeaseCollection leases, ISessionCollection sessions, AclService aclService, IDowntimeService downtimeService, IPoolCollection poolCollection, IDogStatsd dogStatsd, IEnumerable<ITaskSource> taskSources, RedisService redisService, IHostApplicationLifetime applicationLifetime, IClock clock, Tracer tracer, ILogger<AgentService> logger)
+		public AgentService(IAgentCollection agents, ILeaseCollection leases, ISessionCollection sessions, AclService aclService, IDowntimeService downtimeService, IPoolCollection poolCollection, IEnumerable<ITaskSource> taskSources, RedisService redisService, IHostApplicationLifetime applicationLifetime, IClock clock, Tracer tracer, Meter meter, ILogger<AgentService> logger)
 		{
 			Agents = agents;
 			_leases = leases;
@@ -108,14 +107,16 @@ namespace Horde.Server.Agents
 			_downtimeService = downtimeService;
 			_agentRateTable = new AsyncCachedValue<AgentRateTable?>(() => redisService.GetDatabase().StringGetAsync(_agentRateTableData), TimeSpan.FromSeconds(2.0));//.FromMinutes(5.0));
 			_poolsList = new AsyncCachedValue<List<IPool>>(() => poolCollection.GetAsync(), TimeSpan.FromSeconds(30.0));
-			_dogStatsd = dogStatsd;
 			_taskSources = taskSources.ToArray();
 			_applicationLifetime = applicationLifetime;
 			_redisService = redisService;
 			_clock = clock;
 			_ticker = clock.AddTicker<AgentService>(TimeSpan.FromSeconds(30.0), TickAsync, logger);
 			_tracer = tracer;
+			_meter = meter;
 			_logger = logger;
+
+			_agentTotalCount = _meter.CreateGauge<int>("horde.agent.count");
 		}
 
 		/// <inheritdoc/>
@@ -1026,14 +1027,13 @@ namespace Horde.Server.Agents
 			int numAgentsTotalUnhealthy = agentList.Count(a => a.Enabled && a.Status == AgentStatus.Unhealthy);
 			int numAgentsTotalUnspecified = agentList.Count(a => a.Enabled && a.Status == AgentStatus.Unspecified);
 			
-			// TODO: utilize tags argument in a smarter way below
-			_dogStatsd.Gauge("agents.total.count", numAgentsTotal);
-			_dogStatsd.Gauge("agents.total.enabled.count", numAgentsTotalEnabled);
-			_dogStatsd.Gauge("agents.total.disabled.count", numAgentsTotalDisabled);
-			_dogStatsd.Gauge("agents.total.ok.count", numAgentsTotalOk);
-			_dogStatsd.Gauge("agents.total.stopping.count", numAgentsTotalStopping);
-			_dogStatsd.Gauge("agents.total.unhealthy.count", numAgentsTotalUnhealthy);
-			_dogStatsd.Gauge("agents.total.unspecified.count", numAgentsTotalUnspecified);
+			_agentTotalCount.Record(numAgentsTotal);
+			_agentTotalCount.Record(numAgentsTotalEnabled, new KeyValuePair<string, object?>("status", "enabled"));
+			_agentTotalCount.Record(numAgentsTotalDisabled, new KeyValuePair<string, object?>("status", "disabled"));
+			_agentTotalCount.Record(numAgentsTotalOk, new KeyValuePair<string, object?>("status", "ok"));
+			_agentTotalCount.Record(numAgentsTotalStopping, new KeyValuePair<string, object?>("status", "stopping"));
+			_agentTotalCount.Record(numAgentsTotalUnhealthy, new KeyValuePair<string, object?>("status", "unhealthy"));
+			_agentTotalCount.Record(numAgentsTotalUnspecified, new KeyValuePair<string, object?>("status", "unspecified"));
 		}
 
 		/// <summary>
diff --git a/Engine/Source/Programs/Horde/Horde.Server/Agents/Fleet/FleetManagerFactory.cs b/Engine/Source/Programs/Horde/Horde.Server/Agents/Fleet/FleetManagerFactory.cs
--- a/Engine/Source/Programs/Horde/Horde.Server/Agents/Fleet/FleetManagerFactory.cs
+++ b/Engine/Source/Programs/Horde/Horde.Server/Agents/Fleet/FleetManagerFactory.cs
@@ -1,6 +1,7 @@
 // Copyright Epic Games, Inc. All Rights Reserved.
 
 using System;
+using System.Diagnostics.Metrics;
 using System.Text.Json;
 using Amazon.AutoScaling;
 using Amazon.EC2;
@@ -10,7 +11,6 @@ using Microsoft.Extensions.DependencyInjection;
 using Microsoft.Extensions.Logging;
 using Microsoft.Extensions.Options;
 using OpenTelemetry.Trace;
-using StatsdClient;
 
 namespace Horde.Server.Agents.Fleet;
 
@@ -36,7 +36,7 @@ public sealed class FleetManagerFactory : IFleetManagerFactory
 {
 	private readonly IAgentCollection _agentCollection;
 	private readonly IClock _clock;
-	private readonly IDogStatsd _dogStatsd;
+	private readonly Meter _meter;
 	private readonly IServiceProvider _provider;
 	private readonly IOptionsMonitor<ServerSettings> _settings;
 	private readonly Tracer _tracer;
@@ -48,11 +48,11 @@ public sealed class FleetManagerFactory : IFleetManagerFactory
 	/// <summary>
 	/// Constructor
 	/// </summary>
-	public FleetManagerFactory(IAgentCollection agentCollection, IClock clock, IDogStatsd dogStatsd, IServiceProvider provider, IOptionsMonitor<ServerSettings> settings, Tracer tracer, ILoggerFactory loggerFactory)
+	public FleetManagerFactory(IAgentCollection agentCollection, IClock clock, Meter meter, IServiceProvider provider, IOptionsMonitor<ServerSettings> settings, Tracer tracer, ILoggerFactory loggerFactory)
 	{
 		_agentCollection = agentCollection;
 		_clock = clock;
-		_dogStatsd = dogStatsd;
+		_meter = meter;
 		_provider = provider;
 		_settings = settings;
 		_tracer = tracer;
@@ -73,7 +73,7 @@ public sealed class FleetManagerFactory : IFleetManagerFactory
 			FleetManagerType.AwsReuse =>
 				new AwsReuseFleetManager(GetAwsEc2(type), _agentCollection, DeserializeSettings<AwsReuseFleetManagerSettings>(config), _tracer, _loggerFactory.CreateLogger<AwsReuseFleetManager>()),
 			FleetManagerType.AwsRecycle =>
-				new AwsRecyclingFleetManager(GetAwsEc2(type), _agentCollection, _dogStatsd, _clock, DeserializeSettings<AwsRecyclingFleetManagerSettings>(config), _tracer, _loggerFactory.CreateLogger<AwsRecyclingFleetManager>()),
+				new AwsRecyclingFleetManager(GetAwsEc2(type), _agentCollection, _meter, _clock, DeserializeSettings<AwsRecyclingFleetManagerSettings>(config), _tracer, _loggerFactory.CreateLogger<AwsRecyclingFleetManager>()),
 			FleetManagerType.AwsAsg =>
 				new AwsAsgFleetManager(GetAwsAutoScaling(type), DeserializeSettings<AwsAsgSettings>(config), _tracer, _loggerFactory.CreateLogger<AwsAsgFleetManager>()),
 			_ => throw new ArgumentException("Unknown fleet manager type " + type)
diff --git a/Engine/Source/Programs/Horde/Horde.Server/Agents/Fleet/FleetService.cs b/Engine/Source/Programs/Horde/Horde.Server/Agents/Fleet/FleetService.cs
--- a/Engine/Source/Programs/Horde/Horde.Server/Agents/Fleet/FleetService.cs
+++ b/Engine/Source/Programs/Horde/Horde.Server/Agents/Fleet/FleetService.cs
@@ -4,6 +4,7 @@ using System;
 using System.Collections.Concurrent;
 using System.Collections.Generic;
 using System.Diagnostics;
+using System.Diagnostics.Metrics;
 using System.Linq;
 using System.Text.Json;
 using System.Threading;
@@ -21,7 +22,6 @@ using Microsoft.Extensions.Hosting;
 using Microsoft.Extensions.Logging;
 using Microsoft.Extensions.Options;
 using OpenTelemetry.Trace;
-using StatsdClient;
 
 namespace Horde.Server.Agents.Fleet
 {
@@ -69,7 +69,7 @@ namespace Horde.Server.Agents.Fleet
 		private readonly IPoolCollection _poolCollection;
 		private readonly IDowntimeService _downtimeService;
 		private readonly IStreamCollection _streamCollection;
-		private readonly IDogStatsd _dogStatsd;
+		private readonly Meter _meter;
 		private readonly IFleetManagerFactory _fleetManagerFactory;
 		private readonly IClock _clock;
 		private readonly IMemoryCache _cache;
@@ -81,6 +81,8 @@ namespace Horde.Server.Agents.Fleet
 		private readonly IOptionsMonitor<GlobalConfig> _globalConfig;
 		private readonly Tracer _tracer;
 		private readonly ILogger<FleetService> _logger;
+		private readonly Gauge<int> _agentCountCurrentGauge;
+		private readonly Gauge<int> _agentCountTargetGauge;
 		
 		/// <summary>
 		/// Constructor
@@ -93,7 +95,7 @@ namespace Horde.Server.Agents.Fleet
 			IPoolCollection poolCollection,
 			IDowntimeService downtimeService,
 			IStreamCollection streamCollection,
-			IDogStatsd dogStatsd,
+			Meter meter,
 			IFleetManagerFactory fleetManagerFactory,
 			IClock clock,
 			IMemoryCache cache,
@@ -109,7 +111,7 @@ namespace Horde.Server.Agents.Fleet
 			_poolCollection = poolCollection;
 			_downtimeService = downtimeService;
 			_streamCollection = streamCollection;
-			_dogStatsd = dogStatsd;
+			_meter = meter;
 			_fleetManagerFactory = fleetManagerFactory;
 			_clock = clock;
 			_cache = cache;
@@ -121,6 +123,9 @@ namespace Horde.Server.Agents.Fleet
 			_settings = settings;
 			_defaultScaleOutCooldown = TimeSpan.FromSeconds(settings.Value.AgentPoolScaleOutCooldownSeconds);
 			_defaultScaleInCooldown = TimeSpan.FromSeconds(settings.Value.AgentPoolScaleInCooldownSeconds);
+
+			_agentCountCurrentGauge = _meter.CreateGauge<int>("horde.fleet.agentCountCurrent");
+			_agentCountTargetGauge = _meter.CreateGauge<int>("horde.fleet.agentCountTarget");
 		}
 
 		/// <inheritdoc/>
@@ -233,8 +238,8 @@ namespace Horde.Server.Agents.Fleet
 					pool.Name, currentAgentCount, desiredAgentCount, deltaAgentCount);	
 			}
 
-			_dogStatsd.Gauge("agentpools.autoscale.target", desiredAgentCount, tags: new []{"pool:" + pool.Name});
-			_dogStatsd.Gauge("agentpools.autoscale.current", currentAgentCount, tags: new []{"pool:" + pool.Name});
+			_agentCountCurrentGauge.Record(currentAgentCount, KeyValuePair.Create("pool", (object?)pool.Name));
+			_agentCountTargetGauge.Record(desiredAgentCount, KeyValuePair.Create("pool", (object?)pool.Name));
 
 			DateTime? scaleOutTime = null;
 			DateTime? scaleInTime = null;
diff --git a/Engine/Source/Programs/Horde/Horde.Server/Agents/Fleet/Providers/AwsRecyclingFleetManager.cs b/Engine/Source/Programs/Horde/Horde.Server/Agents/Fleet/Providers/AwsRecyclingFleetManager.cs
--- a/Engine/Source/Programs/Horde/Horde.Server/Agents/Fleet/Providers/AwsRecyclingFleetManager.cs
+++ b/Engine/Source/Programs/Horde/Horde.Server/Agents/Fleet/Providers/AwsRecyclingFleetManager.cs
@@ -2,6 +2,7 @@
 
 using System;
 using System.Collections.Generic;
+using System.Diagnostics.Metrics;
 using System.Linq;
 using System.Net;
 using System.Threading;
@@ -14,7 +15,6 @@ using Horde.Server.Agents.Pools;
 using HordeCommon;
 using Microsoft.Extensions.Logging;
 using OpenTelemetry.Trace;
-using StatsdClient;
 
 namespace Horde.Server.Agents.Fleet.Providers;
 
@@ -90,23 +90,25 @@ public sealed class AwsRecyclingFleetManager : IFleetManager
 	
 	private readonly IAmazonEC2 _ec2;
 	private readonly IAgentCollection _agentCollection;
-	private readonly IDogStatsd _dogStatsd;
+	private readonly Meter _meter;
 	private readonly IClock _clock;
 	private readonly Tracer _tracer;
 	private readonly ILogger _logger;
+	private readonly Counter<int> _ec2StartInstanceCounter;
 
 	/// <summary>
 	/// Constructor
 	/// </summary>
-	public AwsRecyclingFleetManager(IAmazonEC2 ec2, IAgentCollection agentCollection, IDogStatsd dogStatsd, IClock clock, AwsRecyclingFleetManagerSettings settings, Tracer tracer, ILogger<AwsRecyclingFleetManager> logger)
+	public AwsRecyclingFleetManager(IAmazonEC2 ec2, IAgentCollection agentCollection, Meter meter, IClock clock, AwsRecyclingFleetManagerSettings settings, Tracer tracer, ILogger<AwsRecyclingFleetManager> logger)
 	{
 		_ec2 = ec2;
 		_agentCollection = agentCollection;
-		_dogStatsd = dogStatsd;
+		_meter = meter;
 		_clock = clock;
 		_tracer = tracer;
 		_logger = logger;
 		Settings = settings;
+		_ec2StartInstanceCounter = _meter.CreateCounter<int>("horde.fleet.awsRecycling.startInstance");
 	}
 
 	/// <inheritdoc/>
@@ -264,22 +266,27 @@ public sealed class AwsRecyclingFleetManager : IFleetManager
 		bool success = false;
 		foreach (InstanceType? instanceType in instanceTypes)
 		{
-			string[] metricTags = { "instanceType:" + instanceType, "az:" + instances[0].Placement.AvailabilityZone };
+			KeyValuePair<string, object?> instanceTypeTag = KeyValuePair.Create("instanceType", (object?)instanceType);
+			KeyValuePair<string, object?> azTag = KeyValuePair.Create("az", (object?)instances[0].Placement.AvailabilityZone);
+
 			try
 			{
 				await StartInstancesAsync(instances, instanceType, cancellationToken);
 				success = true;
-				_dogStatsd.Increment("horde.fleet.awsRecycling.startInstance.success", tags: metricTags);
+				KeyValuePair<string, object?> statusTag = KeyValuePair.Create("status", (object?)"success");
+				_ec2StartInstanceCounter.Add(1, new [] { instanceTypeTag, azTag, statusTag });
 				break;
 			}
 			catch (AwsInsufficientCapacityException)
 			{
-				_dogStatsd.Increment("horde.fleet.awsRecycling.startInstance.errorCapacity", tags: metricTags);
+				KeyValuePair<string, object?> statusTag = KeyValuePair.Create("status", (object?)"errorCapacity");
+				_ec2StartInstanceCounter.Add(1, new [] { instanceTypeTag, azTag, statusTag });
 				_logger.LogInformation("Insufficient capacity for {InstanceType}", instanceType?.ToString());
 			}
 			catch (Exception)
 			{
-				_dogStatsd.Increment("horde.fleet.awsRecycling.startInstance.error", tags: metricTags);
+				KeyValuePair<string, object?> statusTag = KeyValuePair.Create("status", (object?)"error");
+				_ec2StartInstanceCounter.Add(1, new [] { instanceTypeTag, azTag, statusTag });
 				throw;
 			}
 		}
diff --git a/Engine/Source/Programs/Horde/Horde.Server/Horde.Server.csproj b/Engine/Source/Programs/Horde/Horde.Server/Horde.Server.csproj
--- a/Engine/Source/Programs/Horde/Horde.Server/Horde.Server.csproj
+++ b/Engine/Source/Programs/Horde/Horde.Server/Horde.Server.csproj
@@ -26,7 +26,6 @@
     <PackageReference Include="AWSSDK.S3" Version="3.7.100" />
     <PackageReference Include="AWSSDK.SecurityToken" Version="3.7.100" />
     <PackageReference Include="Crc32.NET" Version="1.2.0" />
-    <PackageReference Include="DogStatsD-CSharp-Client" Version="7.0.1" />
     <PackageReference Include="Grpc.AspNetCore" Version="2.49.0" />
     <PackageReference Include="Grpc.AspNetCore.Server.Reflection" Version="2.49.0" />
     <PackageReference Include="Grpc.Net.Client" Version="2.49.0" />
diff --git a/Engine/Source/Programs/Horde/Horde.Server/Notifications/NotificationService.cs b/Engine/Source/Programs/Horde/Horde.Server/Notifications/NotificationService.cs
--- a/Engine/Source/Programs/Horde/Horde.Server/Notifications/NotificationService.cs
+++ b/Engine/Source/Programs/Horde/Horde.Server/Notifications/NotificationService.cs
@@ -3,6 +3,7 @@
 using System;
 using System.Collections.Concurrent;
 using System.Collections.Generic;
+using System.Diagnostics.Metrics;
 using System.Linq;
 using System.Security.Claims;
 using System.Threading;
@@ -26,9 +27,7 @@ using Microsoft.Extensions.Hosting;
 using Microsoft.Extensions.Logging;
 using Microsoft.Extensions.Options;
 using MongoDB.Bson;
-using MongoDB.Driver;
 using StackExchange.Redis;
-using StatsdClient;
 using JsonSerializer = System.Text.Json.JsonSerializer;
 
 namespace Horde.Server.Notifications
@@ -83,11 +82,6 @@ namespace Horde.Server.Notifications
 		/// </summary>
 		private readonly ILogFileService _logFileService;
 		
-		/// <summary>
-		/// Instance of the <see cref="IDogStatsd"/>.
-		/// </summary>
-		private readonly IDogStatsd _dogStatsd;
-		
 		/// <summary>
 		/// Cache for de-duplicating queued notifications
 		/// </summary>
@@ -136,6 +130,9 @@ namespace Horde.Server.Notifications
 
 		readonly IOptionsMonitor<GlobalConfig> _globalConfig;
 
+		readonly Counter<int> _jobCounter;
+		readonly Histogram<double> _jobDurationHistogram;
+
 		static string RedisQueueListKey(string notificationType) => "NotificationService.queued." + notificationType;
 
 		/// <summary>
@@ -153,7 +150,7 @@ namespace Horde.Server.Notifications
 			IStreamCollection streamCollection,
 			IssueService issueService,
 			ILogFileService logFileService,
-			IDogStatsd dogStatsd,
+			Meter meter,
 			IMemoryCache cache,
 			RedisService redisService,
 			ConfigService configService,
@@ -171,7 +168,6 @@ namespace Horde.Server.Notifications
 			_streamCollection = streamCollection;
 			_issueService = issueService;
 			_logFileService = logFileService;
-			_dogStatsd = dogStatsd;
 			_cache = cache;
 			_redisConnectionPool = redisService.ConnectionPool;
 			_globalConfig = globalConfig;
@@ -183,6 +179,8 @@ namespace Horde.Server.Notifications
 			configService.OnConfigUpdate += NotifyConfigUpdate;
 
 			_ticker = clock.AddSharedTicker<NotificationService>(_notificationBatchInterval, TickEveryTwelveHoursAsync, logger);
+			_jobCounter = meter.CreateCounter<int>("horde.notification.job.count");
+			_jobDurationHistogram = meter.CreateHistogram<double>("horde.notification.job.duration");
 		}
 
 		/// <inheritdoc/>
@@ -528,8 +526,15 @@ namespace Horde.Server.Notifications
 					_ => "unspecified"
 				};
 				
-				string[] tags = {"stream:" + job.StreamId, "template:" + job.TemplateId};
-				_dogStatsd.Increment($"horde.{type}.{outcomeStr}.count", 1, tags: tags);
+				KeyValuePair<string, object?>[] tags =
+				{
+					KeyValuePair.Create<string, object?>("stream", job.StreamId.ToString()),
+					KeyValuePair.Create<string, object?>("template", job.TemplateId.ToString()),
+					KeyValuePair.Create<string, object?>("outcome", outcomeStr),
+					KeyValuePair.Create<string, object?>("type", type)
+				};
+
+				_jobCounter.Add(1, tags);
 
 				if (startTime == null || finishTime == null)
 				{
@@ -538,7 +543,7 @@ namespace Horde.Server.Notifications
 				}
 
 				TimeSpan duration = finishTime.Value - startTime.Value;
-				_dogStatsd.Timer($"horde.{type}.{outcomeStr}.duration", duration.TotalSeconds, tags: tags);
+				_jobDurationHistogram.Record(duration.TotalSeconds, tags);
 			}
 
 			JobStepOutcome jobOutcome = job.Batches.SelectMany(x => x.Steps).Min(x => x.Outcome);
diff --git a/Engine/Source/Programs/Horde/Horde.Server/Server/MetricService.cs b/Engine/Source/Programs/Horde/Horde.Server/Server/MetricService.cs
--- a/Engine/Source/Programs/Horde/Horde.Server/Server/MetricService.cs
+++ b/Engine/Source/Programs/Horde/Horde.Server/Server/MetricService.cs
@@ -1,12 +1,13 @@
 // Copyright Epic Games, Inc. All Rights Reserved.
 
 using System;
+using System.Diagnostics.Metrics;
 using System.Threading;
 using System.Threading.Tasks;
+using Horde.Server.Utilities;
 using HordeCommon;
 using Microsoft.Extensions.Hosting;
 using Microsoft.Extensions.Logging;
-using StatsdClient;
 
 namespace Horde.Server.Server
 {
@@ -15,16 +16,48 @@ namespace Horde.Server.Server
 	/// </summary>
 	public sealed class MetricService : IHostedService, IDisposable
 	{
-		readonly IDogStatsd _dogStatsd;
+		readonly Meter _meter;
 		readonly ITicker _ticker;
+		
+		private readonly Gauge<int> _threadPoolIoMax;
+		private readonly Gauge<int> _threadPoolIoMin;
+		private readonly Gauge<int> _threadPoolIoFree;
+		private readonly Gauge<int> _threadPoolIoBusy;
+		private readonly Gauge<int> _threadPoolWorkerMax;
+		private readonly Gauge<int> _threadPoolWorkerMin;
+		private readonly Gauge<int> _threadPoolWorkerFree;
+		private readonly Gauge<int> _threadPoolWorkerBusy;
+		private readonly Gauge<long> _clrGcTotalMem;
+		private readonly Gauge<long> _clrGcTotalAllocated;
+		private readonly Gauge<long> _clrGcHeapSize;
+		private readonly Gauge<long> _clrGcFragmented;
+		private readonly Gauge<long> _clrGcMemoryLoad;
+		private readonly Gauge<long> _clrGcTotalAvailableMem;
+		private readonly Gauge<long> _clrGcHighMemThreshold;
 
 		/// <summary>
 		/// Constructor
 		/// </summary>
-		public MetricService(IDogStatsd dogStatsd, IClock clock, ILogger<MetricService> logger)
+		public MetricService(Meter meter, IClock clock, ILogger<MetricService> logger)
 		{
-			_dogStatsd = dogStatsd;
+			_meter = meter;
 			_ticker = clock.AddTicker<MetricService>(TimeSpan.FromSeconds(20.0), TickAsync, logger);
+
+			_threadPoolIoMax = _meter.CreateGauge<int>("horde.clr.threadpool.io.max");
+			_threadPoolIoMin = _meter.CreateGauge<int>("horde.clr.threadpool.io.min");
+			_threadPoolIoFree = _meter.CreateGauge<int>("horde.clr.threadpool.io.free");
+			_threadPoolIoBusy = _meter.CreateGauge<int>("horde.clr.threadpool.io.busy");
+			_threadPoolWorkerMax = _meter.CreateGauge<int>("horde.clr.threadpool.worker.max");
+			_threadPoolWorkerMin = _meter.CreateGauge<int>("horde.clr.threadpool.worker.min");
+			_threadPoolWorkerFree = _meter.CreateGauge<int>("horde.clr.threadpool.worker.free");
+			_threadPoolWorkerBusy = _meter.CreateGauge<int>("horde.clr.threadpool.worker.busy");
+			_clrGcTotalMem = _meter.CreateGauge<long>("horde.clr.gc.totalMemory");
+			_clrGcTotalAllocated = _meter.CreateGauge<long>("horde.clr.gc.totalAllocated");
+			_clrGcHeapSize = _meter.CreateGauge<long>("horde.clr.gc.heapSize");
+			_clrGcFragmented = _meter.CreateGauge<long>("horde.clr.gc.fragmented");
+			_clrGcMemoryLoad = _meter.CreateGauge<long>("horde.clr.gc.memoryLoad");
+			_clrGcTotalAvailableMem = _meter.CreateGauge<long>("horde.clr.gc.totalAvailableMemory");
+			_clrGcHighMemThreshold = _meter.CreateGauge<long>("horde.clr.gc.highMemoryLoadThreshold");
 		}
 
 		/// <inheritdoc/>
@@ -56,29 +89,28 @@ namespace Horde.Server.Server
 
 			int busyIoThreads = maxIoThreads - freeIoThreads;
 			int busyWorkerThreads = maxWorkerThreads - freeWorkerThreads;
-
-			_dogStatsd.Gauge("horde.clr.threadpool.io.max", maxIoThreads);
-			_dogStatsd.Gauge("horde.clr.threadpool.io.min", minIoThreads);
-			_dogStatsd.Gauge("horde.clr.threadpool.io.free", freeIoThreads);
-			_dogStatsd.Gauge("horde.clr.threadpool.io.busy", busyIoThreads);
-
-			_dogStatsd.Gauge("horde.clr.threadpool.worker.max", maxWorkerThreads);
-			_dogStatsd.Gauge("horde.clr.threadpool.worker.min", minWorkerThreads);
-			_dogStatsd.Gauge("horde.clr.threadpool.worker.free", freeWorkerThreads);
-			_dogStatsd.Gauge("horde.clr.threadpool.worker.busy", busyWorkerThreads);
+			
+			_threadPoolIoMax.Record(maxIoThreads);
+			_threadPoolIoMin.Record(minIoThreads);
+			_threadPoolIoFree.Record(freeIoThreads);
+			_threadPoolIoBusy.Record(busyIoThreads);
+			_threadPoolWorkerMax.Record(maxWorkerThreads);
+			_threadPoolWorkerMin.Record(minWorkerThreads);
+			_threadPoolWorkerFree.Record(freeWorkerThreads);
+			_threadPoolWorkerBusy.Record(busyWorkerThreads);
 		}
 
 		private void ReportGcMetrics()
 		{
 			GCMemoryInfo gcMemoryInfo = GC.GetGCMemoryInfo();
 
-			_dogStatsd.Gauge("horde.clr.gc.totalMemory", GC.GetTotalMemory(false));
-			_dogStatsd.Gauge("horde.clr.gc.totalAllocated", GC.GetTotalAllocatedBytes());
-			_dogStatsd.Gauge("horde.clr.gc.heapSize", gcMemoryInfo.HeapSizeBytes);
-			_dogStatsd.Gauge("horde.clr.gc.fragmented", gcMemoryInfo.FragmentedBytes);
-			_dogStatsd.Gauge("horde.clr.gc.memoryLoad", gcMemoryInfo.MemoryLoadBytes);
-			_dogStatsd.Gauge("horde.clr.gc.totalAvailableMemory", gcMemoryInfo.TotalAvailableMemoryBytes);
-			_dogStatsd.Gauge("horde.clr.gc.highMemoryLoadThreshold", gcMemoryInfo.HighMemoryLoadThresholdBytes);
+			_clrGcTotalMem.Record(GC.GetTotalMemory(false));
+			_clrGcTotalAllocated.Record(GC.GetTotalAllocatedBytes());
+			_clrGcHeapSize.Record(gcMemoryInfo.HeapSizeBytes);
+			_clrGcFragmented.Record(gcMemoryInfo.FragmentedBytes);
+			_clrGcMemoryLoad.Record(gcMemoryInfo.MemoryLoadBytes);
+			_clrGcTotalAvailableMem.Record(gcMemoryInfo.TotalAvailableMemoryBytes);
+			_clrGcHighMemThreshold.Record(gcMemoryInfo.HighMemoryLoadThresholdBytes);
 		}
 	}
 }
diff --git a/Engine/Source/Programs/Horde/Horde.Server/Startup.cs b/Engine/Source/Programs/Horde/Horde.Server/Startup.cs
--- a/Engine/Source/Programs/Horde/Horde.Server/Startup.cs
+++ b/Engine/Source/Programs/Horde/Horde.Server/Startup.cs
@@ -82,7 +82,6 @@ using MongoDB.Bson.Serialization.Conventions;
 using MongoDB.Bson.Serialization.Serializers;
 using Serilog;
 using Serilog.Events;
-using StatsdClient;
 using Status = Grpc.Core.Status;
 using Horde.Server.Users;
 using Horde.Server.Perforce;
@@ -448,24 +447,6 @@ namespace Horde.Server
 			services.AddSingleton<CredentialService>();
 			services.AddSingleton<MongoService>();
 			services.AddSingleton<GlobalsService>();
-			services.AddSingleton<IDogStatsd>(ctx =>
-			{
-				string? datadogAgentHost = Environment.GetEnvironmentVariable("DD_AGENT_HOST");
-				if (datadogAgentHost != null)
-				{
-					// Datadog agent is configured, enable DogStatsD for metric collection
-					StatsdConfig config = new StatsdConfig
-					{
-						StatsdServerName = datadogAgentHost,
-						StatsdPort = 8125,
-					};
-
-					DogStatsdService dogStatsdService = new DogStatsdService();
-					dogStatsdService.Configure(config);
-					return dogStatsdService;
-				}
-				return new NoOpDogStatsd();
-			});
 			services.AddSingleton<ICommitService, CommitService>();
 			services.AddSingleton<IClock, Clock>();
 			services.AddSingleton<IDowntimeService, DowntimeService>();
diff --git a/Engine/Source/Programs/Horde/Horde.Server/Utilities/NoOpDogStatsd.cs b/Engine/Source/Programs/Horde/Horde.Server/Utilities/NoOpDogStatsd.cs
deleted file mode 100644
diff --git a/Engine/Source/Programs/Horde/Horde.Server/Utilities/OpenTelemetry.cs b/Engine/Source/Programs/Horde/Horde.Server/Utilities/OpenTelemetry.cs
--- a/Engine/Source/Programs/Horde/Horde.Server/Utilities/OpenTelemetry.cs
+++ b/Engine/Source/Programs/Horde/Horde.Server/Utilities/OpenTelemetry.cs
@@ -326,6 +326,18 @@ public static class OpenTelemetrySpanExtensions
 		
 		return span;
 	}
+	
+	/// <summary>
+	/// Create a metrics Gauge object.
+	/// </summary>
+	/// <param name="meter">Current meter being extended</param>
+	/// <param name="name">The instrument name. cannot be null.</param>
+	/// <param name="unit">Optional instrument unit of measurements.</param>
+	/// <param name="description">Optional instrument description.</param>
+	/// <remarks>
+	/// Gauge is an Instrument which records absolute values
+	/// </remarks>
+	public static Gauge<T> CreateGauge<T>(this Meter meter, string name, string? unit = null, string? description = null) where T : struct => new (meter, name, unit, description);
 }
 
 /// <summary>
@@ -348,4 +360,29 @@ public static class OpenTelemetryMeters
 	/// Default meter used in Horde
 	/// </summary>
 	public static readonly Meter Horde = new (HordeName);
+}
+
+/// <summary>
+/// A non-observable gauge
+/// </summary>
+/// <typeparam name="T">Type of value being measured</typeparam>
+public sealed class Gauge<T> : Instrument<T> where T : struct
+{
+    internal Gauge(Meter meter, string name, string? unit, string? description) : base(meter, name, unit, description)
+    {
+        Publish();
+    }
+
+    /// <summary>
+    /// Record the current value of the measurement.
+    /// </summary>
+    /// <param name="value">Current value.</param>
+    public void Record(T value) => RecordMeasurement(value);
+
+    /// <summary>
+    /// Record the current value of the measurement.
+    /// </summary>
+    /// <param name="value">The increment measurement.</param>
+    /// <param name="tag">A key-value pair tag associated with the measurement.</param>
+    public void Record(T value, KeyValuePair<string, object?> tag) => RecordMeasurement(value, tag);
 }
\ No newline at end of file
