Horde: Log unhandled event types for ClickHouseTelemetrySink

Also lower the metric sending interval from Horde server to ClickHouse.

diff --git a/Engine/Source/Programs/Horde/Horde.Server/Server/RpcService.cs b/Engine/Source/Programs/Horde/Horde.Server/Server/RpcService.cs
--- a/Engine/Source/Programs/Horde/Horde.Server/Server/RpcService.cs
+++ b/Engine/Source/Programs/Horde/Horde.Server/Server/RpcService.cs
@@ -480,6 +480,7 @@ namespace Horde.Server.Server
 					case WrappedTelemetryEvent.EventOneofCase.AgentMetadata: _telemetrySink.SendEvent("Agent.Metadata", e.AgentMetadata); break;
 					case WrappedTelemetryEvent.EventOneofCase.Cpu: _telemetrySink.SendEvent("Agent.Cpu", e.Cpu); break;
 					case WrappedTelemetryEvent.EventOneofCase.Mem: _telemetrySink.SendEvent("Agent.Memory", e.Mem); break;
+					default: _logger.LogError("Unhandled wrapped telemetry type {Type}", e.EventCase.ToString()); break;
 				}
 			}
 
diff --git a/Engine/Source/Programs/Horde/Horde.Server/Telemetry/ClickHouseTelemetrySink.cs b/Engine/Source/Programs/Horde/Horde.Server/Telemetry/ClickHouseTelemetrySink.cs
--- a/Engine/Source/Programs/Horde/Horde.Server/Telemetry/ClickHouseTelemetrySink.cs
+++ b/Engine/Source/Programs/Horde/Horde.Server/Telemetry/ClickHouseTelemetrySink.cs
@@ -42,7 +42,7 @@ namespace Horde.Server.Telemetry
 		public ClickHouseTelemetrySink(ClickHouseTelemetryConfig config, IHttpClientFactory httpClientFactory, IClock clock, ILogger<ClickHouseTelemetrySink> logger)
 		{
 			_httpClientFactory = httpClientFactory;
-			_ticker = clock.AddTicker<ClickHouseTelemetrySink>(TimeSpan.FromSeconds(10.0), FlushAsync, logger);
+			_ticker = clock.AddTicker<ClickHouseTelemetrySink>(TimeSpan.FromSeconds(5.0), FlushAsync, logger);
 			_logger = logger;
 			_uri = config.Url;
 		}
@@ -75,6 +75,7 @@ namespace Horde.Server.Telemetry
 					case AgentMetadataEvent agentMetadata: agentMetadataEvents.Add(agentMetadata); break;
 					case AgentCpuMetricsEvent agentCpu: agentCpuEvents.Add(agentCpu); break;
 					case AgentMemoryMetricsEvent agentMem: agentMemEvents.Add(agentMem); break;
+					default: _logger.LogError("Unable to parse event type for {EventName}", tuple.eventName); break;
 				}
 				c++;
 			}
