Horde: When retrying a job step, also retry all its dependencies.

diff --git a/Engine/Source/Programs/Horde/Horde.Server.Tests/JobCollectionTests.cs b/Engine/Source/Programs/Horde/Horde.Server.Tests/JobCollectionTests.cs
--- a/Engine/Source/Programs/Horde/Horde.Server.Tests/JobCollectionTests.cs
+++ b/Engine/Source/Programs/Horde/Horde.Server.Tests/JobCollectionTests.cs
@@ -19,6 +19,7 @@ using Microsoft.VisualStudio.TestTools.UnitTesting;
 using MongoDB.Bson;
 using Moq;
 using Horde.Server.Agents.Sessions;
+using Horde.Server.Users;
 
 namespace Horde.Server.Tests
 {
@@ -27,19 +28,19 @@ namespace Horde.Server.Tests
 	{
 		static NewGroup AddGroup(List<NewGroup> groups)
 		{
-			NewGroup @group = new NewGroup("win64", new List<NewNode>());
-			groups.Add(@group);
-			return @group;
+			NewGroup group = new NewGroup("win64", new List<NewNode>());
+			groups.Add(group);
+			return group;
 		}
 
-		static NewNode AddNode(NewGroup @group, string name, string[]? inputDependencies, Action<NewNode>? action = null)
+		static NewNode AddNode(NewGroup group, string name, string[]? inputDependencies, Action<NewNode>? action = null)
 		{
 			NewNode node = new NewNode(name, inputDependencies: inputDependencies?.ToList(), orderDependencies: inputDependencies?.ToList());
 			if (action != null)
 			{
 				action.Invoke(node);
 			}
-			@group.Nodes.Add(node);
+			group.Nodes.Add(node);
 			return node;
 		}
 
@@ -321,6 +322,166 @@ namespace Horde.Server.Tests
 			Assert.AreEqual(JobStepError.Incomplete, job.Batches[2].Steps[0].Error);
 		}
 
+		[TestMethod]
+		public async Task RetryStepSameGroupAsync()
+		{
+			Mock<ITemplate> templateMock = new Mock<ITemplate>(MockBehavior.Strict);
+			templateMock.SetupGet(x => x.InitialAgentType).Returns((string?)null);
+
+			List<NewGroup> newGroups = new List<NewGroup>();
+
+			NewGroup group = AddGroup(newGroups);
+			AddNode(group, "Step 1", null);
+			AddNode(group, "Step 2", new[] { "Step 1" });
+
+			IGraph graph = await GraphCollection.AppendAsync(null, newGroups, null, null);
+
+			CreateJobOptions options = new CreateJobOptions();
+			options.Arguments.Add("-Target=Step 2");
+
+			IJob job = await JobCollection.AddAsync(JobId.GenerateNewId(), new StreamId("ue4-main"), new TemplateId("test-build"), ContentHash.SHA1("hello"), graph, "Test job", 123, 123, options);
+
+			// Fail the first step
+			job = await StartBatch(job, graph, 0);
+			job = await RunStep(job, graph, 0, 0, JobStepOutcome.Failure);
+
+			Assert.AreEqual(1, job.Batches.Count);
+			Assert.AreEqual(2, job.Batches[0].Steps.Count);
+			Assert.AreEqual(JobStepState.Completed, job.Batches[0].Steps[0].State);
+			Assert.AreEqual(JobStepOutcome.Failure, job.Batches[0].Steps[0].Outcome);
+			Assert.AreEqual(JobStepState.Skipped, job.Batches[0].Steps[1].State);
+			Assert.AreEqual(JobStepOutcome.Failure, job.Batches[0].Steps[1].Outcome);
+
+			// Retry the failed step
+			job = Deref(await JobCollection.TryUpdateStepAsync(job, graph, job.Batches[0].Id, job.Batches[0].Steps[0].Id, newRetryByUserId: UserId.Anonymous));
+			Assert.AreEqual(1, job.Batches.Count);
+			Assert.AreEqual(3, job.Batches[0].Steps.Count);
+
+			IJobStep step = job.Batches[0].Steps[0];
+			Assert.AreEqual(0, step.NodeIdx);
+			Assert.AreEqual(JobStepState.Completed, step.State);
+			Assert.AreEqual(JobStepOutcome.Failure, step.Outcome);
+
+			step = job.Batches[0].Steps[1];
+			Assert.AreEqual(0, step.NodeIdx);
+			Assert.AreEqual(JobStepState.Ready, step.State);
+			Assert.AreEqual(JobStepOutcome.Success, step.Outcome);
+
+			step = job.Batches[0].Steps[2];
+			Assert.AreEqual(1, step.NodeIdx);
+			Assert.AreEqual(JobStepState.Waiting, step.State);
+			Assert.AreEqual(JobStepOutcome.Success, step.Outcome);
+
+			// Fail the retried step
+			job = await RunStep(job, graph, 0, 1, JobStepOutcome.Failure);
+			Assert.AreEqual(1, job.Batches.Count);
+			Assert.AreEqual(3, job.Batches[0].Steps.Count);
+
+			step = job.Batches[0].Steps[0];
+			Assert.AreEqual(0, step.NodeIdx);
+			Assert.AreEqual(JobStepState.Completed, step.State);
+			Assert.AreEqual(JobStepOutcome.Failure, step.Outcome);
+
+			step = job.Batches[0].Steps[1];
+			Assert.AreEqual(0, step.NodeIdx);
+			Assert.AreEqual(JobStepState.Completed, step.State);
+			Assert.AreEqual(JobStepOutcome.Failure, step.Outcome);
+
+			step = job.Batches[0].Steps[2];
+			Assert.AreEqual(1, step.NodeIdx);
+			Assert.AreEqual(JobStepState.Skipped, step.State);
+			Assert.AreEqual(JobStepOutcome.Failure, step.Outcome);
+
+			// Retry the failed step for a second time
+			job = Deref(await JobCollection.TryUpdateStepAsync(job, graph, job.Batches[0].Id, job.Batches[0].Steps[1].Id, newRetryByUserId: UserId.Anonymous));
+			Assert.AreEqual(1, job.Batches.Count);
+			Assert.AreEqual(4, job.Batches[0].Steps.Count);
+
+			step = job.Batches[0].Steps[0];
+			Assert.AreEqual(0, step.NodeIdx);
+			Assert.AreEqual(JobStepState.Completed, step.State);
+			Assert.AreEqual(JobStepOutcome.Failure, step.Outcome);
+
+			step = job.Batches[0].Steps[1];
+			Assert.AreEqual(0, step.NodeIdx);
+			Assert.AreEqual(JobStepState.Completed, step.State);
+			Assert.AreEqual(JobStepOutcome.Failure, step.Outcome);
+
+			step = job.Batches[0].Steps[2];
+			Assert.AreEqual(0, step.NodeIdx);
+			Assert.AreEqual(JobStepState.Ready, step.State);
+			Assert.AreEqual(JobStepOutcome.Success, step.Outcome);
+
+			step = job.Batches[0].Steps[3];
+			Assert.AreEqual(1, step.NodeIdx);
+			Assert.AreEqual(JobStepState.Waiting, step.State);
+			Assert.AreEqual(JobStepOutcome.Success, step.Outcome);
+		}
+
+		[TestMethod]
+		public async Task RetryStepDifferentGroupAsync()
+		{
+			Mock<ITemplate> templateMock = new Mock<ITemplate>(MockBehavior.Strict);
+			templateMock.SetupGet(x => x.InitialAgentType).Returns((string?)null);
+
+			List<NewGroup> newGroups = new List<NewGroup>();
+
+			NewGroup group = AddGroup(newGroups);
+			AddNode(group, "Step 1", null);
+
+			NewGroup group2 = AddGroup(newGroups);
+			AddNode(group2, "Step 2", new[] { "Step 1" });
+
+			IGraph graph = await GraphCollection.AppendAsync(null, newGroups, null, null);
+
+			CreateJobOptions options = new CreateJobOptions();
+			options.Arguments.Add("-Target=Step 2");
+
+			IJob job = await JobCollection.AddAsync(JobId.GenerateNewId(), new StreamId("ue4-main"), new TemplateId("test-build"), ContentHash.SHA1("hello"), graph, "Test job", 123, 123, options);
+
+			// Fail the first step
+			job = await StartBatch(job, graph, 0);
+			job = await RunStep(job, graph, 0, 0, JobStepOutcome.Failure);
+			job = Deref(await JobCollection.TryUpdateBatchAsync(job, graph, job.Batches[0].Id, null, JobStepBatchState.Complete, null));
+
+			Assert.AreEqual(2, job.Batches.Count);
+			Assert.AreEqual(1, job.Batches[0].Steps.Count);
+			Assert.AreEqual(1, job.Batches[1].Steps.Count);
+
+			IJobStep step = job.Batches[0].Steps[0];
+			Assert.AreEqual(JobStepState.Completed, step.State);
+			Assert.AreEqual(JobStepOutcome.Failure, step.Outcome);
+
+			step = job.Batches[1].Steps[0];
+			Assert.AreEqual(JobStepState.Skipped, step.State);
+			Assert.AreEqual(JobStepOutcome.Failure, step.Outcome);
+
+			// Retry the failed step
+			job = Deref(await JobCollection.TryUpdateStepAsync(job, graph, job.Batches[0].Id, job.Batches[0].Steps[0].Id, newRetryByUserId: UserId.Anonymous));
+			Assert.AreEqual(3, job.Batches.Count);
+			Assert.AreEqual(1, job.Batches[0].Steps.Count);
+			Assert.AreEqual(1, job.Batches[1].Steps.Count);
+			Assert.AreEqual(1, job.Batches[2].Steps.Count);
+
+			step = job.Batches[0].Steps[0];
+			Assert.AreEqual(0, job.Batches[0].GroupIdx);
+			Assert.AreEqual(0, step.NodeIdx);
+			Assert.AreEqual(JobStepState.Completed, step.State);
+			Assert.AreEqual(JobStepOutcome.Failure, step.Outcome);
+
+			step = job.Batches[1].Steps[0];
+			Assert.AreEqual(0, job.Batches[1].GroupIdx);
+			Assert.AreEqual(0, step.NodeIdx);
+			Assert.AreEqual(JobStepState.Ready, step.State);
+			Assert.AreEqual(JobStepOutcome.Success, step.Outcome);
+
+			step = job.Batches[2].Steps[0];
+			Assert.AreEqual(1, job.Batches[2].GroupIdx);
+			Assert.AreEqual(0, step.NodeIdx);
+			Assert.AreEqual(JobStepState.Waiting, step.State);
+			Assert.AreEqual(JobStepOutcome.Success, step.Outcome);
+		}
+
 		[TestMethod]
 		public async Task UnknownShelfAsync()
 		{
diff --git a/Engine/Source/Programs/Horde/Horde.Server/Jobs/JobCollection.cs b/Engine/Source/Programs/Horde/Horde.Server/Jobs/JobCollection.cs
--- a/Engine/Source/Programs/Horde/Horde.Server/Jobs/JobCollection.cs
+++ b/Engine/Source/Programs/Horde/Horde.Server/Jobs/JobCollection.cs
@@ -2,9 +2,11 @@
 
 using System;
 using System.Collections.Generic;
+using System.Diagnostics;
 using System.Diagnostics.CodeAnalysis;
 using System.IO;
 using System.Linq;
+using System.Text;
 using System.Threading.Tasks;
 using EpicGames.Core;
 using Horde.Server.Acls;
@@ -21,6 +23,7 @@ using Horde.Server.Users;
 using Horde.Server.Utilities;
 using HordeCommon;
 using HordeCommon.Rpc.Tasks;
+using Microsoft.CodeAnalysis.CSharp.Syntax;
 using Microsoft.Extensions.Logging;
 using MongoDB.Bson;
 using MongoDB.Bson.IO;
@@ -110,6 +113,21 @@ namespace Horde.Server.Jobs
 				Id = id;
 				NodeIdx = nodeIdx;
 			}
+
+			public override string ToString()
+			{
+				StringBuilder description = new StringBuilder($"{Id}: {State}");
+				if (Outcome != JobStepOutcome.Unspecified)
+				{
+					description.Append($" ({Outcome}");
+					if (Error != JobStepError.None)
+					{
+						description.Append($" - {Error}");
+					}
+					description.Append($")");
+				}
+				return description.ToString();
+			}
 		}
 
 		class JobStepBatchDocument : IJobStepBatch
@@ -168,6 +186,17 @@ namespace Horde.Server.Jobs
 				Id = id;
 				GroupIdx = groupIdx;
 			}
+
+			public override string ToString()
+			{
+				StringBuilder description = new StringBuilder($"{Id}: {State}");
+				if (Error != JobStepBatchError.None)
+				{
+					description.Append($" - {Error}");
+				}
+				description.Append($" ({Steps.Count} steps");
+				return description.ToString();
+			}
 		}
 
 		class ChainedJobDocument : IChainedJob
@@ -1329,9 +1358,16 @@ namespace Horde.Server.Jobs
 					{
 						failedNodes.Add(node);
 					}
-					else if (step.State == JobStepState.Skipped && (node.InputDependencies.Any(x => failedNodes.Contains(graph.GetNode(x))) || !CanRetryNode(job, batch.GroupIdx, step.NodeIdx)))
+					else if (step.State == JobStepState.Skipped)
 					{
-						failedNodes.Add(node);
+						if(node.InputDependencies.Any(x => failedNodes.Contains(graph.GetNode(x))) || !CanRetryNode(job, batch.GroupIdx, step.NodeIdx))
+						{
+							failedNodes.Add(node);
+						}
+						else
+						{
+							failedNodes.Remove(node);
+						}
 					}
 					else if (step.Outcome == JobStepOutcome.Failure)
 					{
@@ -1511,13 +1547,11 @@ namespace Horde.Server.Jobs
 
 							appendToBatches[groupIdx] = batch;
 						}
-						if (batch.Steps.Count == 0 || nodeIdx > batch.Steps[^1].NodeIdx)
-						{
-							job.NextSubResourceId = job.NextSubResourceId.Next();
 
-							JobStepDocument step = new JobStepDocument(job.NextSubResourceId, nodeIdx);
-							batch.Steps.Add(step);
-						}
+						job.NextSubResourceId = job.NextSubResourceId.Next();
+
+						JobStepDocument step = new JobStepDocument(job.NextSubResourceId, nodeIdx);
+						batch.Steps.Add(step);
 					}
 				}
 			}
