Horde Docs view tweaks
#skipci
#rnx

diff --git a/Engine/Source/Programs/Horde/HordeDashboard/src/base/components/Markdown/Markdown.tsx b/Engine/Source/Programs/Horde/HordeDashboard/src/base/components/Markdown/Markdown.tsx
--- a/Engine/Source/Programs/Horde/HordeDashboard/src/base/components/Markdown/Markdown.tsx
+++ b/Engine/Source/Programs/Horde/HordeDashboard/src/base/components/Markdown/Markdown.tsx
@@ -16,115 +16,126 @@ import { PropsWithChildren } from 'react';
 // (as having a default export) and the way it actually builds its files (for its cjs `main` file,
 // assigning to module.exports and not setting a default export)
 const MarkdownComponent: typeof MarkdownComponentType =
-  // eslint-disable-next-line @typescript-eslint/no-explicit-any
-  (MarkdownModule as any).default || (MarkdownModule as any);
+   // eslint-disable-next-line @typescript-eslint/no-explicit-any
+   (MarkdownModule as any).default || (MarkdownModule as any);
 
 const getStyles: IStyleFunction<IMarkdownStyleProps, IMarkdownStyles> = () => {
-  const imageStyles: Partial<IImageStyles> = {
-    root: {
-      maxWidth: '100%',
-      margin: '8px 0',
-    },
-  };
+   const imageStyles: Partial<IImageStyles> = {
+      root: {
+         maxWidth: '100%',
+         margin: '8px 0',
+      },
+   };
 
-  return {
-     root: { fontSize: "13px", fontFamily: "Horde Open Sans Regular" },
-     subComponentStyles: { image: imageStyles } as IMarkdownSubComponentStyles,
-  };
+   return {
+      root: {
+         fontSize: "13px", fontFamily: "Horde Open Sans Regular", selectors: {
+            'li': {
+               margin: "4px 0px 8px 0px",
+               lineHeight:"1.6"
+            },
+            'ul': {
+               margin: "4px 0px 8px 0px"
+            }
+
+         }
+      },
+      subComponentStyles: { image: imageStyles } as IMarkdownSubComponentStyles,
+   };
 };
 
 const getClassNames = classNamesFunction<IMarkdownStyleProps, IMarkdownStyles>();
 
 const MarkdownBase: React.FunctionComponent<PropsWithChildren<IMarkdownProps>> = props => {
-  const { styles, theme, children } = props;
-  const classNames = getClassNames(styles, { theme: theme! });
+   const { styles, theme, children } = props;
+   const classNames = getClassNames(styles, { theme: theme! });
 
-  return (
-    <div className={classNames.root}>
-      <MarkdownComponent options={{ overrides: getOverrides(classNames.subComponentStyles, props) }}>
-        {children as string}
-      </MarkdownComponent>
-    </div>
-  );
+   return (
+      <div className={classNames.root}>
+         <MarkdownComponent options={{ overrides: getOverrides(classNames.subComponentStyles, props) }}>
+            {children as string}
+         </MarkdownComponent>
+      </div>
+   );
 };
 MarkdownBase.displayName = 'Markdown';
 
 function getOverrides(subComponentStyles: IMarkdownSubComponentStyles, props: IMarkdownProps): MarkdownToJSX.Overrides {
-  return {
-    h1: {
-      component: MarkdownHeader,
-      props: { styles: subComponentStyles.header },
-    },
-    h2: {
-      component: MarkdownHeader,
-      props: { as: 'h2', styles: subComponentStyles.header },
-    },
-    h3: {
-      component: MarkdownHeader,
-      props: { as: 'h3', styles: subComponentStyles.header },
-    },
-    h4: {
-      component: MarkdownHeader,
-      props: { as: 'h4', styles: subComponentStyles.header },
-    },
-    h5: {
-      component: MarkdownHeader,
-      props: { as: 'h5', styles: subComponentStyles.header },
-    },
-    h6: {
-      component: MarkdownHeader,
-      props: { as: 'h6', styles: subComponentStyles.header },
-    },
-    p: {
-      component: MarkdownParagraph,
-      props: { styles: subComponentStyles.paragraph },
-    },
-    pre: {
-      component: MarkdownPre,
-      props: { enableRenderHtmlBlock: props.enableRenderHtmlBlock },
-    },
-    a: {
-      component: MarkdownLink,
-      props: { className: 'ms-mdLink', styles: subComponentStyles.link },
-    },
-    img: {
-      component: Image,
-      props: { className: 'ms-mdImage', styles: subComponentStyles.image },
-    },
-    button: {
-      component: DefaultButton,
-      props: { className: 'ms-mdButton', styles: subComponentStyles.button },
-    },
-    table: {
-      component: MDTable.MarkdownTable,
-      props: { styles: subComponentStyles.table },
-    },
-    thead: {
-      component: MDTable.MarkdownTHead,
-      props: { styles: subComponentStyles.table },
-    },
-    tbody: {
-      component: MDTable.MarkdownTBody,
-      props: { styles: subComponentStyles.table },
-    },
-    tr: {
-      component: MDTable.MarkdownTr,
-      props: { styles: subComponentStyles.table },
-    },
-    th: {
-      component: MDTable.MarkdownCell,
-      props: { as: 'th', styles: subComponentStyles.table },
-    },
-    td: {
-      component: MDTable.MarkdownCell,
-      props: { as: 'td', styles: subComponentStyles.table },
-    },
-    ...props.overrides,
-  };
+   return {
+      h1: {
+         component: MarkdownHeader,
+         props: { styles: subComponentStyles.header },
+      },
+      h2: {
+         component: MarkdownHeader,
+         props: { as: 'h2', styles: subComponentStyles.header },
+      },
+      h3: {
+         component: MarkdownHeader,
+         props: { as: 'h3', styles: subComponentStyles.header },
+      },
+      h4: {
+         component: MarkdownHeader,
+         props: { as: 'h4', styles: subComponentStyles.header },
+      },
+      h5: {
+         component: MarkdownHeader,
+         props: { as: 'h5', styles: subComponentStyles.header },
+      },
+      h6: {
+         component: MarkdownHeader,
+         props: { as: 'h6', styles: subComponentStyles.header },
+      },
+      p: {
+         component: MarkdownParagraph,
+         props: { styles: subComponentStyles.paragraph },
+      },
+      pre: {
+         component: MarkdownPre,
+         props: { enableRenderHtmlBlock: props.enableRenderHtmlBlock },
+      },
+      a: {
+         component: MarkdownLink,
+         props: { className: 'ms-mdLink', styles: subComponentStyles.link },
+      },
+      img: {
+         component: Image,
+         props: { className: 'ms-mdImage', styles: subComponentStyles.image },
+      },
+      button: {
+         component: DefaultButton,
+         props: { className: 'ms-mdButton', styles: subComponentStyles.button },
+      },
+      table: {
+         component: MDTable.MarkdownTable,
+         props: { styles: subComponentStyles.table },
+      },
+      thead: {
+         component: MDTable.MarkdownTHead,
+         props: { styles: subComponentStyles.table },
+      },
+      tbody: {
+         component: MDTable.MarkdownTBody,
+         props: { styles: subComponentStyles.table },
+      },
+      tr: {
+         component: MDTable.MarkdownTr,
+         props: { styles: subComponentStyles.table },
+      },
+      th: {
+         component: MDTable.MarkdownCell,
+         props: { as: 'th', styles: subComponentStyles.table },
+      },
+      td: {
+         component: MDTable.MarkdownCell,
+         props: { as: 'td', styles: subComponentStyles.table },
+      },
+      ...props.overrides,
+   };
 }
 
 export const Markdown: React.FunctionComponent<PropsWithChildren<IMarkdownProps>> = styled<
-  IMarkdownProps,
-  IMarkdownStyleProps,
-  IMarkdownStyles
+   IMarkdownProps,
+   IMarkdownStyleProps,
+   IMarkdownStyles
 >(MarkdownBase, getStyles, undefined, { scope: 'Markdown' });
diff --git a/Engine/Source/Programs/Horde/HordeDashboard/src/base/components/Markdown/MarkdownHeader.tsx b/Engine/Source/Programs/Horde/HordeDashboard/src/base/components/Markdown/MarkdownHeader.tsx
--- a/Engine/Source/Programs/Horde/HordeDashboard/src/base/components/Markdown/MarkdownHeader.tsx
+++ b/Engine/Source/Programs/Horde/HordeDashboard/src/base/components/Markdown/MarkdownHeader.tsx
@@ -7,99 +7,104 @@ import { FontSizes, NeutralColors } from '@fluentui/theme';
  * The component props.
  */
 export interface IMarkdownHeaderProps {
-  as?: keyof React.ReactHTML;
-  children?: React.ReactNode;
-  className?: string;
-  id?: string;
+   as?: keyof React.ReactHTML;
+   children?: React.ReactNode;
+   className?: string;
+   id?: string;
 
-  theme?: ITheme;
-  styles?: IStyleFunctionOrObject<IMarkdownHeaderStyleProps, IMarkdownHeaderStyles>;
+   theme?: ITheme;
+   styles?: IStyleFunctionOrObject<IMarkdownHeaderStyleProps, IMarkdownHeaderStyles>;
 }
 
 /**
  * The getStyles props contract.
  */
 export type IMarkdownHeaderStyleProps = {
-  as: string;
-  className?: string;
+   as: string;
+   className?: string;
 };
 
 /**
  * The styles produced by getStyles.
  */
 export interface IMarkdownHeaderStyles {
-  root: IStyle;
+   root: IStyle;
 }
 
 const getStyles: IStyleFunction<IMarkdownHeaderStyleProps, IMarkdownHeaderStyles> = props => {
-  const { className } = props;
-  return {
-    root: [
-      {
-        fontSize: FontSizes.size16,
-        fontWeight: 600,
-        padding: 0,
-        margin: 0,
-        marginTop: '32px',
-        marginBottom: '12px',
+   const { className } = props;
+   return {
+      root: [
+         {
+            fontSize: FontSizes.size16,
+            fontWeight: 600,
+            padding: 0,
+            margin: 0,
+            marginTop: '32px',
+            marginBottom: '12px',
 
-        selectors: {
-          '&:first-child': {
-            marginTop: 0,
-          },
-        },
-      },
+            selectors: {
+               '&:first-child': {
+                  marginTop: 0,
+               },
+            },
+         },
 
-      props.as === 'h1' && [
-        {
-          fontSize: FontSizes.size32,
-          marginBottom: '28px',
-          color: NeutralColors.gray160,
-        },
-      ],
+         props.as === 'h1' && [
+            {
+               fontSize: FontSizes.size32,
+               marginBottom: '28px',
+               color: NeutralColors.gray160,
+            },
+         ],
 
-      props.as === 'h2' && [
-        {
-          fontSize: FontSizes.size24,
-          marginBottom: '20px',
-          color: NeutralColors.gray160,
-        },
-      ],
+         props.as === 'h2' && [
+            {
+               fontSize: FontSizes.size24,
+               marginTop: '24px',
+               marginRight: '0px',
+               marginBottom: '8px',
+               marginLeft: '0px',
+               color: NeutralColors.gray160,
 
-      props.as === 'h3' && [
-        {
-          fontSize: FontSizes.size20,
-          marginBottom: '8px',
-          color: NeutralColors.gray130,
-        },
-      ],
-      props.as === 'h4' && [
-        {
-          fontSize: FontSizes.size16,
-          marginBottom: '6px',
-          color: NeutralColors.gray160,
-        },
+
+            },
+         ],
+
+         props.as === 'h3' && [
+            {
+               fontSize: FontSizes.size20,
+               marginBottom: '8px',
+               color: NeutralColors.gray130,
+            },
+         ],
+         props.as === 'h4' && [
+            {
+               fontSize: FontSizes.size16,
+               marginBottom: '6px',
+               color: NeutralColors.gray160,
+            },
+         ],
+         className,
       ],
-      className,
-    ],
-  };
+   };
 };
 
 const getClassNames = classNamesFunction<IMarkdownHeaderStyleProps, IMarkdownHeaderStyles>();
 
 const MarkdownHeaderBase: React.FunctionComponent<IMarkdownHeaderProps> = props => {
-  const { as: RootType = 'h1', children, id, styles, className } = props;
+   const { as: RootType = 'h1', children, id, styles, className } = props;
 
-  const classNames = getClassNames(styles, { as: RootType, className });
-  return (
-    <RootType className={classNames.root} id={id}>
-      {children}
-    </RootType>
-  );
+   const classNames = getClassNames(styles, { as: RootType, className });
+   return (
+      <RootType className={classNames.root} id={id}>
+         {children}
+      </RootType>
+   );
 };
 
 export const MarkdownHeader: React.FunctionComponent<IMarkdownHeaderProps> = styled<
-  IMarkdownHeaderProps,
-  IMarkdownHeaderStyleProps,
-  IMarkdownHeaderStyles
+   IMarkdownHeaderProps,
+   IMarkdownHeaderStyleProps,
+   IMarkdownHeaderStyles
 >(MarkdownHeaderBase, getStyles);
diff --git a/Engine/Source/Programs/Horde/HordeDashboard/src/base/components/Markdown/MarkdownParagraph.tsx b/Engine/Source/Programs/Horde/HordeDashboard/src/base/components/Markdown/MarkdownParagraph.tsx
--- a/Engine/Source/Programs/Horde/HordeDashboard/src/base/components/Markdown/MarkdownParagraph.tsx
+++ b/Engine/Source/Programs/Horde/HordeDashboard/src/base/components/Markdown/MarkdownParagraph.tsx
@@ -1,58 +1,59 @@
 import * as React from 'react';
 import {
-  IClassNames,
-  IStyleFunction,
-  classNamesFunction,
-  styled,
-  IStyleFunctionOrObject,
+   IClassNames,
+   IStyleFunction,
+   classNamesFunction,
+   styled,
+   IStyleFunctionOrObject,
 } from '@fluentui/react/lib/Utilities';
 import { ITheme, IStyle } from '@fluentui/react/lib/Styling';
 import { PropsWithChildren } from 'react';
 
 export interface IMarkdownParagraphProps {
-  styles?: IStyleFunctionOrObject<IMarkdownParagraphStyleProps, IMarkdownParagraphStyles>;
-  theme?: ITheme;
+   styles?: IStyleFunctionOrObject<IMarkdownParagraphStyleProps, IMarkdownParagraphStyles>;
+   theme?: ITheme;
 }
 
 export interface IMarkdownParagraphStyles {
-  root: IStyle;
+   root: IStyle;
 }
 
 export interface IMarkdownParagraphStyleProps {
-  theme: ITheme;
-  isTodo: boolean;
+   theme: ITheme;
+   isTodo: boolean;
 }
 
 const getStyles: IStyleFunction<IMarkdownParagraphStyleProps, IMarkdownParagraphStyles> = props => {
-  const { theme, isTodo } = props;
-  return {
-    root: [
-      theme.fonts.medium,
-      {
-        marginBottom: 4,
-      },
-      isTodo && {
-        padding: 8,
-        background: theme.semanticColors.warningBackground,
-      },
-    ],
-  };
+   const { theme, isTodo } = props;
+   return {
+      root: [
+         theme.fonts.medium,
+         {
+            lineHeight: "1.6",
+            marginBottom: 4,
+         },
+         isTodo && {
+            padding: 8,
+            background: theme.semanticColors.warningBackground,
+         },
+      ],
+   };
 };
 
 const getClassNames = classNamesFunction<IMarkdownParagraphStyleProps, IMarkdownParagraphStyles>();
 
 const MarkdownParagraphBase: React.FunctionComponent<PropsWithChildren<IMarkdownParagraphProps>> = props => {
-  const { children, theme } = props;
-  const classNames: IClassNames<IMarkdownParagraphStyles> = getClassNames(props.styles, {
-    theme: theme!,
-    isTodo: typeof children === 'string' && children.indexOf('TODO') === 0,
-  });
+   const { children, theme } = props;
+   const classNames: IClassNames<IMarkdownParagraphStyles> = getClassNames(props.styles, {
+      theme: theme!,
+      isTodo: typeof children === 'string' && children.indexOf('TODO') === 0,
+   });
 
-  return <p className={classNames.root}>{children}</p>;
+   return <p className={classNames.root}>{children}</p>;
 };
 
 export const MarkdownParagraph: React.FunctionComponent<PropsWithChildren<IMarkdownParagraphProps>> = styled<
-  IMarkdownParagraphProps,
-  IMarkdownParagraphStyleProps,
-  IMarkdownParagraphStyles
+   IMarkdownParagraphProps,
+   IMarkdownParagraphStyleProps,
+   IMarkdownParagraphStyles
 >(MarkdownParagraphBase, getStyles);
diff --git a/Engine/Source/Programs/Horde/HordeDashboard/src/base/components/SideRail/SideRail.styles.ts b/Engine/Source/Programs/Horde/HordeDashboard/src/base/components/SideRail/SideRail.styles.ts
--- a/Engine/Source/Programs/Horde/HordeDashboard/src/base/components/SideRail/SideRail.styles.ts
+++ b/Engine/Source/Programs/Horde/HordeDashboard/src/base/components/SideRail/SideRail.styles.ts
@@ -42,7 +42,7 @@ export const getStyles: IStyleFunction<ISideRailStyleProps, ISideRailStyles> = p
           {
             display: 'block',
             flex: '1',
-            padding: '4px 0px',
+            padding: '4px 20px',
             selectors: {              
               ':active,:visited,:hover': { color: modeColors.text, textDecoration: "none", background: theme.palette.neutralLight }
             }
diff --git a/Engine/Source/Programs/Horde/HordeDashboard/src/base/components/SideRail/SideRail.tsx b/Engine/Source/Programs/Horde/HordeDashboard/src/base/components/SideRail/SideRail.tsx
--- a/Engine/Source/Programs/Horde/HordeDashboard/src/base/components/SideRail/SideRail.tsx
+++ b/Engine/Source/Programs/Horde/HordeDashboard/src/base/components/SideRail/SideRail.tsx
@@ -50,7 +50,7 @@ class SideRailBase extends React.Component<ISideRailProps, ISideRailState> {
       return jumpLinkList || relatedLinkList || contactLinkList ? (
          <FocusZone direction={FocusZoneDirection.vertical} className={this._classNames.root}>
             {relatedLinkList}
-            {jumpLinkList}                        
+            {jumpLinkList}
             {contactLinkList}
          </FocusZone>
       ) : null;
@@ -92,8 +92,8 @@ class SideRailBase extends React.Component<ISideRailProps, ISideRailState> {
       ));
       return (
          <div className={css(classNames.section, classNames.jumpLinkSection)}>
-            <Stack style={{paddingBottom: 8}}>
-               <Text variant='medium' style={{fontWeight:600}}>On this page</Text>
+            <Stack style={{ paddingBottom: 8 }}>
+               <Text variant='medium' style={{ fontWeight: 600 }}>On this page</Text>
             </Stack>
             <ul className={classNames.links}>{links}</ul>
          </div>
@@ -124,10 +124,12 @@ class SideRailBase extends React.Component<ISideRailProps, ISideRailState> {
       if (links) {
          return (
             <div className={css(classNames.section)}>
-            <Stack style={{paddingBottom: 8}}>
-               <Text variant='medium' style={{fontWeight:600}}>{title}</Text>
-            </Stack>
-               {links}
+               <Stack style={{ paddingBottom: 8 }}>
+                  <Text variant='medium' style={{ fontWeight: 600 }}>{title}</Text>
+               </Stack>
+               <Stack>
+                  {links}
+               </Stack>
             </div>
          );
       }
diff --git a/Engine/Source/Programs/Horde/HordeDashboard/src/components/docs/DocView.tsx b/Engine/Source/Programs/Horde/HordeDashboard/src/components/docs/DocView.tsx
--- a/Engine/Source/Programs/Horde/HordeDashboard/src/components/docs/DocView.tsx
+++ b/Engine/Source/Programs/Horde/HordeDashboard/src/components/docs/DocView.tsx
@@ -1,4 +1,4 @@
-import { Stack } from "@fluentui/react";
+import { Stack, mergeStyleSets } from "@fluentui/react";
 import { action, makeObservable, observable } from "mobx";
 import { observer } from "mobx-react-lite";
 import { useRef, useState } from "react";
@@ -24,6 +24,16 @@ type State = {
    jumpLinks: ISideRailLink[];
 }
 
+
+export const docClasses = mergeStyleSets({
+   raised: {
+      backgroundColor: "#ffffff",
+      boxShadow: "0 1.6px 3.6px 0 rgba(0,0,0,0.132), 0 0.3px 0.9px 0 rgba(0,0,0,0.108)",
+      padding: "16px 32px"      
+   }
+});
+
+
 class LinkState {
 
    constructor() {
@@ -63,7 +73,7 @@ const DocPanel: React.FC<{ docName: string }> = ({ docName }) => {
 
                // generate anchors
                let lines = (textContent.match(/[^\r\n]+/g) ?? []) as string[];
-               lines = lines.map(i => i.trim()).filter( i => !!i);
+               lines = lines.map(i => i.trim()).filter(i => !!i);
                lines.forEach(line => {
                   let anchor = "";
 
@@ -77,9 +87,9 @@ const DocPanel: React.FC<{ docName: string }> = ({ docName }) => {
 
                   if (anchor) {
                      anchor = anchor.trim();
-                     anchors.push({text: anchor, anchor: anchor.replace(/[^a-z0-9- ]/gi, '').replace(/ /gi, '-').toLowerCase() })
+                     anchors.push({ text: anchor, anchor: anchor.replace(/[^a-z0-9- ]/gi, '').replace(/ /gi, '-').toLowerCase() })
                   }
-               })               
+               })
 
                // generate crumbs
                if (lines && lines.length) {
@@ -141,23 +151,15 @@ const DocPanel: React.FC<{ docName: string }> = ({ docName }) => {
    let anchors = anchorCache.get(docName) ?? [];
 
    linkState.setState(crumbs, anchors.map(a => {
-      return { text: a.text, url: a.anchor}
+      return { text: a.text, url: a.anchor }
    }));
 
 
-   return <Stack>
-      <Stack styles={{ root: { paddingTop: 0, paddingLeft: 12, paddingRight: 12, width: "100%" } }} >
-         <Stack tokens={{ childrenGap: 12 }}>
-            <Stack style={{ paddingBottom: 24, paddingLeft: 6, paddingRight: 6, paddingTop: 12 }}>
-               <div ref={docRef}>
-                  <Markdown>{text}</Markdown>
-               </div>
-            </Stack>
-         </Stack>
-      </Stack>
+   return <Stack styles={{ root: { width: "100%" } }} >
+      <div ref={docRef}>
+         <Markdown>{text}</Markdown>
+      </div>
    </Stack>;
-
-
 }
 
 const DocRail = observer(() => {
@@ -218,7 +220,7 @@ export const DocView = () => {
             <div style={{ overflowY: 'scroll', overflowX: 'hidden', height: "calc(100vh - 162px)" }} data-is-scrollable={true}>
                <Stack horizontal>
                   <Stack style={{ width: 1240, paddingTop: 6, marginLeft: 4, height: '100%' }}>
-                     <Stack className={hordeClasses.raised}>
+                     <Stack className={docClasses.raised}>
                         <Stack style={{ width: "100%", height: "max-content" }} tokens={{ childrenGap: 18 }}>
                            <DocPanel docName={docName} />
                         </Stack>
