Horde: More markdown rendering tweaks
#rnx
#skipci

diff --git a/Engine/Source/Programs/Horde/HordeDashboard/src/base/components/Markdown/Markdown.tsx b/Engine/Source/Programs/Horde/HordeDashboard/src/base/components/Markdown/Markdown.tsx
--- a/Engine/Source/Programs/Horde/HordeDashboard/src/base/components/Markdown/Markdown.tsx
+++ b/Engine/Source/Programs/Horde/HordeDashboard/src/base/components/Markdown/Markdown.tsx
@@ -35,7 +35,7 @@ const getStyles: IStyleFunction<IMarkdownStyleProps, IMarkdownStyles> = () => {
                lineHeight:"1.6"
             },
             'ul': {
-               margin: "4px 0px 8px 0px"
+               margin: "8px 0px 16px 0px"
             }
 
          }
diff --git a/Engine/Source/Programs/Horde/HordeDashboard/src/components/TopNav.tsx b/Engine/Source/Programs/Horde/HordeDashboard/src/components/TopNav.tsx
--- a/Engine/Source/Programs/Horde/HordeDashboard/src/components/TopNav.tsx
+++ b/Engine/Source/Programs/Horde/HordeDashboard/src/components/TopNav.tsx
@@ -455,25 +455,23 @@ export const TopNav: React.FC<{ suppressServer?: boolean }> = observer(({ suppre
          }
       });
 
-      if (getSiteConfig().environment !== "production") {
 
-         const docsItems: IContextualMenuItem[] = [];
-         docsItems.push({
-            key: "server_docs",
-            text: "Documentation",
-            link: `/docs`
-         });
+      const docsItems: IContextualMenuItem[] = [];
+      docsItems.push({
+         key: "server_docs",
+         text: "Documentation",
+         link: `/docs`
+      });
 
-         subItems.push({
-            itemType: ContextualMenuItemType.Section,
-            key: `server_docs_item`,
-            sectionProps: {
-               title: "Documentation",
-               items: docsItems,
-               bottomDivider: true
-            }
-         });
-      }
+      subItems.push({
+         itemType: ContextualMenuItemType.Section,
+         key: `server_docs_item`,
+         sectionProps: {
+            title: "Documentation",
+            items: docsItems,
+            bottomDivider: true
+         }
+      });
 
       // Configuration
       const versionItems: IContextualMenuItem[] = [];
diff --git a/Engine/Source/Programs/Horde/HordeDashboard/src/components/docs/DocView.tsx b/Engine/Source/Programs/Horde/HordeDashboard/src/components/docs/DocView.tsx
--- a/Engine/Source/Programs/Horde/HordeDashboard/src/components/docs/DocView.tsx
+++ b/Engine/Source/Programs/Horde/HordeDashboard/src/components/docs/DocView.tsx
@@ -29,7 +29,7 @@ export const docClasses = mergeStyleSets({
    raised: {
       backgroundColor: "#ffffff",
       boxShadow: "0 1.6px 3.6px 0 rgba(0,0,0,0.132), 0 0.3px 0.9px 0 rgba(0,0,0,0.108)",
-      padding: "16px 32px"      
+      padding: "32px 40px"      
    }
 });
 
