Horde: Retry upstream steps if a skipped node is flagged for retry.

diff --git a/Engine/Source/Programs/Horde/Horde.Server.Tests/JobCollectionTests.cs b/Engine/Source/Programs/Horde/Horde.Server.Tests/JobCollectionTests.cs
--- a/Engine/Source/Programs/Horde/Horde.Server.Tests/JobCollectionTests.cs
+++ b/Engine/Source/Programs/Horde/Horde.Server.Tests/JobCollectionTests.cs
@@ -482,6 +482,74 @@ namespace Horde.Server.Tests
 			Assert.AreEqual(JobStepOutcome.Success, step.Outcome);
 		}
 
+		[TestMethod]
+		public async Task RetryDownstreamStepAsync()
+		{
+			Mock<ITemplate> templateMock = new Mock<ITemplate>(MockBehavior.Strict);
+			templateMock.SetupGet(x => x.InitialAgentType).Returns((string?)null);
+
+			List<NewGroup> newGroups = new List<NewGroup>();
+
+			NewGroup group = AddGroup(newGroups);
+			AddNode(group, "Step 1", null);
+			AddNode(group, "Step 2", new[] { "Step 1" });
+			AddNode(group, "Step 3", new[] { "Step 2" });
+
+			IGraph graph = await GraphCollection.AppendAsync(null, newGroups, null, null);
+
+			CreateJobOptions options = new CreateJobOptions();
+			options.Arguments.Add("-Target=Step 3");
+
+			IJob job = await JobCollection.AddAsync(JobId.GenerateNewId(), new StreamId("ue4-main"), new TemplateId("test-build"), ContentHash.SHA1("hello"), graph, "Test job", 123, 123, options);
+
+			// Fail the first step
+			job = await StartBatch(job, graph, 0);
+			job = await RunStep(job, graph, 0, 0, JobStepOutcome.Failure);
+
+			Assert.AreEqual(1, job.Batches.Count);
+			Assert.AreEqual(3, job.Batches[0].Steps.Count);
+
+			IJobStep step = job.Batches[0].Steps[0];
+			Assert.AreEqual(0, step.NodeIdx);
+			Assert.AreEqual(JobStepState.Completed, step.State);
+			Assert.AreEqual(JobStepOutcome.Failure, step.Outcome);
+
+			step = job.Batches[0].Steps[1];
+			Assert.AreEqual(1, step.NodeIdx);
+			Assert.AreEqual(JobStepState.Skipped, step.State);
+			Assert.AreEqual(JobStepOutcome.Failure, step.Outcome);
+
+			step = job.Batches[0].Steps[2];
+			Assert.AreEqual(2, step.NodeIdx);
+			Assert.AreEqual(JobStepState.Skipped, step.State);
+			Assert.AreEqual(JobStepOutcome.Failure, step.Outcome);
+
+			// Retry the last failed step. All the failed upstream steps should run again.
+			job = Deref(await JobCollection.TryUpdateStepAsync(job, graph, job.Batches[0].Id, job.Batches[0].Steps[2].Id, newRetryByUserId: UserId.Anonymous));
+			Assert.AreEqual(1, job.Batches.Count);
+			Assert.AreEqual(4, job.Batches[0].Steps.Count);
+
+			step = job.Batches[0].Steps[0];
+			Assert.AreEqual(0, step.NodeIdx);
+			Assert.AreEqual(JobStepState.Completed, step.State);
+			Assert.AreEqual(JobStepOutcome.Failure, step.Outcome);
+
+			step = job.Batches[0].Steps[1];
+			Assert.AreEqual(0, step.NodeIdx);
+			Assert.AreEqual(JobStepState.Ready, step.State);
+			Assert.AreEqual(JobStepOutcome.Success, step.Outcome);
+
+			step = job.Batches[0].Steps[2];
+			Assert.AreEqual(1, step.NodeIdx);
+			Assert.AreEqual(JobStepState.Waiting, step.State);
+			Assert.AreEqual(JobStepOutcome.Success, step.Outcome);
+
+			step = job.Batches[0].Steps[3];
+			Assert.AreEqual(2, step.NodeIdx);
+			Assert.AreEqual(JobStepState.Waiting, step.State);
+			Assert.AreEqual(JobStepOutcome.Success, step.Outcome);
+		}
+
 		[TestMethod]
 		public async Task UnknownShelfAsync()
 		{
diff --git a/Engine/Source/Programs/Horde/Horde.Server/Jobs/JobCollection.cs b/Engine/Source/Programs/Horde/Horde.Server/Jobs/JobCollection.cs
--- a/Engine/Source/Programs/Horde/Horde.Server/Jobs/JobCollection.cs
+++ b/Engine/Source/Programs/Horde/Horde.Server/Jobs/JobCollection.cs
@@ -783,6 +783,48 @@ namespace Horde.Server.Jobs
 			return await TryUpdateAsync(jobDocument, updates);
 		}
 
+		/// <summary>
+		/// Marks a node for retry and - if it's a skipped node - all of the nodes that caused it to be skipped
+		/// </summary>
+		static void RetryNodesRecursive(JobDocument jobDocument, IGraph graph, int groupIdx, int nodeIdx, UserId retryByUserId, List<UpdateDefinition<JobDocument>> updates)
+		{
+			HashSet<NodeRef> retryNodes = new HashSet<NodeRef>();
+			retryNodes.Add(new NodeRef(groupIdx, nodeIdx));
+
+			UpdateDefinitionBuilder<JobDocument> updateBuilder = Builders<JobDocument>.Update;
+			for (int batchIdx = jobDocument.Batches.Count - 1; batchIdx >= 0; batchIdx--)
+			{
+				JobStepBatchDocument batch = jobDocument.Batches[batchIdx];
+				for (int stepIdx = batch.Steps.Count - 1; stepIdx >= 0; stepIdx--)
+				{
+					JobStepDocument step = batch.Steps[stepIdx];
+
+					NodeRef nodeRef = new NodeRef(batch.GroupIdx, step.NodeIdx);
+					if (retryNodes.Remove(nodeRef))
+					{
+						if (step.State == JobStepState.Skipped)
+						{
+							// Add all the dependencies to be retried
+							NodeRef[] dependencies = graph.GetNode(nodeRef).InputDependencies;
+							retryNodes.UnionWith(dependencies);
+						}
+						else
+						{
+							// Retry this step
+							int lambdaBatchIdx = batchIdx;
+							int lambdaStepIdx = stepIdx;
+
+							step.Retry = true;
+							updates.Add(updateBuilder.Set(x => x.Batches[lambdaBatchIdx].Steps[lambdaStepIdx].Retry, true));
+
+							step.RetriedByUserId = retryByUserId;
+							updates.Add(updateBuilder.Set(x => x.Batches[lambdaBatchIdx].Steps[lambdaStepIdx].RetriedByUserId, step.RetriedByUserId));
+						}
+					}
+				}
+			}
+		}
+
 		/// <inheritdoc/>
 		public Task<IJob?> TryUpdateStepAsync(IJob job, IGraph graph, SubResourceId batchId, SubResourceId stepId, JobStepState newState, JobStepOutcome newOutcome, JobStepError? newError, bool? newAbortRequested, UserId? newAbortByUserId, LogId? newLogId, ObjectId? newNotificationTriggerId, UserId? newRetryByUserId, Priority? newPriority, List<Report>? newReports, Dictionary<string, string?>? newProperties)
 		{
@@ -887,12 +929,7 @@ namespace Horde.Server.Jobs
 							// Update the retry flag
 							if (newRetryByUserId != null && step.RetriedByUserId == null)
 							{
-								step.Retry = true;
-								updates.Add(updateBuilder.Set(x => x.Batches[batchIdx].Steps[stepIdx].Retry, true));
-
-								step.RetriedByUserId = newRetryByUserId;
-								updates.Add(updateBuilder.Set(x => x.Batches[batchIdx].Steps[stepIdx].RetriedByUserId, step.RetriedByUserId));
-
+								RetryNodesRecursive(jobDocument, graph, batch.GroupIdx, step.NodeIdx, newRetryByUserId.Value, updates);
 								refreshBatches = true;
 							}
 
