Horde: Return the correct id of a new retried node in the response.

diff --git a/Engine/Source/Programs/Horde/Horde.Server/Jobs/IJob.cs b/Engine/Source/Programs/Horde/Horde.Server/Jobs/IJob.cs
--- a/Engine/Source/Programs/Horde/Horde.Server/Jobs/IJob.cs
+++ b/Engine/Source/Programs/Horde/Horde.Server/Jobs/IJob.cs
@@ -1020,6 +1020,32 @@ namespace Horde.Server.Jobs
 			return stepForNode;
 		}
 
+		/// <summary>
+		/// Find the latest step executing the given node
+		/// </summary>
+		/// <param name="job">The job being run</param>
+		/// <param name="nodeRef">Node to find</param>
+		/// <returns>The retried step information</returns>
+		public static JobStepRefId? FindLatestStepForNode(this IJob job, NodeRef nodeRef)
+		{
+			for (int batchIdx = job.Batches.Count - 1; batchIdx >= 0; batchIdx--)
+			{
+				IJobStepBatch batch = job.Batches[batchIdx];
+				if (batch.GroupIdx == nodeRef.GroupIdx)
+				{
+					for (int stepIdx = batch.Steps.Count - 1; stepIdx >= 0; stepIdx--)
+					{
+						IJobStep step = batch.Steps[stepIdx];
+						if (step.NodeIdx == nodeRef.NodeIdx)
+						{
+							return new JobStepRefId(job.Id, batch.Id, step.Id);
+						}
+					}
+				}
+			}
+			return null;
+		}
+
 		/// <summary>
 		/// Gets the estimated timing info for all nodes in the job
 		/// </summary>
diff --git a/Engine/Source/Programs/Horde/Horde.Server/Jobs/JobsController.cs b/Engine/Source/Programs/Horde/Horde.Server/Jobs/JobsController.cs
--- a/Engine/Source/Programs/Horde/Horde.Server/Jobs/JobsController.cs
+++ b/Engine/Source/Programs/Horde/Horde.Server/Jobs/JobsController.cs
@@ -1239,6 +1239,12 @@ namespace Horde.Server.Jobs
 
 			try
 			{
+				NodeRef? retryNodeRef = null;
+				if (request.Retry != null && job.TryGetStep(stepId, out IJobStepBatch? batch, out IJobStep? step))
+				{
+					retryNodeRef = new NodeRef(batch.GroupIdx, step.NodeIdx);
+				}
+
 				IJob? newJob = await _jobService.UpdateStepAsync(job, batchId, stepId, streamConfig, request.State, request.Outcome, null, request.AbortRequested, abortByUser, (request.LogId == null)? null : LogId.Parse(request.LogId), null, retryByUser, request.Priority, null, request.Properties);
 				if (newJob == null)
 				{
@@ -1246,9 +1252,9 @@ namespace Horde.Server.Jobs
 				}
 
 				UpdateStepResponse response = new UpdateStepResponse();
-				if (request.Retry ?? false)
+				if (retryNodeRef != null)
 				{
-					JobStepRefId? retriedStepId = FindRetriedStep(job, batchId, stepId);
+					JobStepRefId? retriedStepId = newJob.FindLatestStepForNode(retryNodeRef);
 					if (retriedStepId != null)
 					{
 						response.BatchId = retriedStepId.Value.BatchId.ToString();
@@ -1263,36 +1269,6 @@ namespace Horde.Server.Jobs
 			}
 		}
 
-		/// <summary>
-		/// Find the first retried step after the given step
-		/// </summary>
-		/// <param name="job">The job being run</param>
-		/// <param name="batchId">Batch id of the last step instance</param>
-		/// <param name="stepId">Step id of the last instance</param>
-		/// <returns>The retried step information</returns>
-		static JobStepRefId? FindRetriedStep(IJob job, SubResourceId batchId, SubResourceId stepId)
-		{
-			NodeRef? lastNodeRef = null;
-			foreach (IJobStepBatch batch in job.Batches)
-			{
-				if ((lastNodeRef == null && batch.Id == batchId) || (lastNodeRef != null && batch.GroupIdx == lastNodeRef.GroupIdx))
-				{
-					foreach (IJobStep step in batch.Steps)
-					{
-						if (lastNodeRef == null && step.Id == stepId)
-						{
-							lastNodeRef = new NodeRef(batch.GroupIdx, step.NodeIdx);
-						}
-						else if (lastNodeRef != null && step.NodeIdx == lastNodeRef.NodeIdx)
-						{
-							return new JobStepRefId(job.Id, batch.Id, step.Id);
-						}
-					}
-				}
-			}
-			return null;
-		}
-
 		/// <summary>
 		/// Gets a particular step currently scheduled to be executed for a job
 		/// </summary>
