Horde: Make step and batch ids stable across calls to CreateOrUpdateBatches().

diff --git a/Engine/Source/Programs/Horde/Horde.Server/Jobs/JobCollection.cs b/Engine/Source/Programs/Horde/Horde.Server/Jobs/JobCollection.cs
--- a/Engine/Source/Programs/Horde/Horde.Server/Jobs/JobCollection.cs
+++ b/Engine/Source/Programs/Horde/Horde.Server/Jobs/JobCollection.cs
@@ -1329,6 +1329,32 @@ namespace Horde.Server.Jobs
 			RefreshJobPriority(job, updates);
 		}
 
+		static void RemoveSteps(JobStepBatchDocument batch, Predicate<JobStepDocument> predicate, Dictionary<NodeRef, SubResourceId> recycleStepIds)
+		{
+			for (int idx = batch.Steps.Count - 1; idx >= 0; idx--)
+			{
+				JobStepDocument step = batch.Steps[idx];
+				if (predicate(step))
+				{
+					recycleStepIds[new NodeRef(batch.GroupIdx, step.NodeIdx)] = step.Id;
+					batch.Steps.RemoveAt(idx);
+				}
+			}
+		}
+
+		static void RemoveBatches(JobDocument job, Predicate<JobStepBatchDocument> predicate, Dictionary<int, SubResourceId> recycleBatchIds)
+		{
+			for (int idx = job.Batches.Count - 1; idx >= 0; idx--)
+			{
+				JobStepBatchDocument batch = job.Batches[idx];
+				if (predicate(batch))
+				{
+					recycleBatchIds[batch.GroupIdx] = batch.Id;
+					job.Batches.RemoveAt(idx);
+				}
+			}
+		}
+
 		/// <summary>
 		/// Update the jobsteps for the given node graph 
 		/// </summary>
@@ -1358,10 +1384,11 @@ namespace Horde.Server.Jobs
 				}
 			}
 
-			// Remove any steps and batches that haven't started yet
+			// Remove any steps and batches that haven't started yet, saving their ids so we can re-use them if we re-add them
+			Dictionary<NodeRef, SubResourceId> recycleStepIds = new Dictionary<NodeRef, SubResourceId>();
 			foreach (JobStepBatchDocument batch in job.Batches)
 			{
-				batch.Steps.RemoveAll(x => x.State == JobStepState.Waiting || x.State == JobStepState.Ready);
+				RemoveSteps(batch, x => x.State == JobStepState.Waiting || x.State == JobStepState.Ready, recycleStepIds);
 			}
 
 			// Mark any steps in failed batches as skipped
@@ -1415,11 +1442,12 @@ namespace Horde.Server.Jobs
 						failedNodes.Remove(node);
 					}
 				}
-				batch.Steps.RemoveAll(x => x.State == JobStepState.Skipped && !failedNodes.Contains(group.Nodes[x.NodeIdx]));
+				RemoveSteps(batch, x => x.State == JobStepState.Skipped && !failedNodes.Contains(group.Nodes[x.NodeIdx]), recycleStepIds);
 			}
 
 			// Remove any batches which are now empty
-			job.Batches.RemoveAll(x => x.Steps.Count == 0 && x.LeaseId == null && x.Error == JobStepBatchError.None);
+			Dictionary<int, SubResourceId> recycleBatchIds = new Dictionary<int, SubResourceId>();
+			RemoveBatches(job, x => x.Steps.Count == 0 && x.LeaseId == null && x.Error == JobStepBatchError.None, recycleBatchIds);
 
 			// Find all the targets in this job
 			HashSet<string> targets = new HashSet<string>(StringComparer.OrdinalIgnoreCase);
@@ -1577,17 +1605,27 @@ namespace Horde.Server.Jobs
 						JobStepBatchDocument? batch = appendToBatches[groupIdx];
 						if (batch == null)
 						{
-							job.NextSubResourceId = job.NextSubResourceId.Next();
+							SubResourceId batchId;
+							if (!recycleBatchIds.Remove(groupIdx, out batchId))
+							{
+								job.NextSubResourceId = job.NextSubResourceId.Next();
+								batchId = job.NextSubResourceId;
+							}
 
-							batch = new JobStepBatchDocument(job.NextSubResourceId, groupIdx);
+							batch = new JobStepBatchDocument(batchId, groupIdx);
 							job.Batches.Add(batch);
 
 							appendToBatches[groupIdx] = batch;
 						}
 
-						job.NextSubResourceId = job.NextSubResourceId.Next();
+						SubResourceId stepId;
+						if (!recycleStepIds.Remove(new NodeRef(groupIdx, nodeIdx), out stepId))
+						{
+							job.NextSubResourceId = job.NextSubResourceId.Next();
+							stepId = job.NextSubResourceId;
+						}
 
-						JobStepDocument step = new JobStepDocument(job.NextSubResourceId, nodeIdx);
+						JobStepDocument step = new JobStepDocument(stepId, nodeIdx);
 						batch.Steps.Add(step);
 					}
 				}
