Horde: Create lease documents for compute jobs, allowing them to be cancelled without 404'ing.

diff --git a/Engine/Source/Programs/Horde/Horde.Server/Agents/AgentService.cs b/Engine/Source/Programs/Horde/Horde.Server/Agents/AgentService.cs
--- a/Engine/Source/Programs/Horde/Horde.Server/Agents/AgentService.cs
+++ b/Engine/Source/Programs/Horde/Horde.Server/Agents/AgentService.cs
@@ -795,7 +795,7 @@ namespace Horde.Server.Agents
 		/// <param name="agent">Agent that will be executing the lease</param>
 		/// <param name="agentLease">The new agent lease</param>
 		/// <returns>New lease document</returns>
-		private Task<ILease> CreateLeaseAsync(IAgent agent, AgentLease agentLease)
+		internal Task<ILease> CreateLeaseAsync(IAgent agent, AgentLease agentLease)
 		{
 			try
 			{
diff --git a/Engine/Source/Programs/Horde/Horde.Server/Compute/ComputeService.cs b/Engine/Source/Programs/Horde/Horde.Server/Compute/ComputeService.cs
--- a/Engine/Source/Programs/Horde/Horde.Server/Compute/ComputeService.cs
+++ b/Engine/Source/Programs/Horde/Horde.Server/Compute/ComputeService.cs
@@ -24,14 +24,15 @@ namespace Horde.Server.Compute
 	public class ComputeService
 	{
 		readonly IAgentCollection _agentCollection;
+		readonly AgentService _agentService;
 
 		/// <summary>
 		/// Constructor
 		/// </summary>
-		/// <param name="agentCollection"></param>
-		public ComputeService(IAgentCollection agentCollection)
+		public ComputeService(IAgentCollection agentCollection, AgentService agentService)
 		{
 			_agentCollection = agentCollection;
+			_agentService = agentService;
 		}
 
 		/// <summary>
@@ -58,6 +59,7 @@ namespace Horde.Server.Compute
 						if (newAgent != null)
 						{
 							await _agentCollection.PublishUpdateEventAsync(agent.Id);
+							await _agentService.CreateLeaseAsync(newAgent, lease);
 							return resource;
 						}
 					}
