Horde: Fix meta data collection index
#rnx

diff --git a/Engine/Source/Programs/Horde/Horde.Server/Jobs/TestData/TestDataCollection.cs b/Engine/Source/Programs/Horde/Horde.Server/Jobs/TestData/TestDataCollection.cs
--- a/Engine/Source/Programs/Horde/Horde.Server/Jobs/TestData/TestDataCollection.cs
+++ b/Engine/Source/Programs/Horde/Horde.Server/Jobs/TestData/TestDataCollection.cs
@@ -416,7 +416,7 @@ namespace Horde.Server.Jobs.TestData
 			_testDataDocuments = mongoService.GetCollection<TestDataDocument>("TestData", indexes);
 
 			List<MongoIndex<TestMetaDocument>> metaIndexes = new List<MongoIndex<TestMetaDocument>>();
-			metaIndexes.Add(keys => keys.Ascending(x => x.ProjectName).Ascending(x => x.Platforms).Ascending(x => x.Configurations).Ascending(x => x.BuildTargets).Ascending(x => x.RHI), unique: true);
+			metaIndexes.Add(keys => keys.Ascending(x => x.ProjectName).Ascending(x => x.Platforms).Ascending(x => x.Configurations).Ascending(x => x.BuildTargets).Ascending(x => x.RHI));
 			_testMeta = mongoService.GetCollection<TestMetaDocument>("TestData.MetaV2", metaIndexes);
 
 			List<MongoIndex<TestDocument>> testIndexes = new List<MongoIndex<TestDocument>>();
