Horde: Add step retries to job operations bar
#rnx
#skipci

diff --git a/Engine/Source/Programs/Horde/HordeDashboard/src/components/jobDetailsV2/JobDetailStepHistory.tsx b/Engine/Source/Programs/Horde/HordeDashboard/src/components/jobDetailsV2/JobDetailStepHistory.tsx
--- a/Engine/Source/Programs/Horde/HordeDashboard/src/components/jobDetailsV2/JobDetailStepHistory.tsx
+++ b/Engine/Source/Programs/Horde/HordeDashboard/src/components/jobDetailsV2/JobDetailStepHistory.tsx
@@ -41,15 +41,25 @@ class StepHistoryDataView extends JobDataView {
          return;
       }
       
-      const jobData = details.jobData!;
-      const stepName = details.getStepName(this.stepId);
+      const stepName = details.getStepName(this.stepId, false);
 
+      if (stepName) {
+         this.loadHistory(stepName);
+      }
+
+   }
+
+   loadHistory(stepName: string) {
+      const jobData = this.details?.jobData;
+      if (!jobData) {
+         return;
+      }
       backend.getJobStepHistory(jobData.streamId, stepName, 1024, jobData.templateId!).then(response => {
          this.history = response;
          this.updateReady();
       }).finally(() => {
          this.initialize(this.history?.length ? [sideRail] : undefined);
-      })
+      })   
    }
 
    clear() {
@@ -59,7 +69,12 @@ class StepHistoryDataView extends JobDataView {
    }
 
    detailsUpdated() {
-
+      if (!this.initialized) {
+         const stepName = this.details?.getStepName(this.stepId, false);
+         if (stepName) {
+            this.loadHistory(stepName);
+         }
+      }       
    }
 
    history: GetJobStepRefResponse[] = [];
diff --git a/Engine/Source/Programs/Horde/HordeDashboard/src/components/jobDetailsV2/JobOperationsBar.tsx b/Engine/Source/Programs/Horde/HordeDashboard/src/components/jobDetailsV2/JobOperationsBar.tsx
--- a/Engine/Source/Programs/Horde/HordeDashboard/src/components/jobDetailsV2/JobOperationsBar.tsx
+++ b/Engine/Source/Programs/Horde/HordeDashboard/src/components/jobDetailsV2/JobOperationsBar.tsx
@@ -4,7 +4,7 @@ import { CommandBar, CommandBarButton, ICommandBarItemProps, IContextualMenuItem
 import { observer } from 'mobx-react-lite';
 import React, { useState } from 'react';
 import { Link, useNavigate } from 'react-router-dom';
-import { ArtifactContextType, JobState, StepData } from '../../backend/Api';
+import { ArtifactContextType, JobState, JobStepOutcome, JobStepState, StepData } from '../../backend/Api';
 import dashboard from '../../backend/Dashboard';
 import { hordeClasses } from '../../styles/Styles';
 import { EditJobModal } from '../EditJobModal';
@@ -14,8 +14,9 @@ import { NotificationDropdown } from '../NotificationDropdown';
 import { PauseStepModal } from '../StepPauseModal';
 import { AbortJobModal } from './AbortJobModal';
 import { JobDetailsV2 } from './JobDetailsViewCommon';
-import { StepRetryModal, StepRetryType } from './StepRetryModal';
+import { RetryStepsModal, StepRetryModal, StepRetryType } from './StepRetryModal';
 import { JobArtifactsModal } from '../artifacts/ArtifactsModal';
+import { getSiteConfig } from '../../backend/Config';
 
 enum ParameterState {
    Hidden,
@@ -29,6 +30,7 @@ export const JobOperations: React.FC<{ jobDetails: JobDetailsV2 }> = observer(({
    const navigate = useNavigate();
    const [abortShown, setAbortShown] = useState(false);
    const [editShown, setEditShown] = useState(false);
+   const [retryStepsShown, setRetryStepsShown] = useState(false);
 
    const [parametersState, setParametersState] = useState(query.get("newbuild") ? ParameterState.Clone : ParameterState.Hidden);
 
@@ -49,6 +51,23 @@ export const JobOperations: React.FC<{ jobDetails: JobDetailsV2 }> = observer(({
    const abortDisabled = jobData.state === JobState.Complete;
    const runAgainDisabled = false; /*jobDetails.jobdata?.state !== JobState.Complete*/
 
+   const failedSteps = jobDetails.getSteps().filter(s => {
+
+      if (!s.finishTime || (s.state !== JobStepState.Aborted && s.outcome !== JobStepOutcome.Failure)) {
+         return false;
+      }
+
+      const retries = jobDetails.getStepRetries(s.id);
+      const retryNumber = jobDetails.getStepRetryNumber(s.id);
+      if (retries.length && retryNumber < (retries.length - 1) ) {
+         return false;
+      }
+
+      return true;
+      
+   });
+   const retryFailedStepsDisabled = !failedSteps.length;
+
    const pinned = dashboard.jobPinned(jobId);
 
    const opsList: IContextualMenuItem[] = [];
@@ -74,6 +93,17 @@ export const JobOperations: React.FC<{ jobDetails: JobDetailsV2 }> = observer(({
       onClick: () => { setEditShown(true); }
    });
 
+   if (getSiteConfig().environment !== "production") {
+
+      opsList.push({
+         key: 'jobops_runfailedsteps',
+         text: "Retry Steps",
+         disabled: retryFailedStepsDisabled,
+         iconProps: { iconName: "Repeat" },
+         onClick: () => { setRetryStepsShown(true); }
+      });
+   }
+
    opsList.push({
       key: 'jobops_abort',
       text: "Cancel Job",
@@ -119,6 +149,7 @@ export const JobOperations: React.FC<{ jobDetails: JobDetailsV2 }> = observer(({
    }
 
    return <Stack>
+      {retryStepsShown && <RetryStepsModal stepIds={failedSteps.map(s => s.id)} jobDetails={jobDetails} onClose={() => { setRetryStepsShown(false); }} />}
       <AbortJobModal jobDetails={jobDetails} show={abortShown} onClose={() => { setAbortShown(false); }} />
       <EditJobModal jobData={jobDetails.jobData} show={editShown} onClose={() => { setEditShown(false); }} />
       <NewBuild streamId={jobDetails.stream!.id} jobDetails={jobDetails} readOnly={parametersState === ParameterState.Parameters}
@@ -210,7 +241,7 @@ const StepArtifactsOperations: React.FC<{ jobDetails: JobDetailsV2, stepId: stri
       text: "Step Artifacts",
       iconProps: { iconName: "Folder" },
       disabled: !atypes.get("step-saved"),
-      onClick: () => { setArtifactsShown("step-saved")}
+      onClick: () => { setArtifactsShown("step-saved") }
    });
 
    opsList.push({
@@ -218,7 +249,7 @@ const StepArtifactsOperations: React.FC<{ jobDetails: JobDetailsV2, stepId: stri
       text: "Output Artifacts",
       iconProps: { iconName: "MenuOpen" },
       disabled: !atypes.get("step-output"),
-      onClick: () => { setArtifactsShown("step-output")}
+      onClick: () => { setArtifactsShown("step-output") }
    });
 
    opsList.push({
@@ -226,7 +257,7 @@ const StepArtifactsOperations: React.FC<{ jobDetails: JobDetailsV2, stepId: stri
       text: "Trace Artifacts",
       iconProps: { iconName: "SearchTemplate" },
       disabled: !atypes.get("step-trace"),
-      onClick: () => { setArtifactsShown("step-trace")}      
+      onClick: () => { setArtifactsShown("step-trace") }
    });
 
    const opsItems: ICommandBarItemProps[] = [
diff --git a/Engine/Source/Programs/Horde/HordeDashboard/src/components/jobDetailsV2/StepRetryModal.tsx b/Engine/Source/Programs/Horde/HordeDashboard/src/components/jobDetailsV2/StepRetryModal.tsx
--- a/Engine/Source/Programs/Horde/HordeDashboard/src/components/jobDetailsV2/StepRetryModal.tsx
+++ b/Engine/Source/Programs/Horde/HordeDashboard/src/components/jobDetailsV2/StepRetryModal.tsx
@@ -1,9 +1,9 @@
 // Copyright Epic Games, Inc. All Rights Reserved.
-import { DefaultButton, IconButton, Modal, PrimaryButton, Stack, Text, TextField } from '@fluentui/react';
-import React from 'react';
+import { DefaultButton, IconButton, Modal, PrimaryButton, Spinner, SpinnerSize, Stack, Text, TextField } from '@fluentui/react';
+import React, { useState } from 'react';
 import backend from '../../backend';
 import { CreateJobRequest } from '../../backend/Api';
-import { useNavigate } from 'react-router-dom';
+import { Link, useNavigate } from 'react-router-dom';
 import { hordeClasses } from '../../styles/Styles';
 import { JobDetailsV2 } from './JobDetailsViewCommon';
 
@@ -12,6 +12,157 @@ export enum StepRetryType {
    TestFix
 }
 
+type RunAgainResult = {
+   stepId: string,
+   newStepId?: string,
+   error?: string
+}
+
+export const RetryStepsModal: React.FC<{ stepIds: string[]; jobDetails: JobDetailsV2; onClose: () => void }> = ({ stepIds, jobDetails, onClose }) => {
+
+   const [submitting, setSubmitting] = useState(false);
+   const [submitResults, setSubmitResults] = useState<RunAgainResult[] | undefined>(undefined);
+
+   if (submitting && !submitResults) {
+      return <Modal className={hordeClasses.modal} isOpen={true} styles={{ main: { padding: 8, width: 800 } }} onDismiss={() => { onClose() }}>
+         <Stack horizontal styles={{ root: { padding: 8, paddingBottom: 32 } }}>
+            <Text variant="mediumPlus">Updating Steps</Text>
+         </Stack>
+         <Stack horizontalAlign="center">
+            <Spinner size={SpinnerSize.large} />
+         </Stack>
+      </Modal>
+   }
+
+   if (submitResults) {
+
+      const resultElements = submitResults.map(result => {
+         const stepName = jobDetails.getStepName(result.stepId, true);
+         let resultText = `${stepName} - `;
+         if (result.newStepId) {
+            const url = `/job/${jobDetails.jobId!}?step=${result.newStepId}`
+            resultText += "Retried"
+            return <Stack horizontal>
+               <Link to={url}>{resultText}</Link>
+            </Stack>
+         } else if (result.error) {
+            resultText += `Error: ${result.error}`;
+         } else {
+            resultText += "Step could not be retried"
+         }
+         return <Stack><Text>{resultText}</Text></Stack>
+      })
+
+
+      return <Modal className={hordeClasses.modal} isOpen={true} styles={{ main: { padding: 8, width: 800 } }} onDismiss={() => { onClose() }}>
+         <Stack horizontal styles={{ root: { padding: 8 } }}>
+            <Stack.Item grow={2}>
+               <Text variant="mediumPlus">Step Retry Results</Text>
+            </Stack.Item>
+            <Stack.Item grow={0}>
+               <IconButton
+                  iconProps={{ iconName: 'Cancel' }}
+                  ariaLabel="Close popup modal"
+                  onClick={() => { onClose(); }}
+               />
+            </Stack.Item>
+         </Stack>
+         <Stack styles={{ root: { paddingLeft: 32 } }} tokens={{ childrenGap: 8 }}>
+            {resultElements}
+         </Stack>
+      </Modal>
+   }
+
+
+   const jobData = jobDetails.jobData;
+   if (!jobData) {
+      console.error("Attempting to display StepRetryModal without job data");
+      return null;
+   }
+
+   const headerText = stepIds.length > 1 ? "Retry Steps?" : "Retry Step?";
+
+   const onRetry = async () => {
+
+      const results: RunAgainResult[] = [];
+
+      for (let i = 0; i < stepIds.length; i++) {
+
+         const stepId = stepIds[i];
+
+         const batch = jobDetails.batchByStepId(stepId);
+         if (!batch) {
+            results.push({
+               stepId: stepId,
+               error: `Unable to get batch for step: ${stepId}`
+            })
+            continue;
+         }
+
+         try {
+            const response = await backend.updateJobStep(jobData.id, batch.id, stepId, { retry: true })
+            
+            if (!response.stepId) {
+               results.push({
+                  stepId: stepId
+               })
+            } else {
+               results.push({
+                  stepId: stepId,
+                  newStepId: response.stepId
+               })
+            }
+
+         } catch (error) {
+
+            results.push({ stepId: stepId, error: error?.toString() });
+
+         }
+      }
+
+      setSubmitResults(results);
+
+   }
+
+   const stepNames: string[] = [];
+   stepIds.forEach(stepId => {
+      const stepName = jobDetails.getStepName(stepId, true);
+      stepNames.push(stepName ? stepName : "Unknown Step");
+   });
+
+   const stepElements = stepNames.map(name => {
+      return <Stack><Text>{name}</Text></Stack>
+   })
+
+   return <Modal className={hordeClasses.modal} isOpen={true} styles={{ main: { padding: 8, width: 800 } }} onDismiss={() => { onClose() }}>
+      <Stack horizontal styles={{ root: { padding: 8 } }}>
+         <Stack.Item grow={2}>
+            <Text variant="mediumPlus">{headerText}</Text>
+         </Stack.Item>
+         <Stack.Item grow={0}>
+            <IconButton
+               iconProps={{ iconName: 'Cancel' }}
+               ariaLabel="Close popup modal"
+               onClick={() => { onClose(); }}
+            />
+         </Stack.Item>
+      </Stack>
+
+      <Stack styles={{ root: { paddingLeft: 32 } }} tokens={{ childrenGap: 8 }}>
+         {stepElements}
+      </Stack>
+
+      <Stack styles={{ root: { padding: 8 } }}>
+         <Stack horizontal tokens={{ childrenGap: 16 }} styles={{ root: { paddingTop: 12, paddingLeft: 8, paddingBottom: 8 } }}>
+            <Stack grow />
+            <PrimaryButton text="Retry" disabled={submitting} onClick={() => { onRetry(); setSubmitting(true) }} />
+            <DefaultButton text="Cancel" disabled={submitting} onClick={() => { onClose(); }} />
+         </Stack>
+      </Stack>
+   </Modal>;
+
+}
+
 export const StepRetryModal: React.FC<{ stepId: string; jobDetails: JobDetailsV2; show: boolean; type: StepRetryType; onClose: () => void }> = ({ stepId, jobDetails, show, type, onClose }) => {
 
    const navigate = useNavigate();
diff --git a/Engine/Source/Programs/Horde/HordeDashboard/src/styles/Styles.tsx b/Engine/Source/Programs/Horde/HordeDashboard/src/styles/Styles.tsx
--- a/Engine/Source/Programs/Horde/HordeDashboard/src/styles/Styles.tsx
+++ b/Engine/Source/Programs/Horde/HordeDashboard/src/styles/Styles.tsx
@@ -131,7 +131,8 @@ registerIcons({
       "Eye": IconContents.EYE_OPEN,
       "ArrowLeft": IconContents.ARROW_LEFT,
       "ArrowRight": IconContents.ARROW_RIGHT,
-      "ArrowUp" : IconContents.ARROW_UP
+      "ArrowUp": IconContents.ARROW_UP,
+      "Repeat": IconContents.REPEAT
    }
 });
 
