Horde: Disable agent-side expiry of temp storage artifacts. This should be handled through the artifact system instead.

diff --git a/Engine/Source/Programs/Horde/Horde.Agent/Execution/JobExecutor.cs b/Engine/Source/Programs/Horde/Horde.Agent/Execution/JobExecutor.cs
--- a/Engine/Source/Programs/Horde/Horde.Agent/Execution/JobExecutor.cs
+++ b/Engine/Source/Programs/Horde/Horde.Agent/Execution/JobExecutor.cs
@@ -992,7 +992,7 @@ namespace Horde.Agent.Execution
 				}
 
 				// Write the final node
-				await treeWriter.WriteAsync(refName, outputNode, new RefOptions { Lifetime = TimeSpan.FromDays(3.0) }, cancellationToken);
+				await treeWriter.WriteAsync(refName, outputNode, new RefOptions(), cancellationToken);
 				logger.LogInformation("Upload took {Time:n1}s", timer.Elapsed.TotalSeconds);
 
 				// Create the artifact
