Horde: Do not log every unknown event to ClickHouse telemetry sink

diff --git a/Engine/Source/Programs/Horde/Horde.Server/Telemetry/ClickHouseTelemetrySink.cs b/Engine/Source/Programs/Horde/Horde.Server/Telemetry/ClickHouseTelemetrySink.cs
--- a/Engine/Source/Programs/Horde/Horde.Server/Telemetry/ClickHouseTelemetrySink.cs
+++ b/Engine/Source/Programs/Horde/Horde.Server/Telemetry/ClickHouseTelemetrySink.cs
@@ -75,7 +75,6 @@ namespace Horde.Server.Telemetry
 					case AgentMetadataEvent agentMetadata: agentMetadataEvents.Add(agentMetadata); break;
 					case AgentCpuMetricsEvent agentCpu: agentCpuEvents.Add(agentCpu); break;
 					case AgentMemoryMetricsEvent agentMem: agentMemEvents.Add(agentMem); break;
-					default: _logger.LogError("Unable to parse event type for {EventName}", tuple.eventName); break;
 				}
 				c++;
 			}
