Horde: Add documentation release notes, enable job rety steps in prod too
#skipci
#rnx

diff --git a/Engine/Source/Programs/Horde/HordeDashboard/src/components/TopNav.tsx b/Engine/Source/Programs/Horde/HordeDashboard/src/components/TopNav.tsx
--- a/Engine/Source/Programs/Horde/HordeDashboard/src/components/TopNav.tsx
+++ b/Engine/Source/Programs/Horde/HordeDashboard/src/components/TopNav.tsx
@@ -462,6 +462,12 @@ export const TopNav: React.FC<{ suppressServer?: boolean }> = observer(({ suppre
          text: "Documentation",
          link: `/docs`
       });
+      
+      docsItems.push({
+         key: "server_docs_releasenotes",
+         text: "Release Notes",
+         link: `/docs/ReleaseNotes.md`
+      });
 
       subItems.push({
          itemType: ContextualMenuItemType.Section,
diff --git a/Engine/Source/Programs/Horde/HordeDashboard/src/components/jobDetailsV2/JobOperationsBar.tsx b/Engine/Source/Programs/Horde/HordeDashboard/src/components/jobDetailsV2/JobOperationsBar.tsx
--- a/Engine/Source/Programs/Horde/HordeDashboard/src/components/jobDetailsV2/JobOperationsBar.tsx
+++ b/Engine/Source/Programs/Horde/HordeDashboard/src/components/jobDetailsV2/JobOperationsBar.tsx
@@ -93,7 +93,7 @@ export const JobOperations: React.FC<{ jobDetails: JobDetailsV2 }> = observer(({
       onClick: () => { setEditShown(true); }
    });
 
-   if (getSiteConfig().environment !== "production") {
+   //if (getSiteConfig().environment !== "production") {
 
       opsList.push({
          key: 'jobops_runfailedsteps',
@@ -102,7 +102,7 @@ export const JobOperations: React.FC<{ jobDetails: JobDetailsV2 }> = observer(({
          iconProps: { iconName: "Repeat" },
          onClick: () => { setRetryStepsShown(true); }
       });
-   }
+   //}
 
    opsList.push({
       key: 'jobops_abort',
