Horde: Update release notes
#skipci
#rnx

diff --git a/Engine/Source/Programs/Horde/HordeDashboard/src/components/docs/DocView.tsx b/Engine/Source/Programs/Horde/HordeDashboard/src/components/docs/DocView.tsx
--- a/Engine/Source/Programs/Horde/HordeDashboard/src/components/docs/DocView.tsx
+++ b/Engine/Source/Programs/Horde/HordeDashboard/src/components/docs/DocView.tsx
@@ -172,6 +172,7 @@ const DocRail = observer(() => {
    refLinks.push({ text: "Deployment", url: "/docs/Deployment.md" });
    refLinks.push({ text: "Configuration", url: "/docs/Config.md" });
    refLinks.push({ text: "Horde Internals", url: "/docs/Internals.md" });
+   refLinks.push({ text: "Release Notes", url: "/docs/ReleaseNotes.md" });
 
    return <SideRail jumpLinks={state.jumpLinks} relatedLinks={refLinks} />
 
