Horde: Remove last vestiges of using claim roles in the dashboard
#rnx
#skipci

diff --git a/Engine/Source/Programs/Horde/HordeDashboard/src/backend/Api.ts b/Engine/Source/Programs/Horde/HordeDashboard/src/backend/Api.ts
--- a/Engine/Source/Programs/Horde/HordeDashboard/src/backend/Api.ts
+++ b/Engine/Source/Programs/Horde/HordeDashboard/src/backend/Api.ts
@@ -1434,7 +1434,7 @@ export type CreateZipRequest = {
 
 /** Describes an artifact */
 export type GetArtifactResponseV2 = {
-	
+
 	id: string;
 
 	type: ArtifactContextType;
@@ -1444,7 +1444,7 @@ export type GetArtifactResponseV2 = {
 }
 
 /** Result of an artifact search */
-export type FindArtifactsResponse =  {	
+export type FindArtifactsResponse = {
 	/** List of artifacts matching the search criteria*/
 	artifacts: GetArtifactResponseV2[];
 }
@@ -3093,6 +3093,18 @@ export type UpdateDashboardSettings = {
 
 }
 
+/** Settings for whether various features should be enabled on the dashboard */
+export type GetDashboardFeaturesResponse = {
+
+	/** Whether the notice editor should be listed in the server menu */
+	showNoticeEditor?: boolean;
+
+	/** Whether controls for modifying pools should be shown */
+	showPoolEditor?: boolean;
+
+	/** Whether the remote desktop button should be shown on the agent modal */
+	showRemoteDesktop?: boolean;
+}
 
 /**  Response describing the current user */
 export type GetUserResponse = {
@@ -3127,6 +3139,11 @@ export type GetUserResponse = {
 	/**  Settings for the dashboard */
 	dashboardSettings?: DashboardSettings;
 
+	/// <summary>
+	/// Settings for whether various dashboard features should be shown for the current user
+	/// </summary>
+	dashboardFeatures?: GetDashboardFeaturesResponse;
+
 	/** List of pinned job ids */
 	pinnedJobIds?: string[];
 
@@ -4453,7 +4470,7 @@ export enum TestOutcome {
 	/// The test was skipped
 	Skipped = "Skipped",
 	/// The test had an unspecified result
-	Unspecified = "Unspecified", 
+	Unspecified = "Unspecified",
 	// Warnings
 	Warning = "Warning"
 }
diff --git a/Engine/Source/Programs/Horde/HordeDashboard/src/backend/Backend.ts b/Engine/Source/Programs/Horde/HordeDashboard/src/backend/Backend.ts
--- a/Engine/Source/Programs/Horde/HordeDashboard/src/backend/Backend.ts
+++ b/Engine/Source/Programs/Horde/HordeDashboard/src/backend/Backend.ts
@@ -948,7 +948,7 @@ export class Backend {
 
         return new Promise<GetUserResponse>((resolve, reject) => {
 
-            this.backend.get("/api/v1/users/current", { suppress404: true }).then((response) => {
+            this.backend.get("/api/v1/user", { suppress404: true }).then((response) => {
 
                 let data = response.data as GetUserResponse;
 
diff --git a/Engine/Source/Programs/Horde/HordeDashboard/src/backend/Dashboard.ts b/Engine/Source/Programs/Horde/HordeDashboard/src/backend/Dashboard.ts
--- a/Engine/Source/Programs/Horde/HordeDashboard/src/backend/Dashboard.ts
+++ b/Engine/Source/Programs/Horde/HordeDashboard/src/backend/Dashboard.ts
@@ -96,18 +96,6 @@ export class Dashboard {
         return this.data.image32;
     }
 
-
-    get hordeAdmin(): boolean {
-
-        return !!this.roles.find(r => r.value === "app-horde-admins");
-    }
-
-    get internalEmployee(): boolean {
-
-        return !!this.roles.find(r => r.value === "Internal-Employees");
-
-    }
-
     get email(): string {
 
         const email = this.claims.find(c => c.type.endsWith("/emailaddress"));
diff --git a/Engine/Source/Programs/Horde/HordeDashboard/src/components/HistoryModal.tsx b/Engine/Source/Programs/Horde/HordeDashboard/src/components/HistoryModal.tsx
--- a/Engine/Source/Programs/Horde/HordeDashboard/src/components/HistoryModal.tsx
+++ b/Engine/Source/Programs/Horde/HordeDashboard/src/components/HistoryModal.tsx
@@ -860,7 +860,7 @@ export const HistoryModal: React.FC<{ agentId: string | undefined, onDismiss: (.
                   </Stack>
                </Stack>
                {!!poolItems.length && <Stack style={{ paddingRight: 24 }}><DefaultButton text="Pools" menuProps={{ shouldFocusOnMount: true, items: poolItems }} /></Stack>}
-               {dashboard.hordeAdmin && <Stack> <Stack horizontal tokens={{ childrenGap: 24 }}>
+               {!!dashboard.user?.dashboardFeatures?.showRemoteDesktop && <Stack> <Stack horizontal tokens={{ childrenGap: 24 }}>
                   <Stack>
                      <DefaultButton text="Actions" menuProps={actionMenuProps} />
                   </Stack>
diff --git a/Engine/Source/Programs/Horde/HordeDashboard/src/components/PoolView.tsx b/Engine/Source/Programs/Horde/HordeDashboard/src/components/PoolView.tsx
--- a/Engine/Source/Programs/Horde/HordeDashboard/src/components/PoolView.tsx
+++ b/Engine/Source/Programs/Horde/HordeDashboard/src/components/PoolView.tsx
@@ -929,7 +929,7 @@ const AutoScalerPanel: React.FC = () => {
       }
    }
 
-   const disabled = !dashboard.hordeAdmin || state.submitting;
+   const disabled = !dashboard.user?.dashboardFeatures?.showPoolEditor || state.submitting;
 
    let selectedColor = state.color ?? "0";
    let defaultColor = "0";
@@ -971,7 +971,7 @@ const AutoScalerPanel: React.FC = () => {
                   </Stack>
                </Stack>
 
-               {dashboard.hordeAdmin && <Stack style={{ paddingTop: 24 }}>
+               {!!dashboard.user?.dashboardFeatures?.showPoolEditor && <Stack style={{ paddingTop: 24 }}>
                   <PrimaryButton disabled={disabled || !state.modified} text="Save Changes" onClick={() => { setState({ ...state, submitting: true, confirmed: false, error: undefined }) }} />
                </Stack>}
             </Stack>
diff --git a/Engine/Source/Programs/Horde/HordeDashboard/src/components/TopNav.tsx b/Engine/Source/Programs/Horde/HordeDashboard/src/components/TopNav.tsx
--- a/Engine/Source/Programs/Horde/HordeDashboard/src/components/TopNav.tsx
+++ b/Engine/Source/Programs/Horde/HordeDashboard/src/components/TopNav.tsx
@@ -6,7 +6,6 @@ import React, { useRef, useState } from 'react';
 import { Link, useNavigate } from 'react-router-dom';
 import backend, { useBackend } from '../backend';
 import { ProjectData } from "../backend/Api";
-import { getSiteConfig } from '../backend/Config';
 import dashboard from '../backend/Dashboard';
 import { ProjectStore } from '../backend/ProjectStore';
 import { modeColors } from '../styles/Styles';
@@ -420,7 +419,7 @@ export const TopNav: React.FC<{ suppressServer?: boolean }> = observer(({ suppre
          link: `/reports/utilization`
       });
 
-      if (dashboard.hordeAdmin) {
+      if (dashboard.user?.dashboardFeatures?.showNoticeEditor) {
          monitoringItems.push({
             key: "admin_notices",
             text: "Notices",
@@ -462,7 +461,7 @@ export const TopNav: React.FC<{ suppressServer?: boolean }> = observer(({ suppre
          text: "Documentation",
          link: `/docs`
       });
-      
+
       docsItems.push({
          key: "server_docs_releasenotes",
          text: "Release Notes",
