Horde: Include the default config files with the server.

diff --git a/Engine/Source/Programs/Horde/BuildHorde.xml b/Engine/Source/Programs/Horde/BuildHorde.xml
--- a/Engine/Source/Programs/Horde/BuildHorde.xml
+++ b/Engine/Source/Programs/Horde/BuildHorde.xml
@@ -88,9 +88,6 @@
 			<DotNet Arguments="publish &quot;$(RootDir)/Engine/Source/Programs/Horde/Horde.Server/Horde.Server.csproj&quot; --output &quot;$(StagingDir)/Server&quot; --runtime win-x64 $(VersionArguments)"/>
 			<DotNet Arguments="publish &quot;$(RootDir)/Engine/Source/Programs/Horde/Horde.Agent/Horde.Agent.csproj&quot; --output &quot;$(StagingDir)/Agent&quot; $(VersionArguments)"/>
 
-			<!-- Do not include development defaults which created a ue5 project and stream -->
-			<Delete Files="$(StagingDir)/Server/Defaults/..."/>
-
 			<!-- Create build settings -->
 			<WriteTextFile File="$(StagingDir)/Server/appsettings.Build.json" Text="{ &quot;Horde&quot;: { &quot;SingleInstance&quot;: true } }"/>
 			
