Horde: Remove the project from the default config file. We may not want to launch with CI support.

diff --git a/Engine/Source/Programs/Horde/Horde.Server/Defaults/globals.json b/Engine/Source/Programs/Horde/Horde.Server/Defaults/globals.json
--- a/Engine/Source/Programs/Horde/Horde.Server/Defaults/globals.json
+++ b/Engine/Source/Programs/Horde/Horde.Server/Defaults/globals.json
@@ -1,10 +1,4 @@
 {
-	"projects": [
-		{
-			"id": "ue5",
-			"path": "ue5.project.json"
-		}
-	],
 	"compute": [
 		{
 			"id": "default",
@@ -76,5 +70,25 @@
 			"name": "Linux",
 			"condition": "Platform == 'Linux'"
 		}
+	],
+	"tools": [
+		{
+			"id": "horde-agent",
+			"name": "Horde Agent",
+			"description": "Horde agent software",
+			"acl": {
+				"entries": [
+					{
+						"claim": {
+							"type": "http://epicgames.com/ue/horde/role",
+							"value": "agent"
+						},
+						"actions": [
+							"DownloadTool"
+						]
+					}
+				]
+			}
+		}
 	]
 }
