Horde: Sort job batches by group index to keep step retries in sync, markdown tweaks
#rnx

diff --git a/Engine/Source/Programs/Horde/HordeDashboard/src/base/components/Markdown/Markdown.tsx b/Engine/Source/Programs/Horde/HordeDashboard/src/base/components/Markdown/Markdown.tsx
--- a/Engine/Source/Programs/Horde/HordeDashboard/src/base/components/Markdown/Markdown.tsx
+++ b/Engine/Source/Programs/Horde/HordeDashboard/src/base/components/Markdown/Markdown.tsx
@@ -11,6 +11,8 @@ import { IMarkdownProps, IMarkdownSubComponentStyles, IMarkdownStyleProps, IMark
 import { MarkdownLink } from './MarkdownLink';
 import { MarkdownPre } from './MarkdownPre';
 import { PropsWithChildren } from 'react';
+import { modeColors } from '../../../styles/Styles';
+import dashboard from '../../../backend/Dashboard';
 
 // This is to work around inconsistency between the way markdown-to-jsx declares its types
 // (as having a default export) and the way it actually builds its files (for its cjs `main` file,
@@ -36,6 +38,11 @@ const getStyles: IStyleFunction<IMarkdownStyleProps, IMarkdownStyles> = () => {
             },
             'ul': {
                margin: "8px 0px 16px 0px"
+            },
+            'pre': {
+               whiteSpace: "pre-wrap",
+               padding: 16,
+               backgroundColor: dashboard.darktheme ? modeColors.header : modeColors.crumbs
             }
 
          }
diff --git a/Engine/Source/Programs/Horde/HordeDashboard/src/components/docs/DocView.tsx b/Engine/Source/Programs/Horde/HordeDashboard/src/components/docs/DocView.tsx
--- a/Engine/Source/Programs/Horde/HordeDashboard/src/components/docs/DocView.tsx
+++ b/Engine/Source/Programs/Horde/HordeDashboard/src/components/docs/DocView.tsx
@@ -217,7 +217,7 @@ export const DocView = () => {
       <DocCrumbs />
       <Stack horizontal>
          <div key={`windowsize_streamview_${windowSize.width}_${windowSize.height}`} style={{ width: vw / 2 - (1440 / 2), flexShrink: 0, backgroundColor: modeColors.background }} />
-         <Stack tokens={{ childrenGap: 0 }} styles={{ root: { backgroundColor: modeColors.background, width: "100%", "position": "relative", paddingTop: 12 } }}>
+         <Stack tokens={{ childrenGap: 0 }} styles={{ root: { backgroundColor: modeColors.background, width: "100%", "position": "relative", padding: "16px 32px" } }}>
             <div style={{ overflowY: 'scroll', overflowX: 'hidden', height: "calc(100vh - 162px)" }} data-is-scrollable={true}>
                <Stack horizontal>
                   <Stack style={{ width: 1240, paddingTop: 6, marginLeft: 4, height: '100%' }}>
diff --git a/Engine/Source/Programs/Horde/HordeDashboard/src/components/jobDetailsV2/JobDetailViewSteps.tsx b/Engine/Source/Programs/Horde/HordeDashboard/src/components/jobDetailsV2/JobDetailViewSteps.tsx
--- a/Engine/Source/Programs/Horde/HordeDashboard/src/components/jobDetailsV2/JobDetailViewSteps.tsx
+++ b/Engine/Source/Programs/Horde/HordeDashboard/src/components/jobDetailsV2/JobDetailViewSteps.tsx
@@ -429,7 +429,9 @@ export const StepsPanelInner: React.FC<{ jobDetails: JobDetailsV2, depStepId?: s
    const query = new URLSearchParams(window.location.search);
    const batchFilter = query.get("batch");
 
-   jobDetails.batches.forEach(b => {
+   let jobBatches = jobDetails.batches.sort((a, b) => a.groupIdx - b.groupIdx);
+
+   jobBatches.forEach(b => {
 
       if (batchFilter && b.id !== batchFilter) {
          return;
@@ -572,7 +574,7 @@ export const StepsPanelInner: React.FC<{ jobDetails: JobDetailsV2, depStepId?: s
       groups.add(jobDetails.getStepGroupIndex(step.id));
    });
 
-   const batches = jobDetails.batches.filter(b => {
+   let batches = jobBatches.filter(b => {
       if ((groups.size && !groups.has(b.groupIdx)) || b.steps.length || b.error === JobStepBatchError.None) {
          return false;
       }
