Horde: Another DocView tweak
#skipci
#rnx

diff --git a/Engine/Source/Programs/Horde/HordeDashboard/src/components/docs/DocView.tsx b/Engine/Source/Programs/Horde/HordeDashboard/src/components/docs/DocView.tsx
--- a/Engine/Source/Programs/Horde/HordeDashboard/src/components/docs/DocView.tsx
+++ b/Engine/Source/Programs/Horde/HordeDashboard/src/components/docs/DocView.tsx
@@ -29,7 +29,7 @@ export const docClasses = mergeStyleSets({
    raised: {
       backgroundColor: "#ffffff",
       boxShadow: "0 1.6px 3.6px 0 rgba(0,0,0,0.132), 0 0.3px 0.9px 0 rgba(0,0,0,0.108)",
-      padding: "32px 40px"      
+      padding: "32px 40px"
    }
 });
 
@@ -217,7 +217,7 @@ export const DocView = () => {
       <DocCrumbs />
       <Stack horizontal>
          <div key={`windowsize_streamview_${windowSize.width}_${windowSize.height}`} style={{ width: vw / 2 - (1440 / 2), flexShrink: 0, backgroundColor: modeColors.background }} />
-         <Stack tokens={{ childrenGap: 0 }} styles={{ root: { backgroundColor: modeColors.background, width: "100%", "position": "relative", padding: "16px 32px" } }}>
+         <Stack tokens={{ childrenGap: 0 }} styles={{ root: { backgroundColor: modeColors.background, width: "100%", "position": "relative", paddingTop: "16px", paddingLeft: "32px", paddingBottom: "16px", paddingRight: 0 } }}>
             <div style={{ overflowY: 'scroll', overflowX: 'hidden', height: "calc(100vh - 162px)" }} data-is-scrollable={true}>
                <Stack horizontal>
                   <Stack style={{ width: 1240, paddingTop: 6, marginLeft: 4, height: '100%' }}>
