Horde: Remove unused code from ComputeTaskSource.

diff --git a/Engine/Source/Programs/Horde/Horde.Server/Compute/ComputeTaskSource.cs b/Engine/Source/Programs/Horde/Horde.Server/Compute/ComputeTaskSource.cs
--- a/Engine/Source/Programs/Horde/Horde.Server/Compute/ComputeTaskSource.cs
+++ b/Engine/Source/Programs/Horde/Horde.Server/Compute/ComputeTaskSource.cs
@@ -102,49 +102,6 @@ namespace Horde.Server.Compute
 			_logger = logger;
 		}
 
-		/// <summary>
-		/// Allocates a compute resource
-		/// </summary>
-		/// <returns></returns>
-		public ComputeResource? TryAllocateResource(ClusterId clusterId, Requirements requirements)
-		{
-			lock (_lockObject)
-			{
-				LinkedList<Waiter>? waiters;
-				if (_waiters.TryGetValue(clusterId, out waiters))
-				{
-					ComputeTask? computeTask = null;
-					for (LinkedListNode<Waiter>? node = waiters.First; node != null; node = node.Next)
-					{
-						Dictionary<string, int> assignedResources = new Dictionary<string, int>(); 
-						if (node.Value.Agent.MeetsRequirements(requirements, assignedResources))
-						{
-							computeTask ??= CreateComputeTask(assignedResources);
-
-							byte[] payload = Any.Pack(computeTask).ToByteArray();
-							AgentLease lease = new AgentLease(LeaseId.GenerateNewId(), "Compute task", null, null, null, LeaseState.Pending, assignedResources, requirements.Exclusive, payload);
-
-							if (node.Value.Lease.TrySetResult(lease))
-							{
-								Waiter waiter = node.Value;
-								return new ComputeResource(waiter.Ip, waiter.Port, computeTask, waiter.Agent.Properties);
-							}
-						}
-					}
-				}
-			}
-			return null;
-		}
-
-		static ComputeTask CreateComputeTask(Dictionary<string, int> assignedResources)
-		{
-			ComputeTask computeTask = new ComputeTask();
-			computeTask.Nonce = UnsafeByteOperations.UnsafeWrap(RandomNumberGenerator.GetBytes(ServerComputeClient.NonceLength));
-			computeTask.Key = UnsafeByteOperations.UnsafeWrap(AesTransport.CreateKey());
-			computeTask.Resources.Add(assignedResources);
-			return computeTask;
-		}
-
 		/// <inheritdoc/>
 		public override Task<Task<AgentLease?>> AssignLeaseAsync(IAgent agent, CancellationToken cancellationToken)
 		{
