Horde: Various installer improvements/fixes.

* Agent is now included in distribution as a bundled tool.
* Add support for running the server as a Windows service, and set up to install as a service by default.
* A setup command is run as part of the installation which generates a unique JWT signing key and gRPC certificate, configures the agent to talk to the agent by default, and packs the agent data into a bundle.

diff --git a/Engine/Source/Programs/Horde/Horde.Server/Commands/Install/SetupCommand.cs b/Engine/Source/Programs/Horde/Horde.Server/Commands/Install/SetupCommand.cs
new file mode 100644
--- /dev/null
+++ b/Engine/Source/Programs/Horde/Horde.Server/Commands/Install/SetupCommand.cs
@@ -0,0 +1,193 @@
+// Copyright Epic Games, Inc. All Rights Reserved.
+
+using System;
+using System.Buffers;
+using System.Diagnostics;
+using System.IdentityModel.Tokens.Jwt;
+using System.IO;
+using System.Net;
+using System.Security.Cryptography;
+using System.Security.Cryptography.X509Certificates;
+using System.Text.Json;
+using System.Text.Json.Nodes;
+using System.Text.RegularExpressions;
+using System.Threading.Tasks;
+using EpicGames.Core;
+using EpicGames.Horde.Storage;
+using EpicGames.Horde.Storage.Backends;
+using EpicGames.Horde.Storage.Nodes;
+using Horde.Server.Agents;
+using Horde.Server.Utilities;
+using Microsoft.Extensions.Logging;
+using Microsoft.IdentityModel.Tokens;
+
+namespace Horde.Server.Commands.Install
+{
+	using JsonObject = System.Text.Json.Nodes.JsonObject;
+
+	/// <summary>
+	/// Performs custom actions required by an MSI installation on Windows; creates a unique certificate, configures the agent to use it, and packages the agent into a bundle.
+	/// </summary>
+	[Command("setup", "Runs post-install setup actions to create certificates, unique keys, etc...")]
+	public class SetupCommand : Command
+	{
+		[CommandLine("-Url=")]
+		string? ServerUrl { get; set; }
+
+		[CommandLine("-BaseDir")]
+		DirectoryReference BaseDir { get; set; } = Program.AppDir.ParentDirectory!;
+
+		/// <inheritdoc/>
+		public override async Task<int> ExecuteAsync(ILogger logger)
+		{
+			DirectoryReference agentDir = DirectoryReference.Combine(BaseDir, "Agent");
+			DirectoryReference serverDir = DirectoryReference.Combine(BaseDir, "Server");
+
+			// Get the DNS name of this server
+			string dnsName = Dns.GetHostName();
+			if (!String.IsNullOrEmpty(ServerUrl))
+			{
+				dnsName = ServerUrl;
+				dnsName = Regex.Replace(dnsName, @"^[a-zA-Z]+://", "");
+				dnsName = Regex.Replace(dnsName, "/.*$", "");
+			}
+			logger.LogInformation("DNS name for private cert: {DnsName}", dnsName);
+
+			// Generate a custom secret token for this installation
+			string jwtIssuer = dnsName;
+			byte[] jwtSecret = RandomNumberGenerator.GetBytes(128);
+
+			// Create a token for agents to register with this server
+			SigningCredentials signingCredentials = new(new SymmetricSecurityKey(jwtSecret), SecurityAlgorithms.HmacSha256);
+			JwtSecurityToken token = new JwtSecurityToken(issuer: jwtIssuer, claims: new[] { HordeClaims.AgentRegistrationClaim.ToClaim() }, signingCredentials: signingCredentials);
+			string tokenText = new JwtSecurityTokenHandler().WriteToken(token);
+
+			// Create a custom self-signed cert to use for agent comms
+			FileReference certFile = FileReference.Combine(serverDir, "ServerToAgent.pfx");
+
+			byte[] privateCertData = CertificateUtils.CreateSelfSignedCert(dnsName, "Horde Server");
+
+			logger.LogInformation("Writing private cert: {PrivateCert}", certFile.FullName);
+			FileReference.WriteAllBytes(certFile, privateCertData);
+
+			using X509Certificate2 certificate = new X509Certificate2(privateCertData);
+			logger.LogInformation("Certificate thumbprint is {Thumbprint}", certificate.Thumbprint);
+			logger.LogInformation("Add this thumbprint to list of trusted servers in appsettings.json to trust this server.");
+
+			// Update the agent to recognize the server certificate, and give it a custom token for being able to connect
+			FileReference agentConfigFile = FileReference.Combine(agentDir, "appsettings.json");
+			{
+				JsonObject agentConfig = await ReadConfigAsync(agentConfigFile);
+
+				JsonObject hordeConfig = FindOrAddNode(agentConfig, "Horde", () => new JsonObject());
+				hordeConfig["Server"] = "Default";
+
+				JsonArray serverProfiles = FindOrAddNode(hordeConfig, "ServerProfiles", () => new JsonArray());
+
+				JsonObject serverProfile = FindOrAddElementByKey(serverProfiles, "Name", "Default");
+				serverProfile["Environment"] = "Prod";
+				serverProfile["Url"] = ServerUrl ?? "http://localhost:5000";
+				serverProfile["Token"] = tokenText;
+				serverProfile["Thumbprints"] = new JsonArray(JsonValue.Create(certificate.Thumbprint));
+
+				await SaveConfigAsync(agentConfigFile, agentConfig);
+			}
+
+			// Create the local agent bundle
+			DirectoryReference bundleDir = DirectoryReference.Combine(serverDir, "Tools", "horde-agent");
+			DirectoryReference.CreateDirectory(bundleDir);
+
+			RefName refName = new RefName("default-ref");
+
+			FileStorageClient client = new FileStorageClient(bundleDir, logger);
+			using (TreeWriter writer = new TreeWriter(client, refName))
+			{
+				DirectoryNode dirNode = new DirectoryNode();
+				await dirNode.CopyFromDirectoryAsync(agentDir.ToDirectoryInfo(), new ChunkingOptions(), writer, null);
+				await writer.WriteAsync(refName, dirNode);
+			}
+
+			// Update the server config to include the bundled tool
+			FileReference serverConfigFile = FileReference.Combine(serverDir, "appsettings.json");
+			{
+				JsonObject serverConfig = await ReadConfigAsync(serverConfigFile);
+
+				JsonObject hordeConfig = FindOrAddNode(serverConfig, "Horde", () => new JsonObject());
+				hordeConfig[nameof(ServerSettings.JwtIssuer)] = dnsName;
+				hordeConfig[nameof(ServerSettings.JwtSecret)] = StringUtils.FormatHexString(jwtSecret);
+				hordeConfig[nameof(ServerSettings.ServerPrivateCert)] = certFile.MakeRelativeTo(serverDir);
+
+				JsonArray bundledTools = FindOrAddNode(hordeConfig, nameof(ServerSettings.BundledTools), () => new JsonArray());
+
+				JsonObject bundledTool = FindOrAddElementByKey(bundledTools, nameof(BundledToolConfig.Id), AgentExtensions.DefaultAgentSoftwareToolId.ToString());
+				bundledTool[nameof(BundledToolConfig.RefName)] = refName.ToString();
+
+				FileVersionInfo versionInfo = FileVersionInfo.GetVersionInfo(FileReference.Combine(agentDir, "hordeagent.exe").FullName);
+				if (!String.IsNullOrEmpty(versionInfo.ProductVersion))
+				{
+					bundledTool[nameof(BundledToolConfig.Version)] = versionInfo.ProductVersion;
+				}
+
+				await SaveConfigAsync(serverConfigFile, serverConfig);
+			}
+
+			return 0;
+		}
+
+		static T FindOrAddNode<T>(JsonObject obj, string name, Func<T> factory) where T : JsonNode
+		{
+			JsonNode? node = obj[name];
+			if (node != null)
+			{
+				if (node is T existingTypedNode)
+				{
+					return existingTypedNode;
+				}
+				else
+				{
+					obj.Remove(name);
+				}
+			}
+
+			T newTypedNode = factory();
+			obj.Add(name, newTypedNode);
+			return newTypedNode;
+		}
+
+		static JsonObject FindOrAddElementByKey(JsonArray array, string key, string name)
+		{
+			foreach (JsonNode? element in array)
+			{
+				if (element is JsonObject obj)
+				{
+					JsonNode? node = obj[key];
+					if (node != null && (string?)node.AsValue() == name)
+					{
+						return obj;
+					}
+				}
+			}
+
+			JsonObject newObj = new JsonObject();
+			newObj[key] = name;
+			return newObj;
+		}
+
+		static async Task<JsonObject> ReadConfigAsync(FileReference file)
+		{
+			byte[] data = await FileReference.ReadAllBytesAsync(file);
+			JsonObject? obj = JsonNode.Parse(data, new JsonNodeOptions { PropertyNameCaseInsensitive = true }, new JsonDocumentOptions { AllowTrailingCommas = true, CommentHandling = JsonCommentHandling.Skip }) as JsonObject;
+			return obj ?? new JsonObject();
+		}
+
+		static async Task SaveConfigAsync(FileReference file, JsonNode node)
+		{
+			ArrayBufferWriter<byte> buffer = new ArrayBufferWriter<byte>();
+			using (Utf8JsonWriter writer = new Utf8JsonWriter(buffer, new JsonWriterOptions { Indented = true }))
+			{
+				node.WriteTo(writer);
+			}
+			await FileReference.WriteAllBytesAsync(file, buffer.WrittenMemory.ToArray());
+		}
+	}
+}
diff --git a/Engine/Source/Programs/Horde/Horde.Server/Commands/ServerCommand.cs b/Engine/Source/Programs/Horde/Horde.Server/Commands/ServerCommand.cs
--- a/Engine/Source/Programs/Horde/Horde.Server/Commands/ServerCommand.cs
+++ b/Engine/Source/Programs/Horde/Horde.Server/Commands/ServerCommand.cs
@@ -2,6 +2,7 @@
 
 using System;
 using System.IO;
+using System.Runtime.InteropServices;
 using System.Security.Cryptography.X509Certificates;
 using System.Threading.Tasks;
 using EpicGames.Core;
@@ -40,6 +41,7 @@ namespace Horde.Server.Commands
 			using (X509Certificate2? grpcCertificate = ReadGrpcCertificate(_hordeSettings))
 			{
 				using IHost host = CreateHostBuilderWithCert(_args, _config, _hordeSettings, grpcCertificate).Build();
+
 				await host.RunAsync();
 				return 0;
 			}
@@ -47,7 +49,7 @@ namespace Horde.Server.Commands
 
 		static IHostBuilder CreateHostBuilderWithCert(string[] args, IConfiguration config, ServerSettings serverSettings, X509Certificate2? sslCert)
 		{
-			return Host.CreateDefaultBuilder(args)
+			IHostBuilder hostBuilder = Host.CreateDefaultBuilder(args)
 				.UseSerilog()
 				.ConfigureAppConfiguration(builder => builder.AddConfiguration(config))
 				.ConfigureWebHostDefaults(webBuilder =>
@@ -81,6 +83,21 @@ namespace Horde.Server.Commands
 					});
 					webBuilder.UseStartup<Startup>();
 				});
+
+			if (RuntimeInformation.IsOSPlatform(OSPlatform.Windows))
+			{
+				// Attempt to setup this process as a Windows service. A race condition inside Microsoft.Extensions.Hosting.WindowsServices.WindowsServiceHelpers.IsWindowsService
+				// can result in accessing the parent process after it's terminated, so catch any exceptions that it throws.
+				try
+				{
+					hostBuilder = hostBuilder.UseWindowsService();
+				}
+				catch (InvalidOperationException)
+				{
+				}
+			}
+
+			return hostBuilder;
 		}
 
 		/// <summary>
diff --git a/Engine/Source/Programs/Horde/Horde.Server/Horde.Server.csproj b/Engine/Source/Programs/Horde/Horde.Server/Horde.Server.csproj
--- a/Engine/Source/Programs/Horde/Horde.Server/Horde.Server.csproj
+++ b/Engine/Source/Programs/Horde/Horde.Server/Horde.Server.csproj
@@ -40,6 +40,7 @@
     <PackageReference Include="Microsoft.AspNetCore.DataProtection.StackExchangeRedis" Version="6.0.10" />
     <PackageReference Include="Microsoft.AspNetCore.SpaServices.Extensions" Version="6.0.10" />
     <PackageReference Include="Microsoft.Extensions.Caching.Memory" Version="6.0.1" />
+    <PackageReference Include="Microsoft.Extensions.Hosting.WindowsServices" Version="6.0.1" />
     <PackageReference Include="Microsoft.Extensions.Http" Version="6.0.0" />
     <PackageReference Include="MongoDB.Driver" Version="2.18.0" />
     <PackageReference Include="OpenTelemetry" Version="1.4.0" />
