Missing build script.

diff --git a/Engine/Source/Programs/Horde/BuildHorde.xml b/Engine/Source/Programs/Horde/BuildHorde.xml
--- a/Engine/Source/Programs/Horde/BuildHorde.xml
+++ b/Engine/Source/Programs/Horde/BuildHorde.xml
@@ -85,7 +85,7 @@
 		<Node Name="Stage Horde Installer" Requires="#DashboardOutputFiles">
 			<Delete Files="$(StagingDir)/Server/..."/>
 			<Delete Files="$(StagingDir)/Agent/..."/>
-			<DotNet Arguments="publish &quot;$(RootDir)/Engine/Source/Programs/Horde/Horde.Server/Horde.Server.csproj&quot; --output &quot;$(StagingDir)/Server&quot; --runtime win-x64 $(VersionArguments)"/>
+			<DotNet Arguments="publish &quot;$(RootDir)/Engine/Source/Programs/Horde/Horde.Server/Horde.Server.csproj&quot; --output &quot;$(StagingDir)/Server&quot; --runtime win-x64 --self-contained $(VersionArguments)"/>
 			<DotNet Arguments="publish &quot;$(RootDir)/Engine/Source/Programs/Horde/Horde.Agent/Horde.Agent.csproj&quot; --output &quot;$(StagingDir)/Agent&quot; $(VersionArguments)"/>
 
 			<!-- Create build settings -->
@@ -103,8 +103,11 @@
 			<Property Name="CommonArgs" Value="-dConfiguration=$(InstallerConfig) -dPlatform=x64 -arch x64"/>
 
 			<Delete Files="$(StagingDir)/Server-Bulk/..."/>
+			<Tag Files="$(StagingDir)/Agent/..." With="#ServerFiles"/>
 			<Tag Files="$(StagingDir)/Server/..." Except="Horde.Server.exe" With="#ServerFiles"/>
-			<Copy Files="#ServerFiles" From="$(StagingDir)/Server" To="$(StagingDir)/Server-Bulk/Server"/>
+
+			<Copy Files="#ServerFiles" From="$(StagingDir)" To="$(StagingDir)/Server-Bulk"/>
+
 			<Spawn Exe="$(WixDir)/heat.exe" Arguments="dir &quot;$(StagingDir)/Server-Bulk&quot; -cg HordeServerFiles -scom -sreg -gg -dr InstallDir -srd -var var.SourceDir -out &quot;$(InstallerDir)/HordeServerFiles.wxs&quot;"/>
 			<Spawn Exe="$(WixDir)/candle.exe" Arguments="$(CommonArgs) HordeServerFiles.wxs -dSourceDir=&quot;$(StagingDir)&quot; -out $(ObjDir)/HordeServerFiles.wixobj" WorkingDir="$(InstallerDir)"/>
 
