Horde: Add flags for showing various bits of functionality on the dashboard.

diff --git a/Engine/Source/Programs/Horde/Horde.Server/Users/UserMessages.cs b/Engine/Source/Programs/Horde/Horde.Server/Users/UserMessages.cs
--- a/Engine/Source/Programs/Horde/Horde.Server/Users/UserMessages.cs
+++ b/Engine/Source/Programs/Horde/Horde.Server/Users/UserMessages.cs
@@ -111,6 +111,26 @@ namespace Horde.Server.Users
 	/// </summary>
 	public class GetDashboardFeaturesResponse
 	{
+		/// <summary>
+		/// Enable CI functionality
+		/// </summary>
+		public bool ShowCI { get; set; }
+
+		/// <summary>
+		/// Show the Perforce server option on the server menu
+		/// </summary>
+		public bool ShowPerforceServers { get; set; }
+
+		/// <summary>
+		/// Show the device manager on the server menu
+		/// </summary>
+		public bool ShowDeviceManager { get; set; }
+
+		/// <summary>
+		/// Show automation on the server menu
+		/// </summary>
+		public bool ShowAutomation { get; set; }
+
 		/// <summary>
 		/// Whether the notice editor should be listed in the server menu
 		/// </summary>
@@ -131,6 +151,10 @@ namespace Horde.Server.Users
 		/// </summary>
 		public GetDashboardFeaturesResponse(GlobalConfig globalConfig, ClaimsPrincipal principal)
 		{
+			ShowCI = true;
+			ShowPerforceServers = true;
+			ShowDeviceManager = true;
+			ShowAutomation = true;
 			ShowNoticeEditor = globalConfig.Authorize(NoticeAclAction.CreateNotice, principal) || globalConfig.Authorize(NoticeAclAction.UpdateNotice, principal);
 			ShowPoolEditor = globalConfig.Authorize(PoolAclAction.CreatePool, principal) || globalConfig.Authorize(PoolAclAction.UpdatePool, principal);
 			ShowRemoteDesktop = globalConfig.Authorize(AgentAclAction.UpdateAgent, principal);
