Rename ShowAutomation flag to ShowTests.

diff --git a/Engine/Source/Programs/Horde/Horde.Server/Users/UserMessages.cs b/Engine/Source/Programs/Horde/Horde.Server/Users/UserMessages.cs
--- a/Engine/Source/Programs/Horde/Horde.Server/Users/UserMessages.cs
+++ b/Engine/Source/Programs/Horde/Horde.Server/Users/UserMessages.cs
@@ -127,9 +127,9 @@ namespace Horde.Server.Users
 		public bool ShowDeviceManager { get; set; }
 
 		/// <summary>
-		/// Show automation on the server menu
+		/// Show automated tests on the server menu
 		/// </summary>
-		public bool ShowAutomation { get; set; }
+		public bool ShowTests { get; set; }
 
 		/// <summary>
 		/// Whether the notice editor should be listed in the server menu
@@ -154,7 +154,7 @@ namespace Horde.Server.Users
 			ShowCI = true;
 			ShowPerforceServers = true;
 			ShowDeviceManager = true;
-			ShowAutomation = true;
+			ShowTests = true;
 			ShowNoticeEditor = globalConfig.Authorize(NoticeAclAction.CreateNotice, principal) || globalConfig.Authorize(NoticeAclAction.UpdateNotice, principal);
 			ShowPoolEditor = globalConfig.Authorize(PoolAclAction.CreatePool, principal) || globalConfig.Authorize(PoolAclAction.UpdatePool, principal);
 			ShowRemoteDesktop = globalConfig.Authorize(AgentAclAction.UpdateAgent, principal);
