Horde: Add a flag for showing agent-related functionality on the dashboard.

diff --git a/Engine/Source/Programs/Horde/Horde.Server/Users/UserMessages.cs b/Engine/Source/Programs/Horde/Horde.Server/Users/UserMessages.cs
--- a/Engine/Source/Programs/Horde/Horde.Server/Users/UserMessages.cs
+++ b/Engine/Source/Programs/Horde/Horde.Server/Users/UserMessages.cs
@@ -116,6 +116,11 @@ namespace Horde.Server.Users
 		/// </summary>
 		public bool ShowCI { get; set; }
 
+		/// <summary>
+		/// Whether to show functionality related to agents, pools, and utilization on the dashboard.
+		/// </summary>
+		public bool ShowAgents { get; set; }
+
 		/// <summary>
 		/// Show the Perforce server option on the server menu
 		/// </summary>
@@ -152,6 +157,7 @@ namespace Horde.Server.Users
 		public GetDashboardFeaturesResponse(GlobalConfig globalConfig, ClaimsPrincipal principal)
 		{
 			ShowCI = true;
+			ShowAgents = true;
 			ShowPerforceServers = true;
 			ShowDeviceManager = true;
 			ShowTests = true;
