Horde: Fix issue with CL previews in change button
#rnx
#skipci

diff --git a/Engine/Source/Programs/Horde/HordeDashboard/src/components/ChangeButton.tsx b/Engine/Source/Programs/Horde/HordeDashboard/src/components/ChangeButton.tsx
--- a/Engine/Source/Programs/Horde/HordeDashboard/src/components/ChangeButton.tsx
+++ b/Engine/Source/Programs/Horde/HordeDashboard/src/components/ChangeButton.tsx
@@ -52,7 +52,7 @@ export const ChangeContextMenu: React.FC<{ target: ChangeContextMenuTarget, onDi
 
       menuItems.push({
          key: "commit", onRender: () => {
-            return <Stack style={{ padding: 18, paddingTop: 18, paddingBottom: 18, maxWidth: 500 }} tokens={{ childrenGap: 12 }}>
+            return <Stack style={{ padding: 18, paddingTop: 18, paddingBottom: 18, maxWidth: 800 }} tokens={{ childrenGap: 12 }}>
                <Stack horizontal tokens={{ childrenGap: 36 }}>
                   <Text variant="small" style={{ fontFamily: "Horde Open Sans Bold" }}>Author:</Text>
                   <Text variant="small" >{commit.authorInfo?.name}</Text>
@@ -64,7 +64,7 @@ export const ChangeContextMenu: React.FC<{ target: ChangeContextMenuTarget, onDi
 
                <Stack horizontal tokens={{ childrenGap: 12 }}>
                   <Text variant="small" style={{ fontFamily: "Horde Open Sans Bold" }}>Description:</Text>
-                  <Text variant="small" >{commit.description}</Text>
+                  <Text variant="small" style={{whiteSpace: "pre-wrap"}}>{commit.description}</Text>
                </Stack>
             </Stack>
          }
@@ -79,10 +79,10 @@ export const ChangeContextMenu: React.FC<{ target: ChangeContextMenuTarget, onDi
    } else if (job.preflightDescription) {
       menuItems.push({
          key: "commit", onRender: () => {
-            return <Stack style={{ padding: 18, paddingTop: 18, paddingBottom: 18, maxWidth: 500 }} tokens={{ childrenGap: 12 }}>
+            return <Stack style={{ padding: 18, paddingTop: 18, paddingBottom: 18, maxWidth: 800 }} tokens={{ childrenGap: 12 }}>
                <Stack horizontal tokens={{ childrenGap: 12 }}>
                   <Text variant="small" style={{ fontFamily: "Horde Open Sans Bold" }}>Description:</Text>
-                  <Text variant="small" >{job.preflightDescription}</Text>
+                  <Text variant="small" style={{whiteSpace: "pre-wrap"}} >{job.preflightDescription}</Text>
                </Stack>
             </Stack>
          }
@@ -131,7 +131,7 @@ export const ChangeContextMenu: React.FC<{ target: ChangeContextMenuTarget, onDi
    }
    
    return (<ContextualMenu
-      styles={{ list: { selectors: { '.ms-ContextualMenu-itemText': { fontSize: "10px", paddingLeft: 8 } } } }}
+      styles={{ root: {paddingBottom: 12, paddingRight: 24, paddingLeft: 8, paddingTop: 12},list: { selectors: { '.ms-ContextualMenu-itemText': { fontSize: "10px", paddingLeft: 8 } } } }}
       items={menuItems}
       hidden={false}
       target={target.ref ?? target.point}
