Horde: Fix cancellation exception preventing shutdown in MongoService.

diff --git a/Engine/Source/Programs/Horde/Horde.Server/Server/MongoService.cs b/Engine/Source/Programs/Horde/Horde.Server/Server/MongoService.cs
--- a/Engine/Source/Programs/Horde/Horde.Server/Server/MongoService.cs
+++ b/Engine/Source/Programs/Horde/Horde.Server/Server/MongoService.cs
@@ -985,6 +985,10 @@ namespace Horde.Server.Server
 						}
 					}
 				}
+				catch(OperationCanceledException) when (cancellationToken.IsCancellationRequested)
+				{
+					break;
+				}
 				catch (Exception ex)
 				{
 					_logger.LogWarning(ex, "Exception updating indexes: {Message}", ex.Message);
