Horde: Fix registration of bundled tools with Horde Server.

diff --git a/Engine/Source/Programs/Horde/Horde.Server/Commands/Install/SetupCommand.cs b/Engine/Source/Programs/Horde/Horde.Server/Commands/Install/SetupCommand.cs
--- a/Engine/Source/Programs/Horde/Horde.Server/Commands/Install/SetupCommand.cs
+++ b/Engine/Source/Programs/Horde/Horde.Server/Commands/Install/SetupCommand.cs
@@ -170,6 +170,7 @@ namespace Horde.Server.Commands.Install
 
 			JsonObject newObj = new JsonObject();
 			newObj[key] = name;
+			array.Add(newObj);
 			return newObj;
 		}
 
