Horde: DocView changes and adding Home/Landing pages
#skipci
#rnx

diff --git a/Engine/Source/Programs/Horde/HordeDashboard/src/App.tsx b/Engine/Source/Programs/Horde/HordeDashboard/src/App.tsx
--- a/Engine/Source/Programs/Horde/HordeDashboard/src/App.tsx
+++ b/Engine/Source/Programs/Horde/HordeDashboard/src/App.tsx
@@ -129,14 +129,14 @@ const Main: React.FC = () => {
       const routes: RouteObject[] = [
          {
             path: "/", element: <Root />, errorElement: <RouteError />, children: [
-               { path: "index", element: <UserHomeView /> },
+               { path: "index", element: (dashboard.user?.dashboardFeatures?.showCI === false) ? <DocView /> : <UserHomeView /> },
                { path: "project/:projectId", element: <ProjectHome /> },
                { path: "pools", element: <PoolView /> },
                { path: "job/:jobId", element: <JobDetailViewV2 /> },
                { path: "log/:logId", element: <LogView /> },
                { path: "testreport/:testdataId", element: <TestReportView /> },
                { path: "stream/:streamId", element: <StreamView /> },
-               { path: "agents", element: <AgentView/> },
+               { path: "agents", element: <AgentView /> },
                { path: "admin/token", element: <AdminToken /> },
                { path: "reports/utilization", element: <UtilizationReportView /> },
                { path: "preflight", element: <PreflightRedirector /> },
diff --git a/Engine/Source/Programs/Horde/HordeDashboard/src/backend/Api.ts b/Engine/Source/Programs/Horde/HordeDashboard/src/backend/Api.ts
--- a/Engine/Source/Programs/Horde/HordeDashboard/src/backend/Api.ts
+++ b/Engine/Source/Programs/Horde/HordeDashboard/src/backend/Api.ts
@@ -162,25 +162,25 @@ export enum EventSeverity {
 
 export enum AclActions {
 
-	// No actions allowed	
+	// No actions allowed
 	None = "None",
 
-	// Reading to the object	
+	// Reading to the object
 	Read = "Read",
 
-	// Modifying the object	
+	// Modifying the object
 	Write = "Write",
 
-	// Creating new objects	
+	// Creating new objects
 	Create = "Create",
 
-	// Deleting the object	
+	// Deleting the object
 	Delete = "Delete",
 
-	// Executing the object	
+	// Executing the object
 	Execute = "Execute",
 
-	// Change permissions on the object	
+	// Change permissions on the object
 	ChangePermissions = "ChangePermissions"
 
 }
@@ -762,7 +762,7 @@ export type LeaseUtilizationSettings = {
 export type JobQueueSettings = {
 
 	/// Factor translating queue size to additional agents to grow the pool with
-	/// The result is always rounded up to nearest integer. 
+	/// The result is always rounded up to nearest integer.
 	/// Example: if there are 20 jobs in queue, a factor 0.25 will result in 5 new agents being added (20 * 0.25)
 	scaleOutFactor: number;
 
@@ -802,7 +802,7 @@ export type GetPoolResponse = {
 	/// Settings for lease utilization pool sizing strategy (if used)
 	leaseUtilizationSettings?: LeaseUtilizationSettings;
 
-	/// Settings for job queue pool sizing strategy (if used) 
+	/// Settings for job queue pool sizing strategy (if used)
 	jobQueueSettings?: JobQueueSettings;
 
 	/// The minimum nunmber of agents to retain in this pool
@@ -1992,19 +1992,19 @@ export type ChangeQueryConfig = {
 	/// Name of this query, for display on the dashboard.
 	name?: string;
 
-	/// Condition to evaluate before deciding to use this query. May query tags in a preflight.	
+	/// Condition to evaluate before deciding to use this query. May query tags in a preflight.
 	condition?: any;
 
-	/// The template id to query	
+	/// The template id to query
 	templateId?: string;
 
-	/// The target to query	
+	/// The target to query
 	target?: string;
 
-	/// Whether to match a job that produced warnings	
+	/// Whether to match a job that produced warnings
 	outcomes?: JobStepOutcome[];
 
-	/// Finds the last commit with this tag	
+	/// Finds the last commit with this tag
 	commitTag?: any;
 }
 
@@ -3104,6 +3104,21 @@ export type GetDashboardFeaturesResponse = {
 
 	/** Whether the remote desktop button should be shown on the agent modal */
 	showRemoteDesktop?: boolean;
+
+	/** Enable CI functionality */
+	showCI?: boolean;
+
+	/** Whether to show functionality related to agents, pools, and utilization on the dashboard. */
+	showAgents?: boolean;
+
+	/** Show the Perforce server option on the server menu */
+	showPerforceServers?: boolean;
+
+	/** Show the device manager on the server menu */
+	showDeviceManager?: boolean;
+
+	/** Show automated tests on the server menu */
+	showTests?: boolean;
 }
 
 /**  Response describing the current user */
@@ -3604,7 +3619,7 @@ export type ProjectConfigRef = {
 	/// Unique id for the project
 	id: string;
 
-	/// Config path for the project		
+	/// Config path for the project
 	path: string;
 }
 
@@ -3699,7 +3714,7 @@ export type PerforceCredentials = {
 }
 
 
-/// Information about a cluster of Perforce servers. 
+/// Information about a cluster of Perforce servers.
 export type PerforceCluster = {
 
 	/// The default cluster name
@@ -4041,7 +4056,7 @@ export type FindIssueSpanResponse = {
 	/// Workflow for this span
 	workflowId?: string;
 
-	/// The previous build 
+	/// The previous build
 	lastSuccess?: GetIssueStepResponse;
 
 	/// The following successful build
@@ -4197,7 +4212,7 @@ export type CreateExternalIssueResponse = {
 /// External issue project information
 export type GetExternalIssueProjectResponse = {
 
-	/// The project key	
+	/// The project key
 	projectKey: string;
 
 	/// The name of the project
@@ -4337,31 +4352,31 @@ export type GetDevicePoolTelemetryResponse = {
 /// Device telemetry respponse
 export type GetTelemetryInfoResponse = {
 
-	/// The UTC time the telemetry data was created	
+	/// The UTC time the telemetry data was created
 	createTimeUtc: Date | string;
 
-	/// The stream id which utilized device	
+	/// The stream id which utilized device
 	streamId?: string;
 
-	/// The job id which utilized device	
+	/// The job id which utilized device
 	jobId?: string;
 
-	/// The job's step id	
+	/// The job's step id
 	stepId?: string;
 
-	/// The job name which utilized device	
+	/// The job name which utilized device
 	jobName?: string;
 
 	/// The job's step name
 	stepName?: string;
 
-	/// If this telemetry has a reservation, the start time of the reservation	
+	/// If this telemetry has a reservation, the start time of the reservation
 	reservationStartUtc?: Date | string;
 
 	/// If this telemetry has a reservation, the finish time of the reservation
 	reservationFinishUtc?: Date | string;
 
-	/// If this telemetry marks a detected device issue, the time of the issue	
+	/// If this telemetry marks a detected device issue, the time of the issue
 	problemTimeUtc?: Date | string;
 }
 
@@ -4372,32 +4387,32 @@ export type GetDeviceTelemetryResponse = {
 	/// The device id for the telemetry data
 	deviceId: string;
 
-	/// Individual telemetry data points	
+	/// Individual telemetry data points
 	telemetry: GetTelemetryInfoResponse[];
 
 }
 
 // Test Data ----------------------------
 
-/// A test emvironment running in a stream	
+/// A test emvironment running in a stream
 export type GetTestMetaResponse = {
 
 	/// Meta unique id for environment
 	id: string;
 
-	/// The platforms in the environment		
+	/// The platforms in the environment
 	platforms: string[];
 
-	/// The build configurations being tested		
+	/// The build configurations being tested
 	configurations: string[];
 
-	/// The build targets being tested		
+	/// The build targets being tested
 	buildTargets: string[];
 
-	/// The test project name		
+	/// The test project name
 	projectName: string;
 
-	/// The rendering hardware interface being used with the test		
+	/// The rendering hardware interface being used with the test
 	rhi: string;
 
 	/// The test variation identifier (or "default")
@@ -4412,10 +4427,10 @@ export type GetTestResponse = {
 	/// The id of the test
 	id: string;
 
-	/// The name of the test 		
+	/// The name of the test
 	name: string;
 
-	/// The display name of the test, if any 		
+	/// The display name of the test, if any
 	displayName?: string;
 
 	/// The suite the test belongs to, if any
@@ -4432,13 +4447,13 @@ export type GetTestsRequest = {
 }
 
 
-/// A test suite that runs in a stream, contain subtests	
+/// A test suite that runs in a stream, contain subtests
 export type GetTestSuiteResponse = {
 
 	/// The id of the test suite
 	id: string;
 
-	/// The name of the test suite		
+	/// The name of the test suite
 	name: string;
 
 	/// The environments the suite runs in
@@ -4448,16 +4463,16 @@ export type GetTestSuiteResponse = {
 /// Describes tests running in a stream
 export type GetTestStreamResponse = {
 
-	/// The stream id	
+	/// The stream id
 	streamId: string;
 
-	/// Individual tests which run in the stream	
+	/// Individual tests which run in the stream
 	tests: GetTestResponse[];
 
-	/// Test suites that run in the stream	
+	/// Test suites that run in the stream
 	testSuites: GetTestSuiteResponse[];
 
-	/// Test suites that run in the stream	
+	/// Test suites that run in the stream
 	testMetadata: GetTestMetaResponse[];
 }
 
@@ -4479,10 +4494,10 @@ export enum TestOutcome {
 /// Suite test data
 export type GetSuiteTestDataResponse = {
 
-	/// The test id	
+	/// The test id
 	testId: string;
 
-	/// The ourcome of the suite test	
+	/// The ourcome of the suite test
 	outcome: TestOutcome;
 
 	/// How long the suite test ran	 (TimeSpan)
@@ -4501,13 +4516,13 @@ export type GetSuiteTestDataResponse = {
 /// Test details
 export type GetTestDataDetailsResponse = {
 
-	/// The corresponding test ref	
+	/// The corresponding test ref
 	id: string;
 
-	/// The test documents for this ref	
+	/// The test documents for this ref
 	testDataIds: string[];
 
-	/// Suite test data		
+	/// Suite test data
 	suiteTests?: GetSuiteTestDataResponse[];
 
 }
@@ -4526,7 +4541,7 @@ export type GetTestDataRefResponse = {
 	/// The build changelist upon which the test ran, may not correspond to the job changelist
 	buildChangeList: number;
 
-	/// The environment the test ran on 
+	/// The environment the test ran on
 	metaId: string;
 
 	/// The test id in stream
diff --git a/Engine/Source/Programs/Horde/HordeDashboard/src/components/TopNav.tsx b/Engine/Source/Programs/Horde/HordeDashboard/src/components/TopNav.tsx
--- a/Engine/Source/Programs/Horde/HordeDashboard/src/components/TopNav.tsx
+++ b/Engine/Source/Programs/Horde/HordeDashboard/src/components/TopNav.tsx
@@ -358,68 +358,84 @@ export const TopNav: React.FC<{ suppressServer?: boolean }> = observer(({ suppre
       // Resources
       const resourceItems: IContextualMenuItem[] = [];
 
-      resourceItems.push({
-         key: "admin_agents",
-         text: "Agents",
-         link: `/agents`
-      });
+      const features = dashboard.user?.dashboardFeatures;
 
-      resourceItems.push({
-         key: "admin_pools",
-         text: "Pools",
-         link: `/pools`
-      });
+      if (features?.showAgents !== false) {
+         resourceItems.push({
+            key: "admin_agents",
+            text: "Agents",
+            link: `/agents`
+         });
 
-      resourceItems.push({
-         key: "admin_perforce_servers",
-         text: "Perforce Servers",
-         link: `/perforce/servers`
-      });
+         resourceItems.push({
+            key: "admin_pools",
+            text: "Pools",
+            link: `/pools`
+         });
+      }
 
-      subItems.push({
-         itemType: ContextualMenuItemType.Section,
-         key: `admin_resources`,
-         sectionProps: {
-            title: "Resources",
-            items: resourceItems,
-            bottomDivider: true
-         }
-      });
+      if (features?.showPerforceServers !== false) {
+         resourceItems.push({
+            key: "admin_perforce_servers",
+            text: "Perforce Servers",
+            link: `/perforce/servers`
+         });
+      }
+
+      if (resourceItems.length) {
+         subItems.push({
+            itemType: ContextualMenuItemType.Section,
+            key: `admin_resources`,
+            sectionProps: {
+               title: "Resources",
+               items: resourceItems,
+               bottomDivider: true
+            }
+         });
+      }
 
       // Automation
       const automatonItems: IContextualMenuItem[] = [];
 
-      automatonItems.push({
-         key: "admin_devices",
-         text: "Devices",
-         link: `/devices`
-      });
-
-      automatonItems.push({
-         key: "admin_tests",
-         text: "Tests",
-         link: `/automation`
-      });
+      if (features?.showDeviceManager !== false) {
+         automatonItems.push({
+            key: "admin_devices",
+            text: "Devices",
+            link: `/devices`
+         });
+      }
 
-      subItems.push({
-         itemType: ContextualMenuItemType.Section,
-         key: `admin_automation`,
-         sectionProps: {
-            title: "Automation",
-            items: automatonItems,
-            bottomDivider: true
-         }
-      });
+      if (features?.showTests !== false) {
+         automatonItems.push({
+            key: "admin_tests",
+            text: "Tests",
+            link: `/automation`
+         });
+      }
+      if (automatonItems.length) {
+         subItems.push({
+            itemType: ContextualMenuItemType.Section,
+            key: `admin_automation`,
+            sectionProps: {
+               title: "Automation",
+               items: automatonItems,
+               bottomDivider: true
+            }
+         });
+      }
 
       // Monitoring
       const monitoringItems: IContextualMenuItem[] = [];
-      monitoringItems.push({
-         key: "admin_utilization",
-         text: "Utilization",
-         link: `/reports/utilization`
-      });
+      if (features?.showAgents !== false) {
+         monitoringItems.push({
+            key: "admin_utilization",
+            text: "Utilization",
+            link: `/reports/utilization`
+         });
+      }
+
 
-      if (dashboard.user?.dashboardFeatures?.showNoticeEditor) {
+      if (features?.showNoticeEditor !== false) {
          monitoringItems.push({
             key: "admin_notices",
             text: "Notices",
@@ -427,15 +443,17 @@ export const TopNav: React.FC<{ suppressServer?: boolean }> = observer(({ suppre
          });
       }
 
-      subItems.push({
-         itemType: ContextualMenuItemType.Section,
-         key: `admin_monitoring`,
-         sectionProps: {
-            title: "Monitoring",
-            items: monitoringItems,
-            bottomDivider: true
-         }
-      });
+      if (monitoringItems.length) {
+         subItems.push({
+            itemType: ContextualMenuItemType.Section,
+            key: `admin_monitoring`,
+            sectionProps: {
+               title: "Monitoring",
+               items: monitoringItems,
+               bottomDivider: true
+            }
+         });   
+      }
 
       const softwareItems: IContextualMenuItem[] = [];
       softwareItems.push({
diff --git a/Engine/Source/Programs/Horde/HordeDashboard/src/components/docs/DocView.tsx b/Engine/Source/Programs/Horde/HordeDashboard/src/components/docs/DocView.tsx
--- a/Engine/Source/Programs/Horde/HordeDashboard/src/components/docs/DocView.tsx
+++ b/Engine/Source/Programs/Horde/HordeDashboard/src/components/docs/DocView.tsx
@@ -9,6 +9,7 @@ import { useWindowSize } from "../../base/utilities/hooks";
 import { hordeClasses, modeColors } from "../../styles/Styles";
 import { BreadcrumbItem, Breadcrumbs } from "../Breadcrumbs";
 import { TopNav } from "../TopNav";
+import dashboard from "../../backend/Dashboard";
 
 type Anchor = {
    text: string;
@@ -16,6 +17,7 @@ type Anchor = {
 }
 
 const documentCache = new Map<string, string>();
+const documentRequests = new Set<string>();
 const crumbCache = new Map<string, BreadcrumbItem[]>();
 const anchorCache = new Map<string, Anchor[]>();
 
@@ -59,7 +61,8 @@ const DocPanel: React.FC<{ docName: string }> = ({ docName }) => {
 
    const text = documentCache.get(docName);
 
-   if (!text) {
+   if (!documentRequests.has(docName)) {
+      documentRequests.add(docName);
       fetch(`/${docName}`, { cache: "no-cache" })
          .then((response) => response.text())
          .then((textContent) => {
@@ -156,7 +159,7 @@ const DocPanel: React.FC<{ docName: string }> = ({ docName }) => {
 
 
    return <Stack styles={{ root: { width: "100%" } }} >
-      <div ref={docRef}>
+      <div ref={docRef} style={{margin: "16px 32px"}}>
          <Markdown>{text}</Markdown>
       </div>
    </Stack>;
@@ -167,7 +170,12 @@ const DocRail = observer(() => {
    const state = linkState.state;
 
    const refLinks: ISideRailLink[] = [];
-   refLinks.push({ text: "Introduction", url: "/docs" });
+
+   if (dashboard.user?.dashboardFeatures?.showCI === false) {
+      refLinks.push({ text: "Landing", url: "/docs" });
+   }
+
+   refLinks.push({ text: "Home", url: "/docs/Home.md" });
    refLinks.push({ text: "User Guild", url: "/docs/Users.md" });
    refLinks.push({ text: "Deployment", url: "/docs/Deployment.md" });
    refLinks.push({ text: "Configuration", url: "/docs/Config.md" });
@@ -198,8 +206,17 @@ export const DocView = () => {
 
    // fixme
    let docName = location.pathname.replace("/docs/", "").replace("/docs", "").trim();
+   if (docName.startsWith("/index")) {
+      docName = docName.replace("/index", "")
+   }
    if (!docName || docName.indexOf("README.md") !== -1) {
-      docName = "documentation/README.md";
+
+      if (dashboard.user?.dashboardFeatures?.showCI === false) {
+         docName = "documentation/Docs/Landing.md";
+      } else {
+         docName = "documentation/Docs/Home.md";
+      }
+      
    } else {
 
       if (!docName.startsWith("Docs/")) {
