Horde: Styling for markdown code element
#rnx
#skipci

diff --git a/Engine/Source/Programs/Horde/HordeDashboard/src/base/components/Markdown/Markdown.tsx b/Engine/Source/Programs/Horde/HordeDashboard/src/base/components/Markdown/Markdown.tsx
--- a/Engine/Source/Programs/Horde/HordeDashboard/src/base/components/Markdown/Markdown.tsx
+++ b/Engine/Source/Programs/Horde/HordeDashboard/src/base/components/Markdown/Markdown.tsx
@@ -43,6 +43,13 @@ const getStyles: IStyleFunction<IMarkdownStyleProps, IMarkdownStyles> = () => {
                whiteSpace: "pre-wrap",
                padding: 16,
                backgroundColor: dashboard.darktheme ? modeColors.header : modeColors.crumbs
+            },
+            'code': {
+               whiteSpace: "pre-wrap",
+               fontFamily: "Horde Cousine Regular",
+               fontSize: "12px",
+               padding: "3px 6px",
+               backgroundColor: dashboard.darktheme ? modeColors.header : "#f2f2f2"
             }
 
          }
