Horde: Add UI to the installer, including a page that allows configuring the server URL.

diff --git a/Engine/Source/Programs/Horde/BuildHorde.xml b/Engine/Source/Programs/Horde/BuildHorde.xml
--- a/Engine/Source/Programs/Horde/BuildHorde.xml
+++ b/Engine/Source/Programs/Horde/BuildHorde.xml
@@ -112,7 +112,7 @@
 			<Spawn Exe="$(WixDir)/candle.exe" Arguments="$(CommonArgs) HordeServerFiles.wxs -dSourceDir=&quot;$(StagingDir)&quot; -out $(ObjDir)/HordeServerFiles.wixobj" WorkingDir="$(InstallerDir)"/>
 
 			<Spawn Exe="$(WixDir)/candle.exe" Arguments="$(CommonArgs) Installer.wxs -dStagingDir=&quot;$(StagingDir)&quot; -out $(ObjDir)/Installer.wixobj" WorkingDir="$(InstallerDir)"/>
-			<Spawn Exe="$(WixDir)/light.exe" Arguments="-out $(BinDir)/Horde.msi -sw1076 -pdbout $(BinDir)/Horde.pdb $(ObjDir)/Installer.wixobj $(ObjDir)/HordeServerFiles.wixobj" WorkingDir="$(InstallerDir)"/>
+			<Spawn Exe="$(WixDir)/light.exe" Arguments="-out $(BinDir)/Horde.msi -sw1076 -pdbout $(BinDir)/Horde.pdb $(ObjDir)/Installer.wixobj $(ObjDir)/HordeServerFiles.wixobj -ext WixUIExtension -ext WixUtilExtension" WorkingDir="$(InstallerDir)"/>
 		</Node>
 		
 	</Agent>
