Horde: Missing file.

diff --git a/Engine/Source/Programs/Horde/BuildHorde.xml b/Engine/Source/Programs/Horde/BuildHorde.xml
--- a/Engine/Source/Programs/Horde/BuildHorde.xml
+++ b/Engine/Source/Programs/Horde/BuildHorde.xml
@@ -102,17 +102,25 @@
 		<Node Name="Build Horde Installer" Requires="Stage Horde Installer">
 			<Property Name="CommonArgs" Value="-dConfiguration=$(InstallerConfig) -dPlatform=x64 -arch x64"/>
 
+			<!-- Create the agent file list -->
+			<Delete Files="$(StagingDir)/Agent-Bulk/..."/>
+			<Tag Files="$(StagingDir)/Agent/..." Except="HordeAgent.exe" With="#AgentFiles"/>
+			<Copy Files="#AgentFiles" From="$(StagingDir)" To="$(StagingDir)/Agent-Bulk"/>
+
+			<Spawn Exe="$(WixDir)/heat.exe" Arguments="dir &quot;$(StagingDir)/Agent-Bulk&quot; -cg HordeAgentFiles -scom -sreg -gg -dr InstallDir -srd -var var.SourceDir -out &quot;$(InstallerDir)/HordeAgentFiles.wxs&quot;"/>
+			<Spawn Exe="$(WixDir)/candle.exe" Arguments="$(CommonArgs) HordeAgentFiles.wxs -dSourceDir=&quot;$(StagingDir)&quot; -out $(ObjDir)/HordeAgentFiles.wixobj" WorkingDir="$(InstallerDir)"/>
+
+			<!-- Create the server file list -->
 			<Delete Files="$(StagingDir)/Server-Bulk/..."/>
-			<Tag Files="$(StagingDir)/Agent/..." With="#ServerFiles"/>
 			<Tag Files="$(StagingDir)/Server/..." Except="Horde.Server.exe" With="#ServerFiles"/>
-
 			<Copy Files="#ServerFiles" From="$(StagingDir)" To="$(StagingDir)/Server-Bulk"/>
 
 			<Spawn Exe="$(WixDir)/heat.exe" Arguments="dir &quot;$(StagingDir)/Server-Bulk&quot; -cg HordeServerFiles -scom -sreg -gg -dr InstallDir -srd -var var.SourceDir -out &quot;$(InstallerDir)/HordeServerFiles.wxs&quot;"/>
 			<Spawn Exe="$(WixDir)/candle.exe" Arguments="$(CommonArgs) HordeServerFiles.wxs -dSourceDir=&quot;$(StagingDir)&quot; -out $(ObjDir)/HordeServerFiles.wixobj" WorkingDir="$(InstallerDir)"/>
 
+			<!-- Main component -->
 			<Spawn Exe="$(WixDir)/candle.exe" Arguments="$(CommonArgs) Installer.wxs -dStagingDir=&quot;$(StagingDir)&quot; -out $(ObjDir)/Installer.wixobj" WorkingDir="$(InstallerDir)"/>
-			<Spawn Exe="$(WixDir)/light.exe" Arguments="-out $(BinDir)/Horde.msi -sw1076 -pdbout $(BinDir)/Horde.pdb $(ObjDir)/Installer.wixobj $(ObjDir)/HordeServerFiles.wixobj -ext WixUIExtension -ext WixUtilExtension" WorkingDir="$(InstallerDir)"/>
+			<Spawn Exe="$(WixDir)/light.exe" Arguments="-out $(BinDir)/Horde.msi -sw1076 -pdbout $(BinDir)/Horde.pdb $(ObjDir)/Installer.wixobj $(ObjDir)/HordeAgentFiles.wixobj $(ObjDir)/HordeServerFiles.wixobj -ext WixUIExtension -ext WixUtilExtension" WorkingDir="$(InstallerDir)"/>
 		</Node>
 		
 	</Agent>
