Horde: Fix typo in docs sidebar.

diff --git a/Engine/Source/Programs/Horde/HordeDashboard/src/components/docs/DocView.tsx b/Engine/Source/Programs/Horde/HordeDashboard/src/components/docs/DocView.tsx
--- a/Engine/Source/Programs/Horde/HordeDashboard/src/components/docs/DocView.tsx
+++ b/Engine/Source/Programs/Horde/HordeDashboard/src/components/docs/DocView.tsx
@@ -176,7 +176,7 @@ const DocRail = observer(() => {
    }
 
    refLinks.push({ text: "Home", url: "/docs/Home.md" });
-   refLinks.push({ text: "User Guild", url: "/docs/Users.md" });
+   refLinks.push({ text: "User Guide", url: "/docs/Users.md" });
    refLinks.push({ text: "Deployment", url: "/docs/Deployment.md" });
    refLinks.push({ text: "Configuration", url: "/docs/Config.md" });
    refLinks.push({ text: "Horde Internals", url: "/docs/Internals.md" });
