Horde: Fix placement of documentation files in installed dashboard.

diff --git a/Engine/Source/Programs/Horde/BuildHorde.xml b/Engine/Source/Programs/Horde/BuildHorde.xml
--- a/Engine/Source/Programs/Horde/BuildHorde.xml
+++ b/Engine/Source/Programs/Horde/BuildHorde.xml
@@ -71,9 +71,9 @@
 			<!-- Remove the container now that files are copied -->
 			<Docker Arguments="rm horde-installer-container"/>
 	
-			<!-- Copy the Horde documentation to the dashboard public folder -->
-			<Copy From="$(RootDir)/Engine/Restricted/NotForLicensees/Source/Programs/Horde/Docs/..." To="$(StagingDir)/Dashboard/public/documentation/Docs/..."/>		
-			<Copy From="$(RootDir)/Engine/Restricted/NotForLicensees/Source/Programs/Horde/README.md" To="$(StagingDir)/Dashboard/public/documentation/README.md"/>
+			<!-- Copy the Horde documentation to the root folder -->
+			<Copy From="$(RootDir)/Engine/Restricted/NotForLicensees/Source/Programs/Horde/Docs/..." To="$(StagingDir)/Dashboard/documentation/Docs/..."/>		
+			<Copy From="$(RootDir)/Engine/Restricted/NotForLicensees/Source/Programs/Horde/README.md" To="$(StagingDir)/Dashboard/documentation/README.md"/>
 			
 			<Tag Files="$(StagingDir)/Dashboard/..." With="#DashboardOutputFiles"/>
 		</Node>
