Horde: Move utility commands into a standalone tool, so we can script processes from the command line.

diff --git a/Engine/Source/Programs/Horde/Horde.Agent/Commands/Bundles/BundleCreate.cs b/Engine/Source/Programs/Horde/Horde.Agent/Commands/Bundles/BundleCreate.cs
deleted file mode 100644
diff --git a/Engine/Source/Programs/Horde/Horde.Agent/Commands/Bundles/BundleDump.cs b/Engine/Source/Programs/Horde/Horde.Agent/Commands/Bundles/BundleDump.cs
deleted file mode 100644
diff --git a/Engine/Source/Programs/Horde/Horde.Agent/Commands/Bundles/BundleExtract.cs b/Engine/Source/Programs/Horde/Horde.Agent/Commands/Bundles/BundleExtract.cs
deleted file mode 100644
diff --git a/Engine/Source/Programs/Horde/Horde.Agent/Commands/Vcs/VcsBase.cs b/Engine/Source/Programs/Horde/Horde.Agent/Commands/Vcs/VcsBase.cs
deleted file mode 100644
diff --git a/Engine/Source/Programs/Horde/Horde.Agent/Commands/Vcs/VcsBranch.cs b/Engine/Source/Programs/Horde/Horde.Agent/Commands/Vcs/VcsBranch.cs
deleted file mode 100644
diff --git a/Engine/Source/Programs/Horde/Horde.Agent/Commands/Vcs/VcsCheckout.cs b/Engine/Source/Programs/Horde/Horde.Agent/Commands/Vcs/VcsCheckout.cs
deleted file mode 100644
diff --git a/Engine/Source/Programs/Horde/Horde.Agent/Commands/Vcs/VcsClone.cs b/Engine/Source/Programs/Horde/Horde.Agent/Commands/Vcs/VcsClone.cs
deleted file mode 100644
diff --git a/Engine/Source/Programs/Horde/Horde.Agent/Commands/Vcs/VcsCommit.cs b/Engine/Source/Programs/Horde/Horde.Agent/Commands/Vcs/VcsCommit.cs
deleted file mode 100644
diff --git a/Engine/Source/Programs/Horde/Horde.Agent/Commands/Vcs/VcsInit.cs b/Engine/Source/Programs/Horde/Horde.Agent/Commands/Vcs/VcsInit.cs
deleted file mode 100644
diff --git a/Engine/Source/Programs/Horde/Horde.Agent/Commands/Vcs/VcsLog.cs b/Engine/Source/Programs/Horde/Horde.Agent/Commands/Vcs/VcsLog.cs
deleted file mode 100644
diff --git a/Engine/Source/Programs/Horde/Horde.Agent/Commands/Vcs/VcsStatus.cs b/Engine/Source/Programs/Horde/Horde.Agent/Commands/Vcs/VcsStatus.cs
deleted file mode 100644
diff --git a/Engine/Source/Programs/Horde/Horde.Agent/Commands/Workspace/WorkspaceBase.cs b/Engine/Source/Programs/Horde/Horde.Agent/Commands/Workspace/WorkspaceBase.cs
deleted file mode 100644
diff --git a/Engine/Source/Programs/Horde/Horde.Agent/Commands/Workspace/WorkspaceClean.cs b/Engine/Source/Programs/Horde/Horde.Agent/Commands/Workspace/WorkspaceClean.cs
deleted file mode 100644
diff --git a/Engine/Source/Programs/Horde/Horde.Agent/Commands/Workspace/WorkspaceClear.cs b/Engine/Source/Programs/Horde/Horde.Agent/Commands/Workspace/WorkspaceClear.cs
deleted file mode 100644
diff --git a/Engine/Source/Programs/Horde/Horde.Agent/Commands/Workspace/WorkspaceDump.cs b/Engine/Source/Programs/Horde/Horde.Agent/Commands/Workspace/WorkspaceDump.cs
deleted file mode 100644
diff --git a/Engine/Source/Programs/Horde/Horde.Agent/Commands/Workspace/WorkspacePopulateCache.cs b/Engine/Source/Programs/Horde/Horde.Agent/Commands/Workspace/WorkspacePopulateCache.cs
deleted file mode 100644
diff --git a/Engine/Source/Programs/Horde/Horde.Agent/Commands/Workspace/WorkspacePurgeCache.cs b/Engine/Source/Programs/Horde/Horde.Agent/Commands/Workspace/WorkspacePurgeCache.cs
deleted file mode 100644
diff --git a/Engine/Source/Programs/Horde/Horde.Agent/Commands/Workspace/WorkspaceRepair.cs b/Engine/Source/Programs/Horde/Horde.Agent/Commands/Workspace/WorkspaceRepair.cs
deleted file mode 100644
diff --git a/Engine/Source/Programs/Horde/Horde.Agent/Commands/Workspace/WorkspaceRevert.cs b/Engine/Source/Programs/Horde/Horde.Agent/Commands/Workspace/WorkspaceRevert.cs
deleted file mode 100644
diff --git a/Engine/Source/Programs/Horde/Horde.Agent/Commands/Workspace/WorkspaceSetup.cs b/Engine/Source/Programs/Horde/Horde.Agent/Commands/Workspace/WorkspaceSetup.cs
deleted file mode 100644
diff --git a/Engine/Source/Programs/Horde/Horde.Agent/Commands/Workspace/WorkspaceSizes.cs b/Engine/Source/Programs/Horde/Horde.Agent/Commands/Workspace/WorkspaceSizes.cs
deleted file mode 100644
diff --git a/Engine/Source/Programs/Horde/Horde.Agent/Commands/Workspace/WorkspaceStatus.cs b/Engine/Source/Programs/Horde/Horde.Agent/Commands/Workspace/WorkspaceStatus.cs
deleted file mode 100644
diff --git a/Engine/Source/Programs/Horde/Horde.Agent/Commands/Workspace/WorkspaceSync.cs b/Engine/Source/Programs/Horde/Horde.Agent/Commands/Workspace/WorkspaceSync.cs
deleted file mode 100644
diff --git a/Engine/Source/Programs/Horde/Horde/Commands/Bundles/BundleCreate.cs b/Engine/Source/Programs/Horde/Horde/Commands/Bundles/BundleCreate.cs
new file mode 100644
--- /dev/null
+++ b/Engine/Source/Programs/Horde/Horde/Commands/Bundles/BundleCreate.cs
@@ -0,0 +1,51 @@
+// Copyright Epic Games, Inc. All Rights Reserved.
+
+using System.Diagnostics;
+using System.Threading;
+using System.Threading.Tasks;
+using EpicGames.Core;
+using EpicGames.Horde.Storage;
+using EpicGames.Horde.Storage.Nodes;
+using Microsoft.Extensions.Logging;
+
+namespace Horde.Commands.Bundles
+{
+	[Command("bundle", "create", "Creates a bundle from a folder on the local hard drive")]
+	class BundleCreate : Command
+	{
+		[CommandLine("-Namespace=", Description = "Namespace to use for storage")]
+		public NamespaceId NamespaceId { get; set; } = new NamespaceId("default");
+
+		[CommandLine("-Ref=")]
+		public RefName RefName { get; set; } = new RefName("default-bundle");
+
+		[CommandLine("-InputDir=", Required = true)]
+		public DirectoryReference InputDir { get; set; } = null!;
+
+		readonly IStorageClientFactory _storageClientFactory;
+
+		public BundleCreate(IStorageClientFactory storageClientFactory)
+		{
+			_storageClientFactory = storageClientFactory;
+		}
+
+		public override async Task<int> ExecuteAsync(ILogger logger)
+		{
+			IStorageClient store = await _storageClientFactory.GetClientAsync(NamespaceId);
+
+			using TreeWriter writer = new TreeWriter(store, prefix: RefName.Text);
+
+			DirectoryNode node = new DirectoryNode(DirectoryFlags.None);
+
+			Stopwatch timer = Stopwatch.StartNew();
+
+			ChunkingOptions options = new ChunkingOptions();
+			await node.CopyFromDirectoryAsync(InputDir.ToDirectoryInfo(), options, writer, null, CancellationToken.None);
+
+			await writer.WriteAsync(RefName, node);
+
+			logger.LogInformation("Time: {Time}", timer.Elapsed.TotalSeconds);
+			return 0;
+		}
+	}
+}
diff --git a/Engine/Source/Programs/Horde/Horde/Commands/Bundles/BundleDump.cs b/Engine/Source/Programs/Horde/Horde/Commands/Bundles/BundleDump.cs
new file mode 100644
--- /dev/null
+++ b/Engine/Source/Programs/Horde/Horde/Commands/Bundles/BundleDump.cs
@@ -0,0 +1,159 @@
+// Copyright Epic Games, Inc. All Rights Reserved.
+
+using System;
+using System.Collections.Generic;
+using System.Reflection;
+using System.Threading.Tasks;
+using EpicGames.Core;
+using EpicGames.Horde.Storage;
+using Microsoft.Extensions.Logging;
+
+namespace Horde.Commands.Bundles
+{
+	[Command("bundle", "dump", "Dumps the contents of a bundle")]
+	internal class BundleDump : Command
+	{
+		[CommandLine("-Namespace=", Description = "Namespace to use for storage")]
+		public NamespaceId NamespaceId { get; set; } = new NamespaceId("default");
+
+		[CommandLine("-Ref=")]
+		public RefName? RefName { get; set; }
+
+		[CommandLine("-Blob=")]
+		public BlobLocator? BlobId { get; set; }
+
+		[CommandLine("-File=")]
+		public string? File { get; set; }
+
+		[CommandLine("-Verbose")]
+		public bool Verbose { get; set; }
+
+		readonly IStorageClientFactory _storageClientFactory;
+
+		public BundleDump(IStorageClientFactory storageClientFactory)
+		{
+			_storageClientFactory = storageClientFactory;
+		}
+
+		public override async Task<int> ExecuteAsync(ILogger logger)
+		{
+			Bundle bundle;
+			if (File != null)
+			{
+				FileReference Location = new FileReference(File);
+				byte[] bytes = await FileReference.ReadAllBytesAsync(Location);
+				bundle = new Bundle(new MemoryReader(bytes));
+				logger.LogInformation("Summary for blob {Location}", Location);
+			}
+			else if (BlobId != null)
+			{
+				IStorageClient store = await _storageClientFactory.GetClientAsync(NamespaceId);
+				bundle = await store.ReadBundleAsync(BlobId.Value);
+				logger.LogInformation("Summary for blob {BlobId}", BlobId.Value);
+			}
+			else if (RefName != null)
+			{
+				IStorageClient store = await _storageClientFactory.GetClientAsync(NamespaceId);
+
+				NodeHandle handle = await store.ReadRefTargetAsync(RefName.Value);
+				logger.LogInformation("Ref {RefName} -> {Locator}", RefName.Value, handle);
+
+				bundle = await store.ReadBundleAsync(handle.Locator.Blob);
+				logger.LogInformation("Summary for blob {BlobId}", handle.Locator.Blob);
+			}
+			else
+			{
+				logger.LogError("No blob specified - use -File/Blob/Ref=...");
+				return 1;
+			}
+
+			BundleHeader header = bundle.Header;
+
+			Dictionary<Guid, string> typeIdToName = new Dictionary<Guid, string>();
+
+			TreeReaderOptions options = new TreeReaderOptions();
+			foreach (Type type in options.Types)
+			{
+				TreeNodeAttribute? attribute = type.GetCustomAttribute<TreeNodeAttribute>();
+				if (attribute != null)
+				{
+					typeIdToName.Add(Guid.Parse(attribute.Guid), type.Name);
+				}
+			}
+
+			string[] types = new string[header.Types.Count];
+			for (int idx = 0; idx < header.Types.Count; idx++)
+			{
+				string? name;
+				if (!typeIdToName.TryGetValue(header.Types[idx].Guid, out name))
+				{
+					name = header.Types[idx].Guid.ToString();
+				}
+				types[idx] = $"{name} v{header.Types[idx].Version}";
+			}
+
+			int packetStart = 0;
+
+			logger.LogInformation("");
+			logger.LogInformation("Imports: {NumImports}", header.Imports.Count);
+
+			List<string> nodeNames = new List<string>();
+
+			int refIdx = 0;
+			foreach (BundleImport import in header.Imports)
+			{
+				logger.LogInformation("  From blob {BlobId} ({NumExports} nodes)", import.Locator, import.Exports.Count);
+				foreach (int exportIdx in import.Exports)
+				{
+					nodeNames.Add($"IMP {import.Locator}#{exportIdx}");
+					if (Verbose)
+					{
+						logger.LogInformation("    [{Index}] IMP {ExportIdx}", refIdx, exportIdx);
+					}
+					refIdx++;
+				}
+			}
+
+			logger.LogInformation("");
+			logger.LogInformation("Exports: {NumExports}", header.Exports.Count);
+
+			int packetIdx = 0;
+			int packetOffset = 0;
+			for(int exportIdx = 0; exportIdx < header.Exports.Count; exportIdx++)
+			{
+				BundleExport export = header.Exports[exportIdx];
+				logger.LogInformation("  [{Index}] EXP {ExportHash} (type: {Type}, length: {NumBytes:n0}, packet: {PacketIdx})", refIdx, export.Hash, types[export.TypeIdx], export.Length, packetIdx);
+
+				nodeNames.Add($"EXP {exportIdx}");
+				refIdx++;
+
+				if (Verbose)
+				{
+					for (int idx = 0; idx < export.References.Count; idx++)
+					{
+						int nodeIdx = export.References[idx];
+						logger.LogInformation("            REF {RefIdx,-3} -> {NodeIdx,-3} ({RefName})", idx, nodeIdx, nodeNames[nodeIdx]);
+					}
+				}
+
+				packetOffset += export.Length;
+				if(packetOffset >= header.Packets[packetIdx].DecodedLength)
+				{
+					packetIdx++;
+					packetOffset = 0;
+				}
+			}
+
+			logger.LogInformation("");
+			logger.LogInformation("Packets: {NumPackets}", header.Packets.Count);
+			for (int idx = 0; idx < header.Packets.Count; idx++)
+			{
+				BundlePacket packet = header.Packets[idx];
+				logger.LogInformation("  PKT {Idx} (file offset: {Offset:n0}, encoded: {EncodedLength:n0}, decoded: {DecodedLength:n0}, ratio: {Ratio}%)", idx, packetStart, packet.EncodedLength, packet.DecodedLength, (int)(packet.EncodedLength * 100) / packet.DecodedLength);
+				packetStart += packet.EncodedLength;
+			}
+
+			return 0;
+		}
+	}
+}
diff --git a/Engine/Source/Programs/Horde/Horde/Commands/Bundles/BundleExtract.cs b/Engine/Source/Programs/Horde/Horde/Commands/Bundles/BundleExtract.cs
new file mode 100644
--- /dev/null
+++ b/Engine/Source/Programs/Horde/Horde/Commands/Bundles/BundleExtract.cs
@@ -0,0 +1,49 @@
+// Copyright Epic Games, Inc. All Rights Reserved.
+
+using System.Diagnostics;
+using System.Threading;
+using System.Threading.Tasks;
+using EpicGames.Core;
+using EpicGames.Horde.Storage;
+using EpicGames.Horde.Storage.Nodes;
+using Microsoft.Extensions.Caching.Memory;
+using Microsoft.Extensions.Logging;
+
+namespace Horde.Commands.Bundles
+{
+	[Command("bundle", "extract", "Extracts data from a bundle to the local hard drive")]
+	internal class BundleExtract : Command
+	{
+		[CommandLine("-Namespace=", Description = "Namespace to use for storage")]
+		public NamespaceId NamespaceId { get; set; } = new NamespaceId("default");
+
+		[CommandLine("-Ref=")]
+		public RefName RefName { get; set; } = new RefName("default-bundle");
+
+		[CommandLine("-OutputDir=", Required = true)]
+		public DirectoryReference OutputDir { get; set; } = null!;
+
+		readonly IStorageClientFactory _storageClientFactory;
+
+		public BundleExtract(IStorageClientFactory storageClientFactory)
+		{
+			_storageClientFactory = storageClientFactory;
+		}
+
+		public override async Task<int> ExecuteAsync(ILogger logger)
+		{
+			IStorageClient store = await _storageClientFactory.GetClientAsync(NamespaceId);
+
+			using MemoryCache cache = new MemoryCache(new MemoryCacheOptions());
+			TreeReader reader = new TreeReader(store, cache, logger);
+
+			Stopwatch timer = Stopwatch.StartNew();
+
+			DirectoryNode node = await reader.ReadNodeAsync<DirectoryNode>(RefName);
+			await node.CopyToDirectoryAsync(reader, OutputDir.ToDirectoryInfo(), logger, CancellationToken.None);
+
+			logger.LogInformation("Elapsed: {Time}s", timer.Elapsed.TotalSeconds);
+			return 0;
+		}
+	}
+}
diff --git a/Engine/Source/Programs/Horde/Horde/Commands/Vcs/VcsBase.cs b/Engine/Source/Programs/Horde/Horde/Commands/Vcs/VcsBase.cs
new file mode 100644
--- /dev/null
+++ b/Engine/Source/Programs/Horde/Horde/Commands/Vcs/VcsBase.cs
@@ -0,0 +1,366 @@
+// Copyright Epic Games, Inc. All Rights Reserved.
+
+using System;
+using System.Collections.Generic;
+using System.Diagnostics;
+using System.IO;
+using System.Linq;
+using System.Threading.Tasks;
+using EpicGames.Core;
+using EpicGames.Horde.Storage;
+using EpicGames.Horde.Storage.Nodes;
+using Microsoft.CodeAnalysis;
+using Microsoft.Extensions.Logging;
+
+namespace Horde.Commands.Vcs
+{
+	abstract class VcsBase : Command
+	{
+		[CommandLine("-BaseDir=")]
+		public DirectoryReference? BaseDir { get; set; }
+
+		[CommandLine("-Namespace=")]
+		public NamespaceId NamespaceId { get; set; } = new NamespaceId("horde-perforce");
+
+		protected class FileState
+		{
+			public IoHash Hash { get; set; }
+			public long Length { get; private set; }
+			public long LastModifiedTimeUtc { get; private set; }
+
+			public FileState(FileInfo fileInfo, IoHash hash = default)
+			{
+				Update(fileInfo);
+				Hash = hash;
+			}
+
+			public FileState(FileState other)
+			{
+				Length = other.Length;
+				LastModifiedTimeUtc = other.LastModifiedTimeUtc;
+			}
+
+			public FileState(IMemoryReader reader)
+			{
+				Hash = reader.ReadIoHash();
+				Length = reader.ReadInt64();
+				LastModifiedTimeUtc = reader.ReadInt64();
+			}
+
+			public bool Modified(FileInfo fileInfo) => Length != fileInfo.Length || LastModifiedTimeUtc != fileInfo.LastWriteTimeUtc.Ticks;
+
+			public void Update(FileInfo fileInfo)
+			{
+				Length = fileInfo.Length;
+				LastModifiedTimeUtc = fileInfo.LastWriteTimeUtc.Ticks;
+			}
+
+			public void Write(IMemoryWriter writer)
+			{
+				Debug.Assert(Hash != IoHash.Zero);
+				writer.WriteIoHash(Hash);
+				writer.WriteInt64(Length);
+				writer.WriteInt64(LastModifiedTimeUtc);
+			}
+		}
+
+		protected class DirectoryState
+		{
+			public IoHash Hash { get; set; }
+			public Dictionary<Utf8String, DirectoryState> Directories { get; } = new Dictionary<Utf8String, DirectoryState>(Utf8StringComparer.OrdinalIgnoreCase);
+			public Dictionary<Utf8String, FileState> Files { get; } = new Dictionary<Utf8String, FileState>(Utf8StringComparer.OrdinalIgnoreCase);
+
+			public DirectoryState()
+			{
+				Directories = new Dictionary<Utf8String, DirectoryState>(Utf8StringComparer.OrdinalIgnoreCase);
+				Files = new Dictionary<Utf8String, FileState>(Utf8StringComparer.OrdinalIgnoreCase);
+			}
+
+			public DirectoryState(IMemoryReader reader)
+			{
+				Hash = reader.ReadIoHash();
+
+				int numDirectories = reader.ReadInt32();
+				Directories = new Dictionary<Utf8String, DirectoryState>(numDirectories, Utf8StringComparer.OrdinalIgnoreCase);
+
+				for (int idx = 0; idx < numDirectories; idx++)
+				{
+					Utf8String name = reader.ReadUtf8String();
+					Directories[name] = new DirectoryState(reader);
+				}
+
+				int numFiles = reader.ReadInt32();
+				Files = new Dictionary<Utf8String, FileState>(numFiles, Utf8StringComparer.OrdinalIgnoreCase);
+
+				for (int idx = 0; idx < numFiles; idx++)
+				{
+					Utf8String name = reader.ReadUtf8String();
+					Files[name] = new FileState(reader);
+				}
+			}
+
+			public void Write(IMemoryWriter writer)
+			{
+				writer.WriteIoHash(Hash);
+
+				writer.WriteInt32(Directories.Count);
+				foreach ((Utf8String name, DirectoryState state) in Directories)
+				{
+					writer.WriteUtf8String(name);
+					state.Write(writer);
+				}
+
+				writer.WriteInt32(Files.Count);
+				foreach ((Utf8String name, FileState state) in Files)
+				{
+					writer.WriteUtf8String(name);
+					state.Write(writer);
+				}
+			}
+		}
+
+		protected class WorkspaceState
+		{
+			public RefName Branch { get; set; }
+			public int Change { get; set; }
+			public DirectoryState Tree { get; set; }
+
+			public WorkspaceState() 
+				: this(new RefName("main"), 0, new DirectoryState())
+			{
+			}
+
+			public WorkspaceState(RefName branch, int change, DirectoryState tree)
+			{
+				Branch = branch;
+				Change = change;
+				Tree = tree;
+			}
+
+			public WorkspaceState(IMemoryReader reader)
+			{
+				Branch = new RefName(reader.ReadUtf8String());
+				Change = reader.ReadInt32();
+				Tree = new DirectoryState(reader);
+			}
+
+			public void Write(IMemoryWriter writer)
+			{
+				writer.WriteUtf8String(Branch.Text);
+				writer.WriteInt32(Change);
+				Tree.Write(writer);
+			}
+		}
+
+		const string DataDir = ".horde";
+		const string StateFileName = "state.dat";
+
+		readonly IStorageClientFactory _storageClientFactory;
+
+		protected VcsBase(IStorageClientFactory storageClientFactory)
+		{
+			_storageClientFactory = storageClientFactory;
+		}
+
+		protected ValueTask<IStorageClient> GetStorageClientAsync() => _storageClientFactory.GetClientAsync(NamespaceId);
+
+		protected static async Task<WorkspaceState> ReadStateAsync(DirectoryReference rootDir)
+		{
+			FileReference stateFile = FileReference.Combine(rootDir, DataDir, StateFileName);
+			byte[] data = await FileReference.ReadAllBytesAsync(stateFile);
+			return new WorkspaceState(new MemoryReader(data));
+		}
+
+		protected static async Task WriteStateAsync(DirectoryReference rootDir, WorkspaceState state)
+		{
+			ByteArrayBuilder builder = new ByteArrayBuilder();
+			state.Write(builder);
+
+			FileReference stateFile = FileReference.Combine(rootDir, DataDir, StateFileName);
+			DirectoryReference.CreateDirectory(stateFile.Directory);
+
+			using (FileStream stream = FileReference.Open(stateFile, FileMode.Create, FileAccess.Write))
+			{
+				foreach (ReadOnlyMemory<byte> chunk in builder.AsSequence())
+				{
+					await stream.WriteAsync(chunk);
+				}
+			}
+		}
+
+		protected static async Task InitAsync(DirectoryReference rootDir)
+		{
+			await WriteStateAsync(rootDir, new WorkspaceState());
+		}
+
+		protected DirectoryReference FindRootDir()
+		{
+			if (BaseDir != null)
+			{
+				return BaseDir;
+			}
+
+			DirectoryReference startDir = DirectoryReference.GetCurrentDirectory();
+			for (DirectoryReference? currentDir = startDir; currentDir != null; currentDir = currentDir.ParentDirectory)
+			{
+				FileReference markerFile = FileReference.Combine(currentDir, DataDir, StateFileName);
+				if (FileReference.Exists(markerFile))
+				{
+					return currentDir;
+				}
+			}
+			throw new InvalidDataException($"No root directory found under {startDir}");
+		}
+
+		protected Task<DirectoryState> GetCurrentDirectoryState(DirectoryReference rootDir, DirectoryState? oldState) => GetCurrentDirectoryState(rootDir.ToDirectoryInfo(), oldState);
+
+		protected async Task<DirectoryState> GetCurrentDirectoryState(DirectoryInfo directoryInfo, DirectoryState? oldState)
+		{
+			List<DirectoryInfo> subDirectoryInfos = new List<DirectoryInfo>();
+			foreach (DirectoryInfo subDirectoryInfo in directoryInfo.EnumerateDirectories().Where(x => !x.Name.StartsWith(".", StringComparison.Ordinal)))
+			{
+				subDirectoryInfos.Add(subDirectoryInfo);
+			}
+
+			Task<DirectoryState>[] tasks = new Task<DirectoryState>[subDirectoryInfos.Count];
+			for (int idx = 0; idx < subDirectoryInfos.Count; idx++)
+			{
+				DirectoryInfo subDirectoryInfo = subDirectoryInfos[idx];
+
+				DirectoryState? prevSubDirectoryState = null;
+				if (oldState != null)
+				{
+					oldState.Directories.TryGetValue(subDirectoryInfo.Name, out prevSubDirectoryState);
+				}
+
+				tasks[idx] = Task.Run(() => GetCurrentDirectoryState(subDirectoryInfo, prevSubDirectoryState));
+			}
+
+			DirectoryState newState = new DirectoryState();
+			for (int idx = 0; idx < subDirectoryInfos.Count; idx++)
+			{
+				newState.Directories[subDirectoryInfos[idx].Name] = await tasks[idx];
+			}
+			foreach (FileInfo fileInfo in directoryInfo.EnumerateFiles())
+			{
+				FileState? fileState;
+				if (oldState == null || !oldState.Files.TryGetValue(fileInfo.Name, out fileState) || fileState.Modified(fileInfo))
+				{
+					fileState = new FileState(fileInfo);
+				}
+				newState.Files[fileInfo.Name] = fileState;
+			}
+
+			if (oldState != null && oldState.Directories.Count == newState.Directories.Count && oldState.Files.Count == newState.Files.Count)
+			{
+				if (newState.Directories.Values.All(x => x.Hash != IoHash.Zero) && newState.Files.Values.All(x => x.Hash != IoHash.Zero))
+				{
+					return oldState;
+				}
+			}
+	
+			return newState;
+		}
+
+		protected static void RemoveAddedFiles(DirectoryState oldState, DirectoryState newState)
+		{
+			List<(Utf8String, DirectoryState?, DirectoryState?)> directoryDeltas = EnumerableExtensions.Zip(oldState.Directories, newState.Directories).ToList();
+			foreach ((Utf8String name, DirectoryState? oldSubDirState, _) in directoryDeltas)
+			{
+				if (oldSubDirState == null)
+				{
+					newState.Directories.Remove(name);
+				}
+			}
+
+			List<(Utf8String, FileState?, FileState?)> fileDeltas = EnumerableExtensions.Zip(oldState.Files, newState.Files).ToList();
+			foreach ((Utf8String name, FileState? oldFileState, _) in fileDeltas)
+			{
+				if (oldFileState == null)
+				{
+					newState.Files.Remove(name);
+				}
+			}
+		}
+
+		protected static DirectoryState DedupTrees(DirectoryState oldState, DirectoryState newState)
+		{
+			if (oldState == newState)
+			{
+				return newState;
+			}
+
+			foreach ((_, DirectoryState? oldSubDirState, DirectoryState? newSubDirState) in EnumerableExtensions.Zip(oldState.Directories, newState.Directories))
+			{
+				if (oldSubDirState == null || newSubDirState == null || oldSubDirState != DedupTrees(oldSubDirState, newSubDirState))
+				{
+					return newState;
+				}
+			}
+
+			foreach ((_, FileState? oldFileState, FileState? newFileState) in EnumerableExtensions.Zip(oldState.Files, newState.Files))
+			{
+				if (oldFileState != newFileState)
+				{
+					return newState;
+				}
+			}
+
+			return oldState;
+		}
+
+		protected static void PrintDelta(DirectoryState oldState, DirectoryState newState, ILogger logger) => PrintDelta("", oldState, newState, logger);
+
+		static void PrintDelta(string prefix, DirectoryState oldState, DirectoryState newState, ILogger logger)
+		{
+			foreach ((Utf8String name, DirectoryState? oldSubDirState, DirectoryState? newSubDirState) in EnumerableExtensions.Zip(oldState.Directories, newState.Directories))
+			{
+				if (oldSubDirState == null)
+				{
+					logger.LogInformation("{Path} (added)", $"{prefix}{name}/");
+				}
+				else if (newSubDirState == null)
+				{
+					logger.LogInformation("{Path} (removed)", $"{prefix}{name}/");
+				}
+				else if (oldSubDirState != newSubDirState)
+				{
+					PrintDelta($"{prefix}{name}/", oldSubDirState, newSubDirState, logger);
+				}
+			}
+
+			foreach ((Utf8String name, FileState? oldFileState, FileState? newFileState) in EnumerableExtensions.Zip(oldState.Files, newState.Files))
+			{
+				if (oldFileState == null)
+				{
+					logger.LogInformation("{Path} (added)", $"{prefix}{name}");
+				}
+				else if (newFileState == null)
+				{
+					logger.LogInformation("{Path} (deleted)", $"{prefix}{name}");
+				}
+				else if (oldFileState != newFileState)
+				{
+					logger.LogInformation("{Path} (modified)", $"{prefix}{name}");
+				}
+			}
+		}
+
+		protected static async Task<CommitNode?> GetCommitAsync(TreeReader reader, RefName branchName, int change = 0)
+		{
+			CommitNode tip = await reader.ReadNodeAsync<CommitNode>(branchName);
+			if (change != 0)
+			{
+				while (tip.Number != change)
+				{
+					if (tip.Parent == null || tip.Number < change)
+					{
+						return null;
+					}
+					tip = await tip.Parent.ExpandAsync(reader);
+				}
+			}
+			return tip;
+		}
+	}
+}
diff --git a/Engine/Source/Programs/Horde/Horde/Commands/Vcs/VcsBranch.cs b/Engine/Source/Programs/Horde/Horde/Commands/Vcs/VcsBranch.cs
new file mode 100644
--- /dev/null
+++ b/Engine/Source/Programs/Horde/Horde/Commands/Vcs/VcsBranch.cs
@@ -0,0 +1,45 @@
+// Copyright Epic Games, Inc. All Rights Reserved.
+
+using System.Threading.Tasks;
+using EpicGames.Core;
+using EpicGames.Horde.Storage;
+using Microsoft.Extensions.Logging;
+
+namespace Horde.Commands.Vcs
+{
+	[Command("vcs", "branch", "Switch to a new branch")]
+	class VcsBranch : VcsBase
+	{
+		[CommandLine(Prefix = "-Name=", Required = true)]
+		public string Name { get; set; } = "";
+
+		public VcsBranch(IStorageClientFactory storageClientFactory)
+			: base(storageClientFactory)
+		{
+		}
+
+		public override async Task<int> ExecuteAsync(ILogger logger)
+		{
+			DirectoryReference rootDir = FindRootDir();
+
+			WorkspaceState workspaceState = await ReadStateAsync(rootDir);
+
+			IStorageClient store = await GetStorageClientAsync();
+
+			RefName branchName = new RefName(Name);
+			if (await store.HasRefAsync(branchName))
+			{
+				logger.LogError("Branch {BranchName} already exists - use checkout instead.", branchName);
+				return 1;
+			}
+
+			logger.LogInformation("Starting work in new branch {BranchName}", branchName);
+
+			workspaceState.Branch = new RefName(Name);
+			workspaceState.Tree = new DirectoryState();
+			await WriteStateAsync(rootDir, workspaceState);
+
+			return 0;
+		}
+	}
+}
diff --git a/Engine/Source/Programs/Horde/Horde/Commands/Vcs/VcsCheckout.cs b/Engine/Source/Programs/Horde/Horde/Commands/Vcs/VcsCheckout.cs
new file mode 100644
--- /dev/null
+++ b/Engine/Source/Programs/Horde/Horde/Commands/Vcs/VcsCheckout.cs
@@ -0,0 +1,133 @@
+// Copyright Epic Games, Inc. All Rights Reserved.
+
+using System.IO;
+using System.Threading;
+using System.Threading.Tasks;
+using EpicGames.Core;
+using EpicGames.Horde.Storage;
+using EpicGames.Horde.Storage.Nodes;
+using Microsoft.Extensions.Caching.Memory;
+using Microsoft.Extensions.Logging;
+
+namespace Horde.Commands.Vcs
+{
+	[Command("vcs", "checkout", "Checkout a particular branch/change")]
+	class VcsCheckout : VcsBase
+	{
+		[CommandLine("-Branch")]
+		public string? Branch { get; set; }
+
+		[CommandLine("-Change=")]
+		public int Change { get; set; }
+
+		[CommandLine("-Clean")]
+		public bool Clean { get; set; }
+
+		[CommandLine("-Force")]
+		public bool Force { get; set; }
+
+		public VcsCheckout(IStorageClientFactory storageClientFactory)
+			: base(storageClientFactory)
+		{
+		}
+
+		public override async Task<int> ExecuteAsync(ILogger logger)
+		{
+			DirectoryReference rootDir = FindRootDir();
+
+			WorkspaceState workspaceState = await ReadStateAsync(rootDir);
+
+			DirectoryState oldState = workspaceState.Tree;
+			DirectoryState newState = await GetCurrentDirectoryState(rootDir, oldState);
+
+			if (!Clean)
+			{
+				RemoveAddedFiles(oldState, newState);
+				newState = DedupTrees(oldState, newState);
+			}
+
+			if (oldState != newState && !Force)
+			{
+				logger.LogInformation("Current workspace has modified files. Resolve before checking out a different change, or run with -Force.");
+				PrintDelta(oldState, newState, logger);
+				return 0;
+			}
+
+			RefName branchName = (Branch != null) ? new RefName(Branch) : workspaceState.Branch;
+
+			IStorageClient store = await GetStorageClientAsync();
+
+			using MemoryCache cache = new MemoryCache(new MemoryCacheOptions());
+			TreeReader reader = new TreeReader(store, cache, logger);
+
+			CommitNode? tip = await GetCommitAsync(reader, branchName, Change);
+			if (tip == null)
+			{
+				logger.LogError("Unable to find change {Change}", Change);
+				return 1;
+			}
+
+			workspaceState.Tree = await RealizeAsync(reader, tip.Contents, rootDir, newState, Clean, logger);
+			await WriteStateAsync(rootDir, workspaceState);
+
+			logger.LogInformation("Updated workspace to change {Number}", tip.Number);
+			return 0;
+		}
+
+		protected static async Task<DirectoryState> RealizeAsync(TreeReader reader, DirectoryNodeRef directoryRef, DirectoryReference dirPath, DirectoryState? directoryState, bool clean, ILogger logger)
+		{
+			DirectoryReference.CreateDirectory(dirPath);
+
+			DirectoryState newState = new DirectoryState();
+
+			DirectoryNode directoryNode = await directoryRef.ExpandAsync(reader);
+			foreach ((Utf8String name, DirectoryEntry? subDirEntry, DirectoryState? subDirState) in EnumerableExtensions.Zip(directoryNode.NameToDirectory, directoryState?.Directories))
+			{
+				DirectoryReference subDirPath = DirectoryReference.Combine(dirPath, name.ToString());
+				if (subDirEntry == null)
+				{
+					if (clean)
+					{
+						DirectoryReference.Delete(subDirPath);
+					}
+				}
+				else
+				{
+					newState.Directories[name] = await RealizeAsync(reader, subDirEntry, subDirPath, subDirState, clean, logger);
+				}
+			}
+
+			foreach ((Utf8String name, FileEntry? fileEntry, FileState? fileState) in EnumerableExtensions.Zip(directoryNode.NameToFile, directoryState?.Files))
+			{
+				FileReference filePath = FileReference.Combine(dirPath, name.ToString());
+				if (fileEntry == null)
+				{
+					if (clean)
+					{
+						FileReference.Delete(filePath);
+					}
+				}
+				else if (fileState == null || fileState.Hash != fileEntry.Hash)
+				{
+					newState.Files[name] = await CheckoutFileAsync(reader, fileEntry, filePath.ToFileInfo(), logger);
+				}
+				else
+				{
+					newState.Files[name] = fileState;
+				}
+			}
+
+			newState.Hash = directoryRef.Hash;
+			return newState;
+		}
+
+		static async Task<FileState> CheckoutFileAsync(TreeReader reader, FileEntry fileRef, FileInfo fileInfo, ILogger logger)
+		{
+			logger.LogInformation("Updating {File} to {Hash}", fileInfo, fileRef.Hash);
+			FileNode fileNode = await fileRef.ExpandAsync(reader);
+			await fileNode.CopyToFileAsync(reader, fileInfo, CancellationToken.None);
+			fileInfo.Refresh();
+			return new FileState(fileInfo, fileRef.Hash);
+		}
+	}
+}
diff --git a/Engine/Source/Programs/Horde/Horde/Commands/Vcs/VcsClone.cs b/Engine/Source/Programs/Horde/Horde/Commands/Vcs/VcsClone.cs
new file mode 100644
--- /dev/null
+++ b/Engine/Source/Programs/Horde/Horde/Commands/Vcs/VcsClone.cs
@@ -0,0 +1,48 @@
+// Copyright Epic Games, Inc. All Rights Reserved.
+
+using System.Threading.Tasks;
+using EpicGames.Core;
+using EpicGames.Horde.Storage;
+using EpicGames.Horde.Storage.Nodes;
+using Microsoft.Extensions.Caching.Memory;
+using Microsoft.Extensions.Logging;
+
+namespace Horde.Commands.Vcs
+{
+	[Command("vcs", "clone", "Initialize a directory for VCS-like operations")]
+	class VcsClone : VcsCheckout
+	{
+		public VcsClone(IStorageClientFactory storageClientFactory)
+			: base(storageClientFactory)
+		{
+		}
+
+		public override async Task<int> ExecuteAsync(ILogger logger)
+		{
+			BaseDir ??= DirectoryReference.GetCurrentDirectory();
+			await InitAsync(BaseDir);
+			logger.LogInformation("Initialized in {RootDir}", BaseDir);
+
+			WorkspaceState workspaceState = new WorkspaceState();
+
+			RefName branchName = new RefName(Branch ?? "ue5-main");
+
+			IStorageClient store = await GetStorageClientAsync();
+
+			using MemoryCache cache = new MemoryCache(new MemoryCacheOptions());
+			TreeReader reader = new TreeReader(store, cache, logger);
+
+			CommitNode? tip = await GetCommitAsync(reader, branchName, Change);
+			if (tip == null)
+			{
+				logger.LogError("Unable to find change {Change}", Change);
+				return 1;
+			}
+
+			workspaceState.Tree = await RealizeAsync(reader, tip.Contents, BaseDir, null, true, logger);
+			await WriteStateAsync(BaseDir, workspaceState);
+
+			return 0;
+		}
+	}
+}
diff --git a/Engine/Source/Programs/Horde/Horde/Commands/Vcs/VcsCommit.cs b/Engine/Source/Programs/Horde/Horde/Commands/Vcs/VcsCommit.cs
new file mode 100644
--- /dev/null
+++ b/Engine/Source/Programs/Horde/Horde/Commands/Vcs/VcsCommit.cs
@@ -0,0 +1,140 @@
+// Copyright Epic Games, Inc. All Rights Reserved.
+
+using System;
+using System.Collections.Generic;
+using System.IO;
+using System.Threading;
+using System.Threading.Tasks;
+using EpicGames.Core;
+using EpicGames.Horde.Storage;
+using EpicGames.Horde.Storage.Nodes;
+using Microsoft.Extensions.Caching.Memory;
+using Microsoft.Extensions.Logging;
+
+namespace Horde.Commands.Vcs
+{
+	[Command("vcs", "commit", "Commits data to the VCS store")]
+	class VcsCommit : VcsBase
+	{
+		[CommandLine("-User=")]
+		public string UserName { get; set; } = Environment.UserName;
+
+		[CommandLine("-Message=")]
+		public string? Message { get; set; }
+
+		public VcsCommit(IStorageClientFactory storageClientFactory)
+			: base(storageClientFactory)
+		{
+		}
+
+		public override async Task<int> ExecuteAsync(ILogger logger)
+		{
+			DirectoryReference rootDir = FindRootDir();
+
+			WorkspaceState workspaceState = await ReadStateAsync(rootDir);
+
+			DirectoryState oldState = workspaceState.Tree;
+			DirectoryState newState = await GetCurrentDirectoryState(rootDir, oldState);
+
+			if (oldState == newState)
+			{
+				logger.LogInformation("No files modified");
+				return 0;
+			}
+
+			if (Message == null)
+			{
+				logger.LogError("Missing -Message=... argument");
+				return 1;
+			}
+
+			PrintDelta(oldState, newState, logger);
+
+			IStorageClient store = await GetStorageClientAsync();
+
+			using MemoryCache cache = new MemoryCache(new MemoryCacheOptions());
+			TreeReader reader = new TreeReader(store, cache, logger);
+
+			CommitNode? tip = await reader.TryReadNodeAsync<CommitNode>(workspaceState.Branch);
+			TreeNodeRef<CommitNode>? tipRef = (tip == null) ? null : new TreeNodeRef<CommitNode>(tip);
+
+			DirectoryNode rootNode;
+			if (tip == null)
+			{
+				rootNode = new DirectoryNode();
+			}
+			else
+			{
+				rootNode = await tip.Contents.ExpandCopyAsync(reader);
+			}
+
+			DirectoryNodeRef rootRef = new DirectoryNodeRef(rootNode);
+
+			List<(DirectoryNode, FileInfo, FileState)> files = new List<(DirectoryNode, FileInfo, FileState)>();
+			List<(DirectoryNodeRef, DirectoryState)> directories = new List<(DirectoryNodeRef, DirectoryState)>();
+			await UpdateTreeAsync(reader, rootRef, rootDir, oldState, newState, files, directories);
+
+			using TreeWriter writer = new TreeWriter(store);
+			await DirectoryNode.CopyFromDirectoryAsync(files.ConvertAll(x => (x.Item1, x.Item2)), new ChunkingOptions(), writer, null, CancellationToken.None);
+
+			CommitNode newTip;
+			if (tip == null)
+			{
+				newTip = new CommitNode(1, null, UserName, Message, DateTime.UtcNow, rootRef);
+			}
+			else
+			{
+				newTip = new CommitNode(tip.Number + 1, tipRef, UserName, Message, DateTime.UtcNow, rootRef);
+			}
+			await store.WriteNodeAsync(workspaceState.Branch, newTip);
+
+			foreach ((DirectoryNodeRef directoryRef, DirectoryState directoryState) in directories)
+			{
+				directoryState.Hash = directoryRef.Hash;
+			}
+
+			foreach ((DirectoryNode directoryNode, FileInfo fileInfo, FileState fileState) in files)
+			{
+				FileEntry fileEntry = directoryNode.GetFileEntry(fileInfo.Name);
+				fileState.Hash = fileEntry.Hash;
+			}
+
+			workspaceState.Change = newTip.Number;
+			workspaceState.Tree = newState;
+			await WriteStateAsync(rootDir, workspaceState);
+
+			logger.LogInformation("Commited in change {Number}", newTip.Number);
+			return 0;
+		}
+
+		private async Task UpdateTreeAsync(TreeReader reader, DirectoryNodeRef rootRef, DirectoryReference rootDir, DirectoryState? oldState, DirectoryState newState, List<(DirectoryNode, FileInfo, FileState)> files, List<(DirectoryNodeRef, DirectoryState)> directories)
+		{
+			directories.Add((rootRef, newState));
+
+			DirectoryNode root = await rootRef.ExpandAsync(reader);
+			foreach ((Utf8String name, DirectoryState? oldSubDirState, DirectoryState? newSubDirState) in EnumerableExtensions.Zip(oldState?.Directories, newState.Directories))
+			{
+				if (newSubDirState == null)
+				{
+					root.DeleteDirectory(name);
+				}
+				else if (oldSubDirState != newSubDirState)
+				{
+					await UpdateTreeAsync(reader, root.FindOrAddDirectoryEntry(name), DirectoryReference.Combine(rootDir, name.ToString()), oldSubDirState, newSubDirState, files, directories);
+				}
+			}
+
+			foreach ((Utf8String name, FileState? oldFileState, FileState? newFileState) in EnumerableExtensions.Zip(oldState?.Files, newState.Files))
+			{
+				if (newFileState == null)
+				{
+					root.DeleteFile(name);
+				}
+				else if (oldFileState != newFileState)
+				{
+					files.Add((root, FileReference.Combine(rootDir, name.ToString()).ToFileInfo(), newFileState));
+				}
+			}
+		}
+	}
+}
diff --git a/Engine/Source/Programs/Horde/Horde/Commands/Vcs/VcsInit.cs b/Engine/Source/Programs/Horde/Horde/Commands/Vcs/VcsInit.cs
new file mode 100644
--- /dev/null
+++ b/Engine/Source/Programs/Horde/Horde/Commands/Vcs/VcsInit.cs
@@ -0,0 +1,26 @@
+// Copyright Epic Games, Inc. All Rights Reserved.
+
+using System.Threading.Tasks;
+using EpicGames.Core;
+using EpicGames.Horde.Storage;
+using Microsoft.Extensions.Logging;
+
+namespace Horde.Commands.Vcs
+{
+	[Command("vcs", "init", "Initialize a directory for VCS-like operations")]
+	class VcsInit : VcsBase
+	{
+		public VcsInit(IStorageClientFactory storageClientFactory)
+			: base(storageClientFactory)
+		{
+		}
+
+		public override async Task<int> ExecuteAsync(ILogger logger)
+		{
+			BaseDir ??= DirectoryReference.GetCurrentDirectory();
+			await InitAsync(BaseDir);
+			logger.LogInformation("Initialized in {RootDir}", BaseDir);
+			return 0;
+		}
+	}
+}
diff --git a/Engine/Source/Programs/Horde/Horde/Commands/Vcs/VcsLog.cs b/Engine/Source/Programs/Horde/Horde/Commands/Vcs/VcsLog.cs
new file mode 100644
--- /dev/null
+++ b/Engine/Source/Programs/Horde/Horde/Commands/Vcs/VcsLog.cs
@@ -0,0 +1,61 @@
+// Copyright Epic Games, Inc. All Rights Reserved.
+
+using System.Collections.Generic;
+using System.Threading.Tasks;
+using EpicGames.Core;
+using EpicGames.Horde.Storage;
+using EpicGames.Horde.Storage.Nodes;
+using Microsoft.Extensions.Caching.Memory;
+using Microsoft.Extensions.Logging;
+
+namespace Horde.Commands.Vcs
+{
+	[Command("vcs", "log", "Print a history of commits")]
+	class VcsLog : VcsBase
+	{
+		[CommandLine("-Count")]
+		public int Count { get; set; } = 20;
+
+		public VcsLog(IStorageClientFactory storageClientFactory)
+			: base(storageClientFactory)
+		{
+		}
+
+		public override async Task<int> ExecuteAsync(ILogger logger)
+		{
+			DirectoryReference rootDir = FindRootDir();
+
+			WorkspaceState workspaceState = await ReadStateAsync(rootDir);
+
+			IStorageClient store = await GetStorageClientAsync();
+
+			using MemoryCache cache = new MemoryCache(new MemoryCacheOptions());
+			TreeReader reader = new TreeReader(store, cache, logger);
+
+			List<CommitNode> commits = new List<CommitNode>();
+
+			CommitNode? tip = await reader.TryReadNodeAsync<CommitNode>(workspaceState.Branch);
+			if (tip != null)
+			{
+				for (int idx = 0; idx < Count; idx++)
+				{
+					commits.Add(tip);
+
+					if (tip.Parent == null)
+					{
+						break;
+					}
+
+					tip = await tip.Parent.ExpandAsync(reader);
+				}
+			}
+
+			foreach (CommitNode commit in commits)
+			{
+				logger.LogInformation("Commit {Number}: {Message}", commit.Number, commit.Message);
+			}
+
+			return 0;
+		}
+	}
+}
diff --git a/Engine/Source/Programs/Horde/Horde/Commands/Vcs/VcsStatus.cs b/Engine/Source/Programs/Horde/Horde/Commands/Vcs/VcsStatus.cs
new file mode 100644
--- /dev/null
+++ b/Engine/Source/Programs/Horde/Horde/Commands/Vcs/VcsStatus.cs
@@ -0,0 +1,41 @@
+// Copyright Epic Games, Inc. All Rights Reserved.
+
+using System.Threading.Tasks;
+using EpicGames.Core;
+using EpicGames.Horde.Storage;
+using Microsoft.Extensions.Logging;
+
+namespace Horde.Commands.Vcs
+{
+	[Command("vcs", "status", "Find status of local files")]
+	class VcsStatus : VcsBase
+	{
+		public VcsStatus(IStorageClientFactory storageClientFactory)
+			: base(storageClientFactory)
+		{
+		}
+
+		public override async Task<int> ExecuteAsync(ILogger logger)
+		{
+			DirectoryReference rootDir = FindRootDir();
+
+			WorkspaceState workspaceState = await ReadStateAsync(rootDir);
+
+			DirectoryState oldState = workspaceState.Tree;
+			DirectoryState newState = await GetCurrentDirectoryState(rootDir, oldState);
+
+			logger.LogInformation("On branch {BranchName}", workspaceState.Branch);
+
+			if (oldState == newState)
+			{
+				logger.LogInformation("No files modified");
+			}
+			else
+			{
+				PrintDelta(oldState, newState, logger);
+			}
+
+			return 0;
+		}
+	}
+}
diff --git a/Engine/Source/Programs/Horde/Horde/Commands/Workspace/WorkspaceBase.cs b/Engine/Source/Programs/Horde/Horde/Commands/Workspace/WorkspaceBase.cs
new file mode 100644
--- /dev/null
+++ b/Engine/Source/Programs/Horde/Horde/Commands/Workspace/WorkspaceBase.cs
@@ -0,0 +1,153 @@
+// Copyright Epic Games, Inc. All Rights Reserved.
+
+using System;
+using System.Collections.Generic;
+using System.ComponentModel;
+using System.Linq;
+using System.Threading;
+using System.Threading.Tasks;
+using EpicGames.Core;
+using EpicGames.Perforce;
+using EpicGames.Perforce.Managed;
+using Microsoft.Extensions.Logging;
+
+namespace Horde.Commands.Workspace
+{
+	abstract class WorkspaceBase : Command
+	{
+		[CommandLine("-ServerAndPort")]
+		[Description("Specifies the Perforce server and port")]
+		protected string ServerAndPort { get; set; } = PerforceSettings.Default.ServerAndPort;
+
+		[CommandLine("-User")]
+		[Description("Specifies the Perforce username")]
+		protected string UserName { get; set; } = PerforceSettings.Default.UserName;
+
+		[CommandLine("-BaseDir", Required = true)]
+		[Description("Base directory to use for syncing workspaces")]
+		protected DirectoryReference BaseDir { get; set; } = null!;
+
+		[CommandLine("-Overwrite")]
+		[Description("")]
+		protected bool Overwrite { get; set; } = false;
+		
+		[CommandLine("-PreferNativeClient")]
+		[Description("Prefer to use native Perforce client (instead of launching separate p4 process)")]
+		protected bool PreferNativeClient { get; set; } = false;
+		
+		[CommandLine("-UseHaveTable")]
+		[Description("Use have-table for syncing")]
+		string UseHaveTable { get; set; } = "true";
+
+		public override void Configure(CommandLineArguments arguments, ILogger logger)
+		{
+			base.Configure(arguments, logger);
+
+			if(BaseDir == null)
+			{
+				for (DirectoryReference? parentDir = DirectoryReference.GetCurrentDirectory(); parentDir != null; parentDir = parentDir.ParentDirectory)
+				{
+					if (ManagedWorkspace.Exists(parentDir))
+					{
+						BaseDir = parentDir;
+						break;
+					}
+				}
+
+				if (BaseDir == null)
+				{
+					throw new FatalErrorException("Unable to find existing repository in current working tree. Specify -BaseDir=... explicitly.");
+				}
+			}
+		}
+
+		public override async Task<int> ExecuteAsync(ILogger logger)
+		{
+			PerforceSettings settings = new(ServerAndPort, UserName) { PreferNativeClient = PreferNativeClient };
+			using IPerforceConnection perforce = await PerforceConnection.CreateAsync(settings, logger);
+			InfoRecord info = await perforce.GetInfoAsync(InfoOptions.ShortOutput, CancellationToken.None);
+			bool useHaveTable = UseHaveTable.Equals("true", StringComparison.Ordinal);
+
+			ManagedWorkspace repo = await ManagedWorkspace.LoadOrCreateAsync(info.ClientHost!, BaseDir, Overwrite, useHaveTable, logger, CancellationToken.None);
+			await ExecuteAsync(perforce, repo, logger);
+			return 0;
+		}
+
+		protected abstract Task ExecuteAsync(IPerforceConnection perforce, ManagedWorkspace repo, ILogger logger);
+
+		protected static long ParseSize(string size)
+		{
+			long value;
+			if (size.EndsWith("gb", StringComparison.OrdinalIgnoreCase))
+			{
+				string sizeValue = size.Substring(0, size.Length - 2).TrimEnd();
+				if (Int64.TryParse(sizeValue, out value))
+				{
+					return value * (1024 * 1024 * 1024);
+				}
+			}
+			else if (size.EndsWith("mb", StringComparison.OrdinalIgnoreCase))
+			{
+				string sizeValue = size.Substring(0, size.Length - 2).TrimEnd();
+				if (Int64.TryParse(sizeValue, out value))
+				{
+					return value * (1024 * 1024);
+				}
+			}
+			else if (size.EndsWith("kb", StringComparison.OrdinalIgnoreCase))
+			{
+				string sizeValue = size.Substring(0, size.Length - 2).TrimEnd();
+				if (Int64.TryParse(sizeValue, out value))
+				{
+					return value * 1024;
+				}
+			}
+			else
+			{
+				if (Int64.TryParse(size, out value))
+				{
+					return value;
+				}
+			}
+			throw new FatalErrorException("Invalid size '{0}'", size);
+		}
+
+		protected static int ParseChangeNumber(string change)
+		{
+			int changeNumber;
+			if (Int32.TryParse(change, out changeNumber) && changeNumber > 0)
+			{
+				return changeNumber;
+			}
+			throw new FatalErrorException("Unable to parse change number from '{0}'", change);
+		}
+
+		protected static int ParseChangeNumberOrLatest(string change)
+		{
+			if (change.Equals("Latest", StringComparison.OrdinalIgnoreCase))
+			{
+				return -1;
+			}
+			else
+			{
+				return ParseChangeNumber(change);
+			}
+		}
+
+		protected static List<string> ExpandFilters(List<string> arguments)
+		{
+			List<string> filters = new List<string>();
+			foreach (string argument in arguments)
+			{
+				foreach (string singleArgument in argument.Split(';').Select(x => x.Trim()))
+				{
+					if (singleArgument.Length > 0)
+					{
+						filters.Add(singleArgument);
+					}
+				}
+			}
+			return filters;
+		}
+	}
+}
diff --git a/Engine/Source/Programs/Horde/Horde/Commands/Workspace/WorkspaceClean.cs b/Engine/Source/Programs/Horde/Horde/Commands/Workspace/WorkspaceClean.cs
new file mode 100644
--- /dev/null
+++ b/Engine/Source/Programs/Horde/Horde/Commands/Workspace/WorkspaceClean.cs
@@ -0,0 +1,25 @@
+// Copyright Epic Games, Inc. All Rights Reserved.
+
+using System.ComponentModel;
+using System.Threading;
+using System.Threading.Tasks;
+using EpicGames.Core;
+using EpicGames.Perforce;
+using EpicGames.Perforce.Managed;
+using Microsoft.Extensions.Logging;
+
+namespace Horde.Commands.Workspace
+{
+	[Command("workspace", "clean", "Cleans all modified files from the workspace.")]
+	class WorkspaceClean : WorkspaceBase
+	{
+		[CommandLine("-Incremental")]
+		[Description("Performs an incremental sync, without removing intermediates")]
+		public bool Incremental { get; set; } = false;
+
+		protected override Task ExecuteAsync(IPerforceConnection perforce, ManagedWorkspace repo, ILogger logger)
+		{
+			return repo.CleanAsync(!Incremental, CancellationToken.None);
+		}
+	}
+}
diff --git a/Engine/Source/Programs/Horde/Horde/Commands/Workspace/WorkspaceClear.cs b/Engine/Source/Programs/Horde/Horde/Commands/Workspace/WorkspaceClear.cs
new file mode 100644
--- /dev/null
+++ b/Engine/Source/Programs/Horde/Horde/Commands/Workspace/WorkspaceClear.cs
@@ -0,0 +1,20 @@
+// Copyright Epic Games, Inc. All Rights Reserved.
+
+using System.Threading;
+using System.Threading.Tasks;
+using EpicGames.Core;
+using EpicGames.Perforce;
+using EpicGames.Perforce.Managed;
+using Microsoft.Extensions.Logging;
+
+namespace Horde.Commands.Workspace
+{
+	[Command("workspace", "clear", "Empties the staging directory of any files, returning them to the cache")]
+	class WorkspaceClear : WorkspaceBase
+	{
+		protected override Task ExecuteAsync(IPerforceConnection perforce, ManagedWorkspace repo, ILogger logger)
+		{
+			return repo.ClearAsync(CancellationToken.None);
+		}
+	}
+}
diff --git a/Engine/Source/Programs/Horde/Horde/Commands/Workspace/WorkspaceDump.cs b/Engine/Source/Programs/Horde/Horde/Commands/Workspace/WorkspaceDump.cs
new file mode 100644
--- /dev/null
+++ b/Engine/Source/Programs/Horde/Horde/Commands/Workspace/WorkspaceDump.cs
@@ -0,0 +1,20 @@
+// Copyright Epic Games, Inc. All Rights Reserved.
+
+using System.Threading.Tasks;
+using EpicGames.Core;
+using EpicGames.Perforce;
+using EpicGames.Perforce.Managed;
+using Microsoft.Extensions.Logging;
+
+namespace Horde.Commands.Workspace
+{
+	[Command("workspace", "dump", "Dumps the contents of the repository to the log for analysis")]
+	class WorkspaceDump : WorkspaceBase
+	{
+		protected override Task ExecuteAsync(IPerforceConnection perforce, ManagedWorkspace repo, ILogger logger)
+		{
+			repo.Dump();
+			return Task.CompletedTask;
+		}
+	}
+}
diff --git a/Engine/Source/Programs/Horde/Horde/Commands/Workspace/WorkspacePopulateCache.cs b/Engine/Source/Programs/Horde/Horde/Commands/Workspace/WorkspacePopulateCache.cs
new file mode 100644
--- /dev/null
+++ b/Engine/Source/Programs/Horde/Horde/Commands/Workspace/WorkspacePopulateCache.cs
@@ -0,0 +1,50 @@
+// Copyright Epic Games, Inc. All Rights Reserved.
+
+using System;
+using System.Collections.Generic;
+using System.ComponentModel;
+using System.Threading;
+using System.Threading.Tasks;
+using EpicGames.Core;
+using EpicGames.Perforce;
+using EpicGames.Perforce.Managed;
+using Microsoft.Extensions.Logging;
+
+namespace Horde.Commands.Workspace
+{
+	[Command("workspace", "populatecache", "Populates the cache with the head revision of the given streams")]
+	class WorkspacePopulateCache : WorkspaceBase
+	{
+		[CommandLine("-ClientAndStream=")]
+		[Description("Specifies client and stream pairs, in the format Client:Stream")]
+		List<string> ClientAndStreamParams { get; set; } = new List<string>();
+
+		[CommandLine("-Filter=")]
+		[Description("Filters for the files to sync, in P4 syntax (eg. /Engine/...)")]
+		List<string> Filters { get; set; } = new List<string>();
+
+		[CommandLine("-FakeSync")]
+		[Description("Simulates the sync without actually fetching any files")]
+		bool FakeSync { get; set; } = false;
+
+		protected override async Task ExecuteAsync(IPerforceConnection perforce, ManagedWorkspace repo, ILogger logger)
+		{
+			List<string> expandedFilters = ExpandFilters(Filters);
+
+			List<PopulateRequest> populateRequests = new List<PopulateRequest>();
+			foreach (string clientAndStreamParam in ClientAndStreamParams)
+			{
+				int idx = clientAndStreamParam.IndexOf(':', StringComparison.Ordinal);
+				if (idx == -1)
+				{
+					throw new FatalErrorException("Expected -ClientAndStream=<ClientName>:<StreamName>");
+				}
+
+				using IPerforceConnection perforceClient = await perforce.WithClientAsync(clientAndStreamParam.Substring(0, idx));
+				populateRequests.Add(new PopulateRequest(perforceClient, clientAndStreamParam.Substring(idx + 1), expandedFilters));
+			}
+
+			await repo.PopulateAsync(populateRequests, FakeSync, CancellationToken.None);
+		}
+	}
+}
diff --git a/Engine/Source/Programs/Horde/Horde/Commands/Workspace/WorkspacePurgeCache.cs b/Engine/Source/Programs/Horde/Horde/Commands/Workspace/WorkspacePurgeCache.cs
new file mode 100644
--- /dev/null
+++ b/Engine/Source/Programs/Horde/Horde/Commands/Workspace/WorkspacePurgeCache.cs
@@ -0,0 +1,29 @@
+// Copyright Epic Games, Inc. All Rights Reserved.
+
+using System.Threading;
+using System.Threading.Tasks;
+using EpicGames.Core;
+using EpicGames.Perforce;
+using EpicGames.Perforce.Managed;
+using Microsoft.Extensions.Logging;
+
+namespace Horde.Commands.Workspace
+{
+	[Command("workspace", "purgecache", "Shrink the size of the cache to the given size")]
+	class WorkspacePurgeCache : WorkspaceBase
+	{
+		[CommandLine("-Size=")]
+		string? SizeParam { get; set; } = null;
+
+		protected override Task ExecuteAsync(IPerforceConnection perforce, ManagedWorkspace repo, ILogger logger)
+		{
+			long size = 0;
+			if (SizeParam != null)
+			{
+				size = ParseSize(SizeParam);
+			}
+
+			return repo.PurgeAsync(size, CancellationToken.None);
+		}
+	}
+}
diff --git a/Engine/Source/Programs/Horde/Horde/Commands/Workspace/WorkspaceRepair.cs b/Engine/Source/Programs/Horde/Horde/Commands/Workspace/WorkspaceRepair.cs
new file mode 100644
--- /dev/null
+++ b/Engine/Source/Programs/Horde/Horde/Commands/Workspace/WorkspaceRepair.cs
@@ -0,0 +1,20 @@
+// Copyright Epic Games, Inc. All Rights Reserved.
+
+using System.Threading;
+using System.Threading.Tasks;
+using EpicGames.Core;
+using EpicGames.Perforce;
+using EpicGames.Perforce.Managed;
+using Microsoft.Extensions.Logging;
+
+namespace Horde.Commands.Workspace
+{
+	[Command("workspace", "repaircache", "Checks the integrity of the cache, and removes any invalid files")]
+	class WorkspaceRepair : WorkspaceBase
+	{
+		protected override Task ExecuteAsync(IPerforceConnection perforce, ManagedWorkspace repo, ILogger logger)
+		{
+			return repo.RepairAsync(CancellationToken.None);
+		}
+	}
+}
diff --git a/Engine/Source/Programs/Horde/Horde/Commands/Workspace/WorkspaceRevert.cs b/Engine/Source/Programs/Horde/Horde/Commands/Workspace/WorkspaceRevert.cs
new file mode 100644
--- /dev/null
+++ b/Engine/Source/Programs/Horde/Horde/Commands/Workspace/WorkspaceRevert.cs
@@ -0,0 +1,26 @@
+// Copyright Epic Games, Inc. All Rights Reserved.
+
+using System.ComponentModel;
+using System.Threading;
+using System.Threading.Tasks;
+using EpicGames.Core;
+using EpicGames.Perforce;
+using EpicGames.Perforce.Managed;
+using Microsoft.Extensions.Logging;
+
+namespace Horde.Commands.Workspace
+{
+	[Command("workspace", "revert", "Revert all files that are open in the current workspace. Does not replace them with valid revisions.")]
+	class WorkspaceRevert : WorkspaceBase
+	{
+		[CommandLine("-Client=", Required = true)]
+		[Description("Client to revert all files for")]
+		string ClientName { get; set; } = null!;
+
+		protected override async Task ExecuteAsync(IPerforceConnection perforce, ManagedWorkspace repo, ILogger logger)
+		{
+			using IPerforceConnection perforceClient = await perforce.WithClientAsync(ClientName);
+			await repo.RevertAsync(perforceClient, CancellationToken.None);
+		}
+	}
+}
diff --git a/Engine/Source/Programs/Horde/Horde/Commands/Workspace/WorkspaceSetup.cs b/Engine/Source/Programs/Horde/Horde/Commands/Workspace/WorkspaceSetup.cs
new file mode 100644
--- /dev/null
+++ b/Engine/Source/Programs/Horde/Horde/Commands/Workspace/WorkspaceSetup.cs
@@ -0,0 +1,30 @@
+// Copyright Epic Games, Inc. All Rights Reserved.
+
+using System.ComponentModel;
+using System.Threading;
+using System.Threading.Tasks;
+using EpicGames.Core;
+using EpicGames.Perforce;
+using EpicGames.Perforce.Managed;
+using Microsoft.Extensions.Logging;
+
+namespace Horde.Commands.Workspace
+{
+	[Command("workspace", "setup", "Creates or updates a client to use a given stream")]
+	class WorkspaceSetup : WorkspaceBase
+	{
+		[CommandLine("-Client=", Required = true)]
+		[Description("Name of the client to create")]
+		string ClientName { get; set; } = null!;
+
+		[CommandLine("-Stream=", Required = true)]
+		[Description("Name of the stream to configure")]
+		string StreamName { get; set; } = null!;
+
+		protected override async Task ExecuteAsync(IPerforceConnection perforce, ManagedWorkspace repo, ILogger logger)
+		{
+			using IPerforceConnection perforceClient = await perforce.WithClientAsync(ClientName);
+			await repo.SetupAsync(perforceClient, StreamName, CancellationToken.None);
+		}
+	}
+}
diff --git a/Engine/Source/Programs/Horde/Horde/Commands/Workspace/WorkspaceSizes.cs b/Engine/Source/Programs/Horde/Horde/Commands/Workspace/WorkspaceSizes.cs
new file mode 100644
--- /dev/null
+++ b/Engine/Source/Programs/Horde/Horde/Commands/Workspace/WorkspaceSizes.cs
@@ -0,0 +1,35 @@
+// Copyright Epic Games, Inc. All Rights Reserved.
+
+using System.Collections.Generic;
+using System.ComponentModel;
+using System.Threading;
+using System.Threading.Tasks;
+using EpicGames.Core;
+using EpicGames.Perforce;
+using EpicGames.Perforce.Managed;
+using Microsoft.Extensions.Logging;
+
+namespace Horde.Commands.Workspace
+{
+	[Command("workspace", "sizes", "Gathers stats for which streams take the most amount of space in the cache for the given configuration")]
+	class WorkspaceSizes : WorkspaceBase
+	{
+		[CommandLine("-TempClient=", Required = true)]
+		[Description("Name of a temporary client to switch between streams gathering metadata. Will be created if it does not exist.")]
+		string TempClientName { get; set; } = null!;
+
+		[CommandLine("-Stream=")]
+		[Description("Streams that should be included in the output")]
+		List<string> StreamNames { get; set; } = new List<string>();
+
+		[CommandLine("-Filter=")]
+		[Description("Filters for the files to sync, in P4 syntax (eg. /Engine/...)")]
+		List<string> Filters { get; set; } = new List<string>();
+
+		protected override async Task ExecuteAsync(IPerforceConnection perforce, ManagedWorkspace repo, ILogger logger)
+		{
+			using IPerforceConnection perforceClient = await perforce.WithClientAsync(TempClientName);
+			await repo.StatsAsync(perforceClient, StreamNames, Filters, CancellationToken.None);
+		}
+	}
+}
diff --git a/Engine/Source/Programs/Horde/Horde/Commands/Workspace/WorkspaceStatus.cs b/Engine/Source/Programs/Horde/Horde/Commands/Workspace/WorkspaceStatus.cs
new file mode 100644
--- /dev/null
+++ b/Engine/Source/Programs/Horde/Horde/Commands/Workspace/WorkspaceStatus.cs
@@ -0,0 +1,20 @@
+// Copyright Epic Games, Inc. All Rights Reserved.
+
+using System.Threading.Tasks;
+using EpicGames.Core;
+using EpicGames.Perforce;
+using EpicGames.Perforce.Managed;
+using Microsoft.Extensions.Logging;
+
+namespace Horde.Commands.Workspace
+{
+	[Command("workspace", "status", "Prints information about the state of the cache and workspace")]
+	class WorkspaceStatus : WorkspaceBase
+	{
+		protected override Task ExecuteAsync(IPerforceConnection perforce, ManagedWorkspace repo, ILogger logger)
+		{
+			repo.Status();
+			return Task.CompletedTask;
+		}
+	}
+}
diff --git a/Engine/Source/Programs/Horde/Horde/Commands/Workspace/WorkspaceSync.cs b/Engine/Source/Programs/Horde/Horde/Commands/Workspace/WorkspaceSync.cs
new file mode 100644
--- /dev/null
+++ b/Engine/Source/Programs/Horde/Horde/Commands/Workspace/WorkspaceSync.cs
@@ -0,0 +1,60 @@
+// Copyright Epic Games, Inc. All Rights Reserved.
+
+using System.ComponentModel;
+using EpicGames.Core;
+using EpicGames.Perforce;
+using EpicGames.Perforce.Managed;
+using Microsoft.Extensions.Logging;
+
+namespace Horde.Commands.Workspace
+{
+	[Command("workspace", "sync", "Syncs the files for a particular stream and changelist")]
+	class WorkspaceSync : WorkspaceBase
+	{
+		[CommandLine("-Client=", Required = true)]
+		[Description("Name of the client to sync. Will be created if it does not exist.")]
+		public string ClientName { get; set; } = null!;
+
+		[CommandLine("-Stream=", Required = true)]
+		[Description("The stream to sync")]
+		public string StreamName { get; set; } = null!;
+
+		[CommandLine("-Change=")]
+		[Description("The change to sync. May be a changelist number, or 'Latest'")]
+		public string Change { get; set; } = "Latest";
+
+		[CommandLine("-Preflight=")]
+		[Description("The change to unshelve into the workspace")]
+		public int PreflightChange { get; set; } = -1;
+
+		[CommandLine]
+		[Description("Optional path to a cache file used to store workspace metadata. Using a location on a network share allows multiple machines syncing the same CL to only query Perforce state once.")]
+		FileReference? CacheFile { get; set; } = null;
+
+		[CommandLine("-Filter=")]
+		[Description("Filters for the files to sync, in P4 syntax (eg. /Engine/...)")]
+		public List<string> Filters { get; set; } = new List<string>();
+
+		[CommandLine("-Incremental")]
+		[Description("Performs an incremental sync, without removing intermediates")]
+		public bool IncrementalSync { get; set; } = false;
+
+		[CommandLine("-FakeSync")]
+		[Description("Simulates the sync without actually fetching any files")]
+		public bool FakeSync { get; set; } = false;
+
+		protected override async Task ExecuteAsync(IPerforceConnection perforce, ManagedWorkspace repo, ILogger logger)
+		{
+			int changeNumber = ParseChangeNumberOrLatest(Change);
+			List<string> expandedFilters = ExpandFilters(Filters);
+
+			using IPerforceConnection perforceClient = await perforce.WithClientAsync(ClientName);
+			await repo.SyncAsync(perforceClient, StreamName, changeNumber, expandedFilters, !IncrementalSync, FakeSync, CacheFile, CancellationToken.None);
+
+			if (PreflightChange != -1)
+			{
+				await repo.UnshelveAsync(perforceClient, PreflightChange, CancellationToken.None);
+			}
+		}
+	}
+}
diff --git a/Engine/Source/Programs/Horde/Horde/Horde.csproj b/Engine/Source/Programs/Horde/Horde/Horde.csproj
new file mode 100644
--- /dev/null
+++ b/Engine/Source/Programs/Horde/Horde/Horde.csproj
@@ -0,0 +1,34 @@
+<Project Sdk="Microsoft.NET.Sdk">
+
+  <PropertyGroup>
+    <OutputType>Exe</OutputType>
+    <TargetFramework>net6.0</TargetFramework>
+    <ImplicitUsings>enable</ImplicitUsings>
+    <Nullable>enable</Nullable>
+  </PropertyGroup>
+
+  <ItemGroup>
+    <ProjectReference Include="..\..\Shared\EpicGames.Horde\EpicGames.Horde.csproj" />
+    <ProjectReference Include="..\..\Shared\EpicGames.OIDC\EpicGames.OIDC.csproj" />
+    <ProjectReference Include="..\..\Shared\EpicGames.Perforce.Managed\EpicGames.Perforce.Managed.csproj" />
+    <ProjectReference Include="..\..\Shared\EpicGames.Perforce\EpicGames.Perforce.csproj" />
+  </ItemGroup>
+
+  <ItemGroup>
+    <PackageReference Include="Serilog" Version="2.12.0" />
+    <PackageReference Include="Serilog.Extensions.Hosting" Version="5.0.1" />
+    <PackageReference Include="Serilog.Extensions.Logging" Version="3.1.0" />
+    <PackageReference Include="Serilog.Settings.Configuration" Version="3.4.0" />
+    <PackageReference Include="Serilog.Sinks.Console" Version="4.1.0" />
+    <PackageReference Include="Serilog.Sinks.File" Version="5.0.0" />
+  </ItemGroup>
+
+  <ItemGroup>
+    <Content Include="appsettings.json" Condition="Exists('appsettings.json')">
+      <Generator></Generator>
+      <CopyToOutputDirectory>Always</CopyToOutputDirectory>
+      <CopyToPublishDirectory>Always</CopyToPublishDirectory>
+    </Content>
+  </ItemGroup>
+
+</Project>
diff --git a/Engine/Source/Programs/Horde/Horde/Program.cs b/Engine/Source/Programs/Horde/Horde/Program.cs
new file mode 100644
--- /dev/null
+++ b/Engine/Source/Programs/Horde/Horde/Program.cs
@@ -0,0 +1,89 @@
+﻿// Copyright Epic Games, Inc. All Rights Reserved.
+
+using EpicGames.Core;
+using System.Reflection;
+using Microsoft.Extensions.DependencyInjection;
+using Microsoft.Extensions.Logging;
+using EpicGames.Horde.Storage.Backends;
+using Microsoft.Extensions.Configuration;
+using System.Runtime.InteropServices;
+using Serilog.Sinks.SystemConsole.Themes;
+using Serilog;
+using Serilog.Formatting.Json;
+using System.Diagnostics;
+
+namespace Horde
+{
+	class Program
+	{
+		static DirectoryReference DataDir { get; } = GetDataDir();
+
+		static async Task<int> Main(string[] args)
+		{
+			CommandLineArguments arguments = new CommandLineArguments(args);
+			IConfiguration configuration = new ConfigurationBuilder()
+				.AddJsonFile("appsettings.json", optional: false)
+				.AddJsonFile("appsettings.User.json", optional: true)
+//				.AddEnvironmentVariables()
+				.Build();
+
+			using ILoggerFactory loggerFactory = CreateLoggerFactory(configuration);
+
+			IServiceCollection services = new ServiceCollection();
+			services.AddCommandsFromAssembly(Assembly.GetExecutingAssembly());
+			services.AddSingleton(loggerFactory);
+			services.AddMemoryCache();
+
+			// Execute all the commands
+			IServiceProvider serviceProvider = services.BuildServiceProvider();
+			return await CommandHost.RunAsync(arguments, serviceProvider, null);
+		}
+
+		static DirectoryReference GetAppDir()
+		{
+			return new DirectoryReference(Path.GetDirectoryName(Assembly.GetExecutingAssembly().Location)!);
+		}
+
+		static DirectoryReference GetDataDir()
+		{
+			if (RuntimeInformation.IsOSPlatform(OSPlatform.Windows))
+			{
+				DirectoryReference? programDataDir = DirectoryReference.GetSpecialFolder(Environment.SpecialFolder.CommonApplicationData);
+				if (programDataDir != null)
+				{
+					return DirectoryReference.Combine(programDataDir, "HordeAgent");
+				}
+			}
+			return GetAppDir();
+		}
+
+		public static ILoggerFactory CreateLoggerFactory(IConfiguration configuration)
+		{
+			Serilog.ILogger logger = CreateSerilogLogger(configuration);
+			return new Serilog.Extensions.Logging.SerilogLoggerFactory(logger, true);
+		}
+
+		static Serilog.ILogger CreateSerilogLogger(IConfiguration configuration)
+		{
+			DirectoryReference.CreateDirectory(Program.DataDir);
+
+			ConsoleTheme theme;
+			if (RuntimeInformation.IsOSPlatform(OSPlatform.Windows) && Environment.OSVersion.Version < new Version(10, 0))
+			{
+				theme = SystemConsoleTheme.Literate;
+			}
+			else
+			{
+				theme = AnsiConsoleTheme.Code;
+			}
+
+			return new LoggerConfiguration()
+				.WriteTo.Console(outputTemplate: "[{Timestamp:HH:mm:ss} {Level:w3}] {Indent}{Message:l}{NewLine}{Exception}", theme: theme)
+				.WriteTo.File(FileReference.Combine(Program.DataDir, "Log-.txt").FullName, fileSizeLimitBytes: 50 * 1024 * 1024, rollingInterval: RollingInterval.Day, rollOnFileSizeLimit: true, retainedFileCountLimit: 10)
+				.WriteTo.File(new JsonFormatter(renderMessage: true), FileReference.Combine(Program.DataDir, "Log-.json").FullName, fileSizeLimitBytes: 50 * 1024 * 1024, rollingInterval: RollingInterval.Day, rollOnFileSizeLimit: true, retainedFileCountLimit: 10)
+				.ReadFrom.Configuration(configuration)
+				.Enrich.FromLogContext()
+				.CreateLogger();
+		}
+	}
+}
\ No newline at end of file
diff --git a/Engine/Source/Programs/Horde/Horde/appsettings.json b/Engine/Source/Programs/Horde/Horde/appsettings.json
new file mode 100644
--- /dev/null
+++ b/Engine/Source/Programs/Horde/Horde/appsettings.json
@@ -0,0 +1,21 @@
+{
+	"Serilog": {
+		"MinimumLevel": {
+			"Default": "Verbose",
+			"Override": {
+				"Microsoft": "Information",
+				"System.Net.Http.HttpClient": "Warning",
+				"Microsoft.AspNetCore.Routing.EndpointMiddleware": "Warning",
+				"Microsoft.AspNetCore.Authorization.DefaultAuthorizationService": "Warning",
+				"HordeServer.Authentication.HordeJwtBearerHandler": "Warning",
+				"HordeServer.Authentication.OktaHandler": "Warning",
+				"Microsoft.AspNetCore.Hosting.Diagnostics": "Warning",
+				"Microsoft.AspNetCore.Mvc.Infrastructure.ControllerActionInvoker": "Warning",
+				"Serilog.AspNetCore.RequestLoggingMiddleware": "Warning",
+				"Grpc.Net.Client.Internal.GrpcCall": "Warning",
+				"EpicGames.Horde.Compute.ComputeSocket": "Verbose",
+				"EpicGames.Horde.Compute.ComputeChannel": "Verbose"
+			}
+		}
+	}
+}
