Horde: Add command line tools as a separate component to the installer.

diff --git a/Engine/Source/Programs/Horde/BuildHorde.xml b/Engine/Source/Programs/Horde/BuildHorde.xml
--- a/Engine/Source/Programs/Horde/BuildHorde.xml
+++ b/Engine/Source/Programs/Horde/BuildHorde.xml
@@ -83,14 +83,13 @@
 		<Property Name="StagingDir" Value="$(RootDir)/Staging"/>
 
 		<Node Name="Stage Horde Installer" Requires="#DashboardOutputFiles">
+			<Delete Files="$(StagingDir)/Tools/..."/>
 			<Delete Files="$(StagingDir)/Server/..."/>
 			<Delete Files="$(StagingDir)/Agent/..."/>
+			<DotNet Arguments="publish &quot;$(RootDir)/Engine/Source/Programs/Horde/Horde/Horde.csproj&quot; --output &quot;$(StagingDir)/Tools&quot; --runtime win-x64 --self-contained $(VersionArguments)"/>
 			<DotNet Arguments="publish &quot;$(RootDir)/Engine/Source/Programs/Horde/Horde.Server/Horde.Server.csproj&quot; --output &quot;$(StagingDir)/Server&quot; --runtime win-x64 --self-contained $(VersionArguments)"/>
 			<DotNet Arguments="publish &quot;$(RootDir)/Engine/Source/Programs/Horde/Horde.Agent/Horde.Agent.csproj&quot; --output &quot;$(StagingDir)/Agent&quot; $(VersionArguments)"/>
 
-			<!-- Create build settings -->
-			<WriteTextFile File="$(StagingDir)/Server/appsettings.Build.json" Text="{ &quot;Horde&quot;: { &quot;SingleInstance&quot;: true } }"/>
-			
 			<!-- Copy the Dashboard -->
 			<Copy Files="#DashboardOutputFiles" From="$(StagingDir)/Dashboard/..." To="$(StagingDir)/Server/DashboardApp/..." />
 		</Node>
@@ -102,6 +101,14 @@
 		<Node Name="Build Horde Installer" Requires="Stage Horde Installer">
 			<Property Name="CommonArgs" Value="-dConfiguration=$(InstallerConfig) -dPlatform=x64 -arch x64"/>
 
+			<!-- Create the tools file list -->
+			<Delete Files="$(StagingDir)/Tools-Bulk/..."/>
+			<Tag Files="$(StagingDir)/Tools/..." With="#ToolsFiles"/>
+			<Copy Files="#ToolsFiles" From="$(StagingDir)" To="$(StagingDir)/Tools-Bulk"/>
+
+			<Spawn Exe="$(WixDir)/heat.exe" Arguments="dir &quot;$(StagingDir)/Tools-Bulk&quot; -cg HordeToolsFiles -scom -sreg -gg -dr InstallDir -srd -var var.SourceDir -out &quot;$(InstallerDir)/HordeToolsFiles.wxs&quot;"/>
+			<Spawn Exe="$(WixDir)/candle.exe" Arguments="$(CommonArgs) HordeToolsFiles.wxs -dSourceDir=&quot;$(StagingDir)&quot; -out $(ObjDir)/HordeToolsFiles.wixobj" WorkingDir="$(InstallerDir)"/>
+
 			<!-- Create the agent file list -->
 			<Delete Files="$(StagingDir)/Agent-Bulk/..."/>
 			<Tag Files="$(StagingDir)/Agent/..." Except="HordeAgent.exe" With="#AgentFiles"/>
@@ -120,7 +127,7 @@
 
 			<!-- Main component -->
 			<Spawn Exe="$(WixDir)/candle.exe" Arguments="$(CommonArgs) Installer.wxs -dStagingDir=&quot;$(StagingDir)&quot; -out $(ObjDir)/Installer.wixobj" WorkingDir="$(InstallerDir)"/>
-			<Spawn Exe="$(WixDir)/light.exe" Arguments="-out $(BinDir)/Horde.msi -sw1076 -pdbout $(BinDir)/Horde.pdb $(ObjDir)/Installer.wixobj $(ObjDir)/HordeAgentFiles.wixobj $(ObjDir)/HordeServerFiles.wixobj -ext WixUIExtension -ext WixUtilExtension" WorkingDir="$(InstallerDir)"/>
+			<Spawn Exe="$(WixDir)/light.exe" Arguments="-out $(BinDir)/Horde.msi -sw1076 -pdbout $(BinDir)/Horde.pdb $(ObjDir)/Installer.wixobj $(ObjDir)/HordeAgentFiles.wixobj $(ObjDir)/HordeServerFiles.wixobj $(ObjDir)/HordeToolsFiles.wixobj -ext WixUIExtension -ext WixUtilExtension" WorkingDir="$(InstallerDir)"/>
 		</Node>
 		
 	</Agent>
diff --git a/Engine/Source/Programs/Horde/Horde.Agent/Horde.Agent.csproj b/Engine/Source/Programs/Horde/Horde.Agent/Horde.Agent.csproj
--- a/Engine/Source/Programs/Horde/Horde.Agent/Horde.Agent.csproj
+++ b/Engine/Source/Programs/Horde/Horde.Agent/Horde.Agent.csproj
@@ -102,9 +102,6 @@
     <Content Include="..\Horde.Server\Protos\horde\*.proto">
       <Link>Protos\%(FileName)%(Extension)</Link>
     </Content>
-    <Content Include="Horde.bat">
-      <CopyToOutputDirectory>PreserveNewest</CopyToOutputDirectory>
-    </Content>
     <Content Include="..\Horde.Server\Protos\horde\common\**\*.proto">
       <Link>Protos\horde\common\%(RecursiveDir)%(FileName)%(Extension)</Link>
     </Content>
diff --git a/Engine/Source/Programs/Horde/Horde.Agent/Horde.bat b/Engine/Source/Programs/Horde/Horde.Agent/Horde.bat
deleted file mode 100755
