Horde: Make bundle commands in Horde command line tool file-system centric.

Also fix generation of Blob ids/locators with empty prefixes.

diff --git a/Engine/Source/Programs/Horde/Horde/Commands/Bundles/BundleCreate.cs b/Engine/Source/Programs/Horde/Horde/Commands/Bundles/BundleCreate.cs
--- a/Engine/Source/Programs/Horde/Horde/Commands/Bundles/BundleCreate.cs
+++ b/Engine/Source/Programs/Horde/Horde/Commands/Bundles/BundleCreate.cs
@@ -5,7 +5,9 @@ using System.Threading;
 using System.Threading.Tasks;
 using EpicGames.Core;
 using EpicGames.Horde.Storage;
+using EpicGames.Horde.Storage.Backends;
 using EpicGames.Horde.Storage.Nodes;
+using JetBrains.Annotations;
 using Microsoft.Extensions.Logging;
 
 namespace Horde.Commands.Bundles
@@ -13,36 +15,60 @@ namespace Horde.Commands.Bundles
 	[Command("bundle", "create", "Creates a bundle from a folder on the local hard drive")]
 	class BundleCreate : Command
 	{
-		[CommandLine("-Namespace=", Description = "Namespace to use for storage")]
-		public NamespaceId NamespaceId { get; set; } = new NamespaceId("default");
+		[CommandLine("-Output=", Required = true, Description = "Output file. Should have a .ref extension.")]
+		public FileReference OutputFile { get; set; } = null!;
 
-		[CommandLine("-Ref=")]
-		public RefName RefName { get; set; } = new RefName("default-bundle");
+		[CommandLine("-Input=", Required = true, Description = "Input file or directory")]
+		public string Input { get; set; } = null!;
 
-		[CommandLine("-InputDir=", Required = true)]
-		public DirectoryReference InputDir { get; set; } = null!;
+		[CommandLine("-Filter=", Description = "Filter for files to include, in P4 syntax (eg. Foo/...).")]
+		public string Filter { get; set; } = "...";
 
-		readonly IStorageClientFactory _storageClientFactory;
-
-		public BundleCreate(IStorageClientFactory storageClientFactory)
-		{
-			_storageClientFactory = storageClientFactory;
-		}
+		[CommandLine("-Blobs=", Description = "Directory to store blobs")]
+		public DirectoryReference? BlobDir { get; set; }
 
 		public override async Task<int> ExecuteAsync(ILogger logger)
 		{
-			IStorageClient store = await _storageClientFactory.GetClientAsync(NamespaceId);
+			DirectoryReference baseDir;
+			List<FileReference> files = new List<FileReference>();
 
-			using TreeWriter writer = new TreeWriter(store, prefix: RefName.Text);
+			if (File.Exists(Input))
+			{
+				FileReference file = new FileReference(Input);
+				baseDir = file.Directory;
+				files.Add(file);
+			}
+			else if (Directory.Exists(Input))
+			{
+				baseDir = new DirectoryReference(Input);
+				FileFilter filter = new FileFilter(Filter.Split(';'));
+				files.AddRange(filter.ApplyToDirectory(baseDir, true));
+			}
+			else
+			{
+				logger.LogError("{Path} does not exist", Input);
+				return 1;
+			}
 
-			DirectoryNode node = new DirectoryNode(DirectoryFlags.None);
+			DirectoryReference blobDir = BlobDir ?? OutputFile.Directory;
+			DirectoryReference.CreateDirectory(blobDir);
+
+			FileStorageClient store = new FileStorageClient(blobDir, logger);
 
 			Stopwatch timer = Stopwatch.StartNew();
 
-			ChunkingOptions options = new ChunkingOptions();
-			await node.CopyFromDirectoryAsync(InputDir.ToDirectoryInfo(), options, writer, null, CancellationToken.None);
+			using (TreeWriter writer = new TreeWriter(store))
+			{
+				DirectoryNode node = new DirectoryNode(DirectoryFlags.None);
+
+				ChunkingOptions options = new ChunkingOptions();
+				await node.CopyFilesAsync(baseDir, files, options, writer, null, CancellationToken.None);
+
+				NodeHandle handle = await writer.FlushAsync(node);
 
-			await writer.WriteAsync(RefName, node);
+				logger.LogInformation("Writing {File}", OutputFile);
+				await FileReference.WriteAllTextAsync(OutputFile, handle.ToString());
+			}
 
 			logger.LogInformation("Time: {Time}", timer.Elapsed.TotalSeconds);
 			return 0;
diff --git a/Engine/Source/Programs/Horde/Horde/Commands/Bundles/BundleDump.cs b/Engine/Source/Programs/Horde/Horde/Commands/Bundles/BundleDump.cs
--- a/Engine/Source/Programs/Horde/Horde/Commands/Bundles/BundleDump.cs
+++ b/Engine/Source/Programs/Horde/Horde/Commands/Bundles/BundleDump.cs
@@ -13,59 +13,17 @@ namespace Horde.Commands.Bundles
 	[Command("bundle", "dump", "Dumps the contents of a bundle")]
 	internal class BundleDump : Command
 	{
-		[CommandLine("-Namespace=", Description = "Namespace to use for storage")]
-		public NamespaceId NamespaceId { get; set; } = new NamespaceId("default");
-
-		[CommandLine("-Ref=")]
-		public RefName? RefName { get; set; }
-
-		[CommandLine("-Blob=")]
-		public BlobLocator? BlobId { get; set; }
-
-		[CommandLine("-File=")]
-		public string? File { get; set; }
+		[CommandLine("-Input=")]
+		public FileReference Input { get; set; } = null!;
 
 		[CommandLine("-Verbose")]
 		public bool Verbose { get; set; }
 
-		readonly IStorageClientFactory _storageClientFactory;
-
-		public BundleDump(IStorageClientFactory storageClientFactory)
-		{
-			_storageClientFactory = storageClientFactory;
-		}
-
 		public override async Task<int> ExecuteAsync(ILogger logger)
 		{
-			Bundle bundle;
-			if (File != null)
-			{
-				FileReference Location = new FileReference(File);
-				byte[] bytes = await FileReference.ReadAllBytesAsync(Location);
-				bundle = new Bundle(new MemoryReader(bytes));
-				logger.LogInformation("Summary for blob {Location}", Location);
-			}
-			else if (BlobId != null)
-			{
-				IStorageClient store = await _storageClientFactory.GetClientAsync(NamespaceId);
-				bundle = await store.ReadBundleAsync(BlobId.Value);
-				logger.LogInformation("Summary for blob {BlobId}", BlobId.Value);
-			}
-			else if (RefName != null)
-			{
-				IStorageClient store = await _storageClientFactory.GetClientAsync(NamespaceId);
-
-				NodeHandle handle = await store.ReadRefTargetAsync(RefName.Value);
-				logger.LogInformation("Ref {RefName} -> {Locator}", RefName.Value, handle);
-
-				bundle = await store.ReadBundleAsync(handle.Locator.Blob);
-				logger.LogInformation("Summary for blob {BlobId}", handle.Locator.Blob);
-			}
-			else
-			{
-				logger.LogError("No blob specified - use -File/Blob/Ref=...");
-				return 1;
-			}
+			byte[] bytes = await FileReference.ReadAllBytesAsync(Input);
+			Bundle bundle = new Bundle(new MemoryReader(bytes));
+			logger.LogInformation("Summary for blob {Location}", Input);
 
 			BundleHeader header = bundle.Header;
 
diff --git a/Engine/Source/Programs/Horde/Horde/Commands/Bundles/BundleExtract.cs b/Engine/Source/Programs/Horde/Horde/Commands/Bundles/BundleExtract.cs
--- a/Engine/Source/Programs/Horde/Horde/Commands/Bundles/BundleExtract.cs
+++ b/Engine/Source/Programs/Horde/Horde/Commands/Bundles/BundleExtract.cs
@@ -5,6 +5,7 @@ using System.Threading;
 using System.Threading.Tasks;
 using EpicGames.Core;
 using EpicGames.Horde.Storage;
+using EpicGames.Horde.Storage.Backends;
 using EpicGames.Horde.Storage.Nodes;
 using Microsoft.Extensions.Caching.Memory;
 using Microsoft.Extensions.Logging;
@@ -14,32 +15,28 @@ namespace Horde.Commands.Bundles
 	[Command("bundle", "extract", "Extracts data from a bundle to the local hard drive")]
 	internal class BundleExtract : Command
 	{
-		[CommandLine("-Namespace=", Description = "Namespace to use for storage")]
-		public NamespaceId NamespaceId { get; set; } = new NamespaceId("default");
+		[CommandLine("-Input=")]
+		public FileReference Input { get; set; } = null!;
 
-		[CommandLine("-Ref=")]
-		public RefName RefName { get; set; } = new RefName("default-bundle");
+		[CommandLine("-BlobDir=")]
+		public DirectoryReference? BlobDir { get; set; }
 
 		[CommandLine("-OutputDir=", Required = true)]
 		public DirectoryReference OutputDir { get; set; } = null!;
 
-		readonly IStorageClientFactory _storageClientFactory;
-
-		public BundleExtract(IStorageClientFactory storageClientFactory)
-		{
-			_storageClientFactory = storageClientFactory;
-		}
-
 		public override async Task<int> ExecuteAsync(ILogger logger)
 		{
-			IStorageClient store = await _storageClientFactory.GetClientAsync(NamespaceId);
+			string Text = await FileReference.ReadAllTextAsync(Input);
+			NodeHandle handle = NodeHandle.Parse(Text);
+
+			FileStorageClient store = new FileStorageClient(BlobDir ?? Input.Directory, logger);
 
 			using MemoryCache cache = new MemoryCache(new MemoryCacheOptions());
 			TreeReader reader = new TreeReader(store, cache, logger);
 
 			Stopwatch timer = Stopwatch.StartNew();
 
-			DirectoryNode node = await reader.ReadNodeAsync<DirectoryNode>(RefName);
+			DirectoryNode node = await reader.ReadNodeAsync<DirectoryNode>(handle.Locator);
 			await node.CopyToDirectoryAsync(reader, OutputDir.ToDirectoryInfo(), logger, CancellationToken.None);
 
 			logger.LogInformation("Elapsed: {Time}s", timer.Elapsed.TotalSeconds);
diff --git a/Engine/Source/Programs/Horde/Horde/Program.cs b/Engine/Source/Programs/Horde/Horde/Program.cs
--- a/Engine/Source/Programs/Horde/Horde/Program.cs
+++ b/Engine/Source/Programs/Horde/Horde/Program.cs
@@ -32,6 +32,7 @@ namespace Horde
 			IServiceCollection services = new ServiceCollection();
 			services.AddCommandsFromAssembly(Assembly.GetExecutingAssembly());
 			services.AddSingleton(loggerFactory);
+			services.AddLogging();
 			services.AddMemoryCache();
 
 			// Execute all the commands
diff --git a/Engine/Source/Programs/Shared/EpicGames.Horde/Storage/BlobId.cs b/Engine/Source/Programs/Shared/EpicGames.Horde/Storage/BlobId.cs
--- a/Engine/Source/Programs/Shared/EpicGames.Horde/Storage/BlobId.cs
+++ b/Engine/Source/Programs/Shared/EpicGames.Horde/Storage/BlobId.cs
@@ -89,24 +89,27 @@ namespace EpicGames.Horde.Storage
 		/// <returns>New content id</returns>
 		public static BlobId CreateNew(Utf8String prefix = default)
 		{
-			if (prefix.Length == 0)
-			{
-				DateTime now = DateTime.UtcNow.Date;
-				prefix = $"_by_date_/{now.Year}-{now.Month:D2}/{now.Day:D2}";
-			}
-			else
+			byte[] buffer;
+			Span<byte> span;
+
+			if (prefix.Length > 0)
 			{
 				ValidateArgument(nameof(prefix), prefix);
-			}
 
-			byte[] buffer = new byte[prefix.Length + 1 + 24];
-			Span<byte> span = buffer;
+				buffer = new byte[prefix.Length + 1 + 24];
+				span = buffer;
 
-			prefix.Span.CopyTo(span);
-			span = span.Slice(prefix.Length);
+				prefix.Span.CopyTo(span);
+				span = span.Slice(prefix.Length);
 
-			span[0] = (byte)'/';
-			span = span.Slice(1);
+				span[0] = (byte)'/';
+				span = span.Slice(1);
+			}
+			else
+			{
+				buffer = new byte[24];
+				span = buffer;
+			}
 
 			GenerateUniqueId(span);
 			return new BlobId(new Utf8String(buffer), Validate.None);
diff --git a/Engine/Source/Programs/Shared/EpicGames.Horde/Storage/BlobLocator.cs b/Engine/Source/Programs/Shared/EpicGames.Horde/Storage/BlobLocator.cs
--- a/Engine/Source/Programs/Shared/EpicGames.Horde/Storage/BlobLocator.cs
+++ b/Engine/Source/Programs/Shared/EpicGames.Horde/Storage/BlobLocator.cs
@@ -149,41 +149,37 @@ namespace EpicGames.Horde.Storage
 		/// <returns>New content id</returns>
 		public static BlobLocator Create(HostId serverId, Utf8String prefix = default)
 		{
-			if (prefix.Length == 0)
+			int length = 24;
+			if (serverId.IsValid())
 			{
-				DateTime now = DateTime.UtcNow.Date;
-				prefix = $"_by_date_/{now.Year}-{now.Month:D2}/{now.Day:D2}";
+				length += serverId.Inner.Length + 1;
 			}
-			else
+			if (prefix.Length > 0)
 			{
-				BlobId.ValidateArgument(nameof(prefix), prefix);
+				length += prefix.Length + 1;
 			}
 
-			byte[] buffer;
-			Span<byte> span;
+			byte[] buffer = new byte[length];
+			Span<byte> span = buffer;
 
 			if (serverId.IsValid())
 			{
-				buffer = new byte[serverId.Inner.Length + 1 + prefix.Length + 1 + 24];
-				span = buffer;
-
 				serverId.Inner.Span.CopyTo(span);
 				span = span.Slice(serverId.Inner.Length);
 
 				span[0] = (byte)':';
 				span = span.Slice(1);
 			}
-			else
+			if (prefix.Length > 0)
 			{
-				buffer = new byte[prefix.Length + 1 + 24];
-				span = buffer;
-			}
+				BlobId.ValidateArgument(nameof(prefix), prefix);
 
-			prefix.Span.CopyTo(span);
-			span = span.Slice(prefix.Length);
+				prefix.Span.CopyTo(span);
+				span = span.Slice(prefix.Length);
 
-			span[0] = (byte)'/';
-			span = span.Slice(1);
+				span[0] = (byte)'/';
+				span = span.Slice(1);
+			}
 
 			BlobId.GenerateUniqueId(span);
 			return new BlobLocator(new Utf8String(buffer), Sanitize.None);
