Horde: Fix support for environment variables in command line tool.

diff --git a/Engine/Source/Programs/Horde/Horde/Horde.csproj b/Engine/Source/Programs/Horde/Horde/Horde.csproj
--- a/Engine/Source/Programs/Horde/Horde/Horde.csproj
+++ b/Engine/Source/Programs/Horde/Horde/Horde.csproj
@@ -15,6 +15,7 @@
   </ItemGroup>
 
   <ItemGroup>
+    <PackageReference Include="Microsoft.Extensions.Configuration.EnvironmentVariables" Version="7.0.0" />
     <PackageReference Include="Serilog" Version="2.12.0" />
     <PackageReference Include="Serilog.Extensions.Hosting" Version="5.0.1" />
     <PackageReference Include="Serilog.Extensions.Logging" Version="3.1.0" />
diff --git a/Engine/Source/Programs/Horde/Horde/Program.cs b/Engine/Source/Programs/Horde/Horde/Program.cs
--- a/Engine/Source/Programs/Horde/Horde/Program.cs
+++ b/Engine/Source/Programs/Horde/Horde/Program.cs
@@ -2,15 +2,13 @@
 
 using EpicGames.Core;
 using System.Reflection;
+using Microsoft.Extensions.Configuration;
 using Microsoft.Extensions.DependencyInjection;
 using Microsoft.Extensions.Logging;
-using EpicGames.Horde.Storage.Backends;
-using Microsoft.Extensions.Configuration;
 using System.Runtime.InteropServices;
 using Serilog.Sinks.SystemConsole.Themes;
 using Serilog;
 using Serilog.Formatting.Json;
-using System.Diagnostics;
 
 namespace Horde
 {
@@ -24,7 +22,7 @@ namespace Horde
 			IConfiguration configuration = new ConfigurationBuilder()
 				.AddJsonFile("appsettings.json", optional: false)
 				.AddJsonFile("appsettings.User.json", optional: true)
-//				.AddEnvironmentVariables()
+				.AddEnvironmentVariables()
 				.Build();
 
 			using ILoggerFactory loggerFactory = CreateLoggerFactory(configuration);
