Horde: Add a separate option to route to the landing page by default, and enable it by default for new builds.

diff --git a/Engine/Source/Programs/Horde/Horde.Server/Defaults/globals.json b/Engine/Source/Programs/Horde/Horde.Server/Defaults/globals.json
--- a/Engine/Source/Programs/Horde/Horde.Server/Defaults/globals.json
+++ b/Engine/Source/Programs/Horde/Horde.Server/Defaults/globals.json
@@ -1,4 +1,12 @@
 {
+	"dashboard": {
+		"showLandingPage": true,
+		"showCI": true,
+		"showAgents": true,
+		"showPerforceServers": true,
+		"showDeviceManager": true,
+		"showTests": true
+	},
 	"compute": [
 		{
 			"id": "default",
@@ -9,7 +17,7 @@
 		"backends": [
 			{
 				"id": "default-backend",
-				"type":"FileSystem",
+				"type": "FileSystem",
 				"baseDir": "Storage" // Default base directory is C:\ProgramData\HordeServer
 			},
 			{
diff --git a/Engine/Source/Programs/Horde/Horde.Server/Server/GlobalConfig.cs b/Engine/Source/Programs/Horde/Horde.Server/Server/GlobalConfig.cs
--- a/Engine/Source/Programs/Horde/Horde.Server/Server/GlobalConfig.cs
+++ b/Engine/Source/Programs/Horde/Horde.Server/Server/GlobalConfig.cs
@@ -39,6 +39,42 @@ namespace Horde.Server.Server
 		public string Path { get; set; } = null!;
 	}
 
+	/// <summary>
+	/// Configuration for global features
+	/// </summary>
+	public class DashboardConfig
+	{
+		/// <summary>
+		/// Navigate to the landing page by default
+		/// </summary>
+		public bool ShowLandingPage { get; set; } = false;
+
+		/// <summary>
+		/// Enable CI functionality
+		/// </summary>
+		public bool ShowCI { get; set; } = true;
+
+		/// <summary>
+		/// Whether to show functionality related to agents, pools, and utilization on the dashboard.
+		/// </summary>
+		public bool ShowAgents { get; set; } = true;
+
+		/// <summary>
+		/// Show the Perforce server option on the server menu
+		/// </summary>
+		public bool ShowPerforceServers { get; set; } = true;
+
+		/// <summary>
+		/// Show the device manager on the server menu
+		/// </summary>
+		public bool ShowDeviceManager { get; set; } = true;
+
+		/// <summary>
+		/// Show automated tests on the server menu
+		/// </summary>
+		public bool ShowTests { get; set; } = true;
+	}
+
 	/// <summary>
 	/// Global configuration
 	/// </summary>
@@ -72,6 +108,11 @@ namespace Horde.Server.Server
 		/// </summary>
 		public List<ConfigInclude> Include { get; set; } = new List<ConfigInclude>();
 
+		/// <summary>
+		/// Settings for the dashboard
+		/// </summary>
+		public DashboardConfig Dashboard { get; set; } = new DashboardConfig();
+
 		/// <summary>
 		/// List of projects
 		/// </summary>
diff --git a/Engine/Source/Programs/Horde/Horde.Server/Users/UserMessages.cs b/Engine/Source/Programs/Horde/Horde.Server/Users/UserMessages.cs
--- a/Engine/Source/Programs/Horde/Horde.Server/Users/UserMessages.cs
+++ b/Engine/Source/Programs/Horde/Horde.Server/Users/UserMessages.cs
@@ -111,29 +111,22 @@ namespace Horde.Server.Users
 	/// </summary>
 	public class GetDashboardFeaturesResponse
 	{
-		/// <summary>
-		/// Enable CI functionality
-		/// </summary>
+		/// <inheritdoc cref="DashboardConfig.ShowLandingPage"/>
+		public bool ShowLandingPage { get; set; }
+
+		/// <inheritdoc cref="DashboardConfig.ShowCI"/>
 		public bool ShowCI { get; set; }
 
-		/// <summary>
-		/// Whether to show functionality related to agents, pools, and utilization on the dashboard.
-		/// </summary>
+		/// <inheritdoc cref="DashboardConfig.ShowAgents"/>
 		public bool ShowAgents { get; set; }
 
-		/// <summary>
-		/// Show the Perforce server option on the server menu
-		/// </summary>
+		/// <inheritdoc cref="DashboardConfig.ShowPerforceServers"/>
 		public bool ShowPerforceServers { get; set; }
 
-		/// <summary>
-		/// Show the device manager on the server menu
-		/// </summary>
+		/// <inheritdoc cref="DashboardConfig.ShowDeviceManager"/>
 		public bool ShowDeviceManager { get; set; }
 
-		/// <summary>
-		/// Show automated tests on the server menu
-		/// </summary>
+		/// <inheritdoc cref="DashboardConfig.ShowTests"/>
 		public bool ShowTests { get; set; }
 
 		/// <summary>
@@ -156,11 +149,12 @@ namespace Horde.Server.Users
 		/// </summary>
 		public GetDashboardFeaturesResponse(GlobalConfig globalConfig, ClaimsPrincipal principal)
 		{
-			ShowCI = true;
-			ShowAgents = true;
-			ShowPerforceServers = true;
-			ShowDeviceManager = true;
-			ShowTests = true;
+			ShowLandingPage = globalConfig.Dashboard.ShowLandingPage;
+			ShowCI = globalConfig.Dashboard.ShowCI;
+			ShowAgents = globalConfig.Dashboard.ShowAgents;
+			ShowPerforceServers = globalConfig.Dashboard.ShowPerforceServers;
+			ShowDeviceManager = globalConfig.Dashboard.ShowDeviceManager;
+			ShowTests = globalConfig.Dashboard.ShowTests;
 			ShowNoticeEditor = globalConfig.Authorize(NoticeAclAction.CreateNotice, principal) || globalConfig.Authorize(NoticeAclAction.UpdateNotice, principal);
 			ShowPoolEditor = globalConfig.Authorize(PoolAclAction.CreatePool, principal) || globalConfig.Authorize(PoolAclAction.UpdatePool, principal);
 			ShowRemoteDesktop = globalConfig.Authorize(AgentAclAction.UpdateAgent, principal);
diff --git a/Engine/Source/Programs/Horde/HordeDashboard/src/App.tsx b/Engine/Source/Programs/Horde/HordeDashboard/src/App.tsx
--- a/Engine/Source/Programs/Horde/HordeDashboard/src/App.tsx
+++ b/Engine/Source/Programs/Horde/HordeDashboard/src/App.tsx
@@ -129,7 +129,7 @@ const Main: React.FC = () => {
       const routes: RouteObject[] = [
          {
             path: "/", element: <Root />, errorElement: <RouteError />, children: [
-               { path: "index", element: (dashboard.user?.dashboardFeatures?.showCI === false) ? <DocView /> : <UserHomeView /> },
+               { path: "index", element: (dashboard.user?.dashboardFeatures?.showLandingPage === true) ? <DocView /> : <UserHomeView /> },
                { path: "project/:projectId", element: <ProjectHome /> },
                { path: "pools", element: <PoolView /> },
                { path: "job/:jobId", element: <JobDetailViewV2 /> },
diff --git a/Engine/Source/Programs/Horde/HordeDashboard/src/backend/Api.ts b/Engine/Source/Programs/Horde/HordeDashboard/src/backend/Api.ts
--- a/Engine/Source/Programs/Horde/HordeDashboard/src/backend/Api.ts
+++ b/Engine/Source/Programs/Horde/HordeDashboard/src/backend/Api.ts
@@ -3105,6 +3105,9 @@ export type GetDashboardFeaturesResponse = {
 	/** Whether the remote desktop button should be shown on the agent modal */
 	showRemoteDesktop?: boolean;
 
+	/** Show the landing page by default */
+	showLandingPage?: boolean;
+
 	/** Enable CI functionality */
 	showCI?: boolean;
 
diff --git a/Engine/Source/Programs/Horde/HordeDashboard/src/components/docs/DocView.tsx b/Engine/Source/Programs/Horde/HordeDashboard/src/components/docs/DocView.tsx
--- a/Engine/Source/Programs/Horde/HordeDashboard/src/components/docs/DocView.tsx
+++ b/Engine/Source/Programs/Horde/HordeDashboard/src/components/docs/DocView.tsx
@@ -171,7 +171,7 @@ const DocRail = observer(() => {
 
    const refLinks: ISideRailLink[] = [];
 
-   if (dashboard.user?.dashboardFeatures?.showCI === false) {
+   if (dashboard.user?.dashboardFeatures?.showLandingPage === true) {
       refLinks.push({ text: "Landing", url: "/docs" });
    }
 
@@ -211,7 +211,7 @@ export const DocView = () => {
    }
    if (!docName || docName.indexOf("README.md") !== -1) {
 
-      if (dashboard.user?.dashboardFeatures?.showCI === false) {
+      if (dashboard.user?.dashboardFeatures?.showLandingPage === true) {
          docName = "documentation/Docs/Landing.md";
       } else {
          docName = "documentation/Docs/Home.md";
