Horde: Produce an error if an input file is not specified to the bundle dump command.

diff --git a/Engine/Source/Programs/Horde/Horde/Commands/Bundles/BundleDump.cs b/Engine/Source/Programs/Horde/Horde/Commands/Bundles/BundleDump.cs
--- a/Engine/Source/Programs/Horde/Horde/Commands/Bundles/BundleDump.cs
+++ b/Engine/Source/Programs/Horde/Horde/Commands/Bundles/BundleDump.cs
@@ -13,7 +13,7 @@ namespace Horde.Commands.Bundles
 	[Command("bundle", "dump", "Dumps the contents of a bundle")]
 	internal class BundleDump : Command
 	{
-		[CommandLine("-Input=")]
+		[CommandLine("-Input=", Required = true)]
 		public FileReference Input { get; set; } = null!;
 
 		[CommandLine("-Verbose")]
