Horde: Add tracing to agent telemetry events handling

Also prevent exceptions from bubbling with try/catch.

diff --git a/Engine/Source/Programs/Horde/Horde.Server/Telemetry/TelemetryManager.cs b/Engine/Source/Programs/Horde/Horde.Server/Telemetry/TelemetryManager.cs
--- a/Engine/Source/Programs/Horde/Horde.Server/Telemetry/TelemetryManager.cs
+++ b/Engine/Source/Programs/Horde/Horde.Server/Telemetry/TelemetryManager.cs
@@ -9,6 +9,7 @@ using HordeCommon;
 using Microsoft.Extensions.Hosting;
 using Microsoft.Extensions.Logging;
 using Microsoft.Extensions.Options;
+using OpenTelemetry.Trace;
 
 namespace Horde.Server.Telemetry
 {
@@ -22,6 +23,7 @@ namespace Horde.Server.Telemetry
 
 		private readonly List<ITelemetrySinkInternal> _telemetrySinks = new();
 		private readonly ITicker _ticker;
+		private readonly Tracer _tracer;
 		private readonly ILogger<TelemetryManager> _logger;
 
 		/// <summary>
@@ -31,8 +33,9 @@ namespace Horde.Server.Telemetry
 		/// <param name="clock"></param>
 		/// <param name="serverSettings"></param>
 		/// <param name="loggerFactory"></param>
-		public TelemetryManager(IHttpClientFactory httpClientFactory, IClock clock, IOptions<ServerSettings> serverSettings, ILoggerFactory loggerFactory)
+		public TelemetryManager(IHttpClientFactory httpClientFactory, IClock clock, IOptions<ServerSettings> serverSettings, Tracer tracer, ILoggerFactory loggerFactory)
 		{
+			_tracer = tracer;
 			_logger = loggerFactory.CreateLogger<TelemetryManager>();
 			_ticker = clock.AddTicker<EpicTelemetrySink>(TimeSpan.FromSeconds(30.0), FlushAsync, _logger);
 			
@@ -53,11 +56,25 @@ namespace Horde.Server.Telemetry
 		/// <inheritdoc/>
 		public void SendEvent(string eventName, object attributes)
 		{
+			using TelemetrySpan span = _tracer.StartActiveSpan($"{nameof(TelemetryManager)}.{nameof(SendEvent)}");
+			span.SetAttribute("eventName", eventName);
+			
 			foreach (ITelemetrySinkInternal sink in _telemetrySinks)
 			{
+				using TelemetrySpan sinkSpan = _tracer.StartActiveSpan($"SendEvent");
+				string fullName = sink.GetType().FullName ?? "Unknown";
+				span.SetAttribute("sink", fullName);
+				
 				if (sink.Enabled)
 				{
-					sink.SendEvent(eventName, attributes);
+					try
+					{
+						sink.SendEvent(eventName, attributes);
+					}
+					catch (Exception e)
+					{
+						_logger.LogError(e, "Failed sending event to {Sink}. Message: {Message}", fullName, e.Message);
+					}
 				}
 			}
 		}
