Horde: Add missing xml-doc

diff --git a/Engine/Source/Programs/Horde/Horde.Server/Telemetry/TelemetryManager.cs b/Engine/Source/Programs/Horde/Horde.Server/Telemetry/TelemetryManager.cs
--- a/Engine/Source/Programs/Horde/Horde.Server/Telemetry/TelemetryManager.cs
+++ b/Engine/Source/Programs/Horde/Horde.Server/Telemetry/TelemetryManager.cs
@@ -32,6 +32,7 @@ namespace Horde.Server.Telemetry
 		/// <param name="httpClientFactory"></param>
 		/// <param name="clock"></param>
 		/// <param name="serverSettings"></param>
+		/// <param name="tracer"></param>
 		/// <param name="loggerFactory"></param>
 		public TelemetryManager(IHttpClientFactory httpClientFactory, IClock clock, IOptions<ServerSettings> serverSettings, Tracer tracer, ILoggerFactory loggerFactory)
 		{
