Horde: Enable keep-alive pings for gRPC client connections

diff --git a/Engine/Source/Programs/Horde/Horde.Agent/Services/GrpcService.cs b/Engine/Source/Programs/Horde/Horde.Agent/Services/GrpcService.cs
--- a/Engine/Source/Programs/Horde/Horde.Agent/Services/GrpcService.cs
+++ b/Engine/Source/Programs/Horde/Horde.Agent/Services/GrpcService.cs
@@ -3,6 +3,7 @@
 using System;
 using System.Net.Http;
 using System.Net.Http.Headers;
+using System.Net.Security;
 using Grpc.Net.Client;
 using Horde.Agent.Utility;
 using Microsoft.Extensions.Logging;
@@ -47,14 +48,19 @@ namespace Horde.Agent.Services
 		/// <returns>New grpc channel</returns>
 		public GrpcChannel CreateGrpcChannel(string? bearerToken)
 		{
-			#pragma warning disable CA2000 // Dispose objects before losing scope
-			// HTTP client handler is disposed by GrpcChannel below
-			HttpClientHandler customCertHandler = new HttpClientHandler();
-			#pragma warning restore CA2000 // Dispose objects before losing scope
-			customCertHandler.ServerCertificateCustomValidationCallback += (sender, cert, chain, errors) => CertificateHelper.CertificateValidationCallBack(_logger, sender, cert, chain, errors, _serverProfile);
-			customCertHandler.CheckCertificateRevocationList = true;
+#pragma warning disable CA2000 // Dispose objects before losing scope
+			// HTTP handler is disposed by GrpcChannel below
+			SocketsHttpHandler httpHandler = new()
+			{
+				KeepAlivePingDelay = TimeSpan.FromSeconds(50),
+				SslOptions = new SslClientAuthenticationOptions
+				{
+					RemoteCertificateValidationCallback = (sender, cert, chain, errors) => CertificateHelper.CertificateValidationCallBack(_logger, sender, cert, chain, errors, _serverProfile)
+				}
+			};
+#pragma warning restore CA2000 // Dispose objects before losing scope
 
-			HttpClient httpClient = new (customCertHandler, true);
+			HttpClient httpClient = new (httpHandler, true);
 			httpClient.DefaultRequestHeaders.Add("Accept", "application/json");
 			if (bearerToken != null)
 			{
@@ -69,7 +75,6 @@ namespace Horde.Agent.Services
 				// Required payloads coming from CAS service can be large
 				MaxReceiveMessageSize = 1024 * 1024 * 1024, // 1 GB
 				MaxSendMessageSize = 1024 * 1024 * 1024, // 1 GB
-				
 				LoggerFactory = _loggerFactory,
 				HttpClient = httpClient,
 				DisposeHttpClient = true
