Horde: Remove unused tickers in telemetry sinks

diff --git a/Engine/Source/Programs/Horde/Horde.Server/Telemetry/ClickHouseTelemetrySink.cs b/Engine/Source/Programs/Horde/Horde.Server/Telemetry/ClickHouseTelemetrySink.cs
--- a/Engine/Source/Programs/Horde/Horde.Server/Telemetry/ClickHouseTelemetrySink.cs
+++ b/Engine/Source/Programs/Horde/Horde.Server/Telemetry/ClickHouseTelemetrySink.cs
@@ -8,7 +8,6 @@ using System.Text;
 using System.Threading;
 using System.Threading.Tasks;
 using Horde.Server.Jobs;
-using HordeCommon;
 using HordeCommon.Rpc.Messages.Telemetry;
 using Microsoft.Extensions.Logging;
 
@@ -29,7 +28,6 @@ namespace Horde.Server.Telemetry
 		private readonly IHttpClientFactory _httpClientFactory;
 		private readonly Uri? _uri;
 
-		private readonly ITicker _ticker;
 		private readonly ILogger _logger;
 		private readonly ConcurrentQueue<(string eventName, object evt)> _queuedEvents = new();
 
@@ -39,10 +37,9 @@ namespace Horde.Server.Telemetry
 		/// <summary>
 		/// Constructor
 		/// </summary>
-		public ClickHouseTelemetrySink(ClickHouseTelemetryConfig config, IHttpClientFactory httpClientFactory, IClock clock, ILogger<ClickHouseTelemetrySink> logger)
+		public ClickHouseTelemetrySink(ClickHouseTelemetryConfig config, IHttpClientFactory httpClientFactory, ILogger<ClickHouseTelemetrySink> logger)
 		{
 			_httpClientFactory = httpClientFactory;
-			_ticker = clock.AddTicker<ClickHouseTelemetrySink>(TimeSpan.FromSeconds(5.0), FlushAsync, logger);
 			_logger = logger;
 			_uri = config.Url;
 		}
@@ -50,7 +47,6 @@ namespace Horde.Server.Telemetry
 		/// <inheritdoc/>
 		public ValueTask DisposeAsync()
 		{
-			_ticker.Dispose();
 			return ValueTask.CompletedTask;
 		}
 		
diff --git a/Engine/Source/Programs/Horde/Horde.Server/Telemetry/EpicTelemetrySink.cs b/Engine/Source/Programs/Horde/Horde.Server/Telemetry/EpicTelemetrySink.cs
--- a/Engine/Source/Programs/Horde/Horde.Server/Telemetry/EpicTelemetrySink.cs
+++ b/Engine/Source/Programs/Horde/Horde.Server/Telemetry/EpicTelemetrySink.cs
@@ -10,7 +10,6 @@ using System.Text.Json.Serialization;
 using System.Threading;
 using System.Threading.Tasks;
 using EpicGames.Core;
-using HordeCommon;
 using Microsoft.AspNetCore.WebUtilities;
 using Microsoft.Extensions.Logging;
 
@@ -73,7 +72,6 @@ namespace Horde.Server.Telemetry
 		readonly ArrayMemoryWriter _eventMemoryWriter;
 		readonly Utf8JsonWriter _eventWriter;
 
-		readonly ITicker _ticker;
 		readonly ILogger _logger;
 
 		/// <inheritdoc/>
@@ -82,7 +80,7 @@ namespace Horde.Server.Telemetry
 		/// <summary>
 		/// Constructor
 		/// </summary>
-		public EpicTelemetrySink(IEpicTelemetrySinkConfig config, IHttpClientFactory httpClientFactory, IClock clock, ILogger<EpicTelemetrySink> logger)
+		public EpicTelemetrySink(IEpicTelemetrySinkConfig config, IHttpClientFactory httpClientFactory, ILogger<EpicTelemetrySink> logger)
 		{
 			_httpClientFactory = httpClientFactory;
 
@@ -97,7 +95,6 @@ namespace Horde.Server.Telemetry
 			_eventMemoryWriter = new ArrayMemoryWriter(65536);
 			_eventWriter = new Utf8JsonWriter(_eventMemoryWriter);
 
-			_ticker = clock.AddTicker<EpicTelemetrySink>(TimeSpan.FromSeconds(60.0), FlushAsync, logger);
 			_logger = logger;
 
 			if (config.Url != null)
@@ -119,7 +116,6 @@ namespace Horde.Server.Telemetry
 		{
 			await _packetWriter.DisposeAsync();
 			await _eventWriter.DisposeAsync();
-			_ticker.Dispose();
 		}
 		
 		/// <inheritdoc/>
diff --git a/Engine/Source/Programs/Horde/Horde.Server/Telemetry/TelemetryManager.cs b/Engine/Source/Programs/Horde/Horde.Server/Telemetry/TelemetryManager.cs
--- a/Engine/Source/Programs/Horde/Horde.Server/Telemetry/TelemetryManager.cs
+++ b/Engine/Source/Programs/Horde/Horde.Server/Telemetry/TelemetryManager.cs
@@ -45,10 +45,10 @@ namespace Horde.Server.Telemetry
 				switch (config)
 				{
 					case EpicTelemetryConfig epicConfig:
-						_telemetrySinks.Add(new EpicTelemetrySink(epicConfig, httpClientFactory, clock, loggerFactory.CreateLogger<EpicTelemetrySink>()));
+						_telemetrySinks.Add(new EpicTelemetrySink(epicConfig, httpClientFactory, loggerFactory.CreateLogger<EpicTelemetrySink>()));
 						break;
 					case ClickHouseTelemetryConfig chConfig:
-						_telemetrySinks.Add(new ClickHouseTelemetrySink(chConfig, httpClientFactory, clock, loggerFactory.CreateLogger<ClickHouseTelemetrySink>()));
+						_telemetrySinks.Add(new ClickHouseTelemetrySink(chConfig, httpClientFactory, loggerFactory.CreateLogger<ClickHouseTelemetrySink>()));
 						break;
 				}
 			}
