Horde: Add a simple tray application that shows status of the Horde Agent.

Agent periodically sets current status regarding health to the injected StatusService singleton. StatusService operates a named pipe server allowing clients to query current status of the service.

diff --git a/Engine/Source/Programs/Horde/BuildHorde.xml b/Engine/Source/Programs/Horde/BuildHorde.xml
--- a/Engine/Source/Programs/Horde/BuildHorde.xml
+++ b/Engine/Source/Programs/Horde/BuildHorde.xml
@@ -89,6 +89,7 @@
 			<DotNet Arguments="publish &quot;$(RootDir)/Engine/Source/Programs/Horde/Horde/Horde.csproj&quot; --output &quot;$(StagingDir)/Tools&quot; --runtime win-x64 --self-contained $(VersionArguments)"/>
 			<DotNet Arguments="publish &quot;$(RootDir)/Engine/Source/Programs/Horde/Horde.Server/Horde.Server.csproj&quot; --output &quot;$(StagingDir)/Server&quot; --runtime win-x64 --self-contained $(VersionArguments)"/>
 			<DotNet Arguments="publish &quot;$(RootDir)/Engine/Source/Programs/Horde/Horde.Agent/Horde.Agent.csproj&quot; --output &quot;$(StagingDir)/Agent&quot; $(VersionArguments)"/>
+			<DotNet Arguments="publish &quot;$(RootDir)/Engine/Source/Programs/Horde/Horde.Agent.TrayApp/Horde.Agent.TrayApp.csproj&quot; --output &quot;$(StagingDir)/Agent&quot; $(VersionArguments)"/>
 
 			<!-- Copy the Dashboard -->
 			<Copy Files="#DashboardOutputFiles" From="$(StagingDir)/Dashboard/..." To="$(StagingDir)/Server/DashboardApp/..." />
@@ -111,7 +112,7 @@
 
 			<!-- Create the agent file list -->
 			<Delete Files="$(StagingDir)/Agent-Bulk/..."/>
-			<Tag Files="$(StagingDir)/Agent/..." Except="HordeAgent.exe" With="#AgentFiles"/>
+			<Tag Files="$(StagingDir)/Agent/..." Except="HordeAgent.exe;Horde.Agent.TrayApp.exe" With="#AgentFiles"/>
 			<Copy Files="#AgentFiles" From="$(StagingDir)" To="$(StagingDir)/Agent-Bulk"/>
 
 			<Spawn Exe="$(WixDir)/heat.exe" Arguments="dir &quot;$(StagingDir)/Agent-Bulk&quot; -cg HordeAgentFiles -scom -sreg -gg -dr InstallDir -srd -var var.SourceDir -out &quot;$(InstallerDir)/HordeAgentFiles.wxs&quot;"/>
@@ -126,8 +127,8 @@
 			<Spawn Exe="$(WixDir)/candle.exe" Arguments="$(CommonArgs) HordeServerFiles.wxs -dSourceDir=&quot;$(StagingDir)&quot; -out $(ObjDir)/HordeServerFiles.wixobj" WorkingDir="$(InstallerDir)"/>
 
 			<!-- Main component -->
-			<Spawn Exe="$(WixDir)/candle.exe" Arguments="$(CommonArgs) Installer.wxs -dStagingDir=&quot;$(StagingDir)&quot; -out $(ObjDir)/Installer.wixobj" WorkingDir="$(InstallerDir)"/>
-			<Spawn Exe="$(WixDir)/light.exe" Arguments="-out $(BinDir)/Horde.msi -sw1076 -pdbout $(BinDir)/Horde.pdb $(ObjDir)/Installer.wixobj $(ObjDir)/HordeAgentFiles.wixobj $(ObjDir)/HordeServerFiles.wixobj $(ObjDir)/HordeToolsFiles.wixobj -ext WixUIExtension -ext WixUtilExtension" WorkingDir="$(InstallerDir)"/>
+			<Spawn Exe="$(WixDir)/candle.exe" Arguments="$(CommonArgs) Installer.wxs -dStagingDir=&quot;$(StagingDir)&quot; -out $(ObjDir)/Installer.wixobj -ext WixUtilExtension" WorkingDir="$(InstallerDir)"/>
+			<Spawn Exe="$(WixDir)/light.exe" Arguments="-out $(BinDir)/Horde.msi -sw1076 -pdbout $(BinDir)/Horde.pdb $(ObjDir)/Installer.wixobj $(ObjDir)/HordeAgentFiles.wixobj $(ObjDir)/HordeServerFiles.wixobj $(ObjDir)/HordeToolsFiles.wixobj -ext WixUtilExtension -ext WixUIExtension" WorkingDir="$(InstallerDir)"/>
 		</Node>
 		
 	</Agent>
diff --git a/Engine/Source/Programs/Horde/Horde.Agent.Tests/WorkerServiceTest.cs b/Engine/Source/Programs/Horde/Horde.Agent.Tests/WorkerServiceTest.cs
--- a/Engine/Source/Programs/Horde/Horde.Agent.Tests/WorkerServiceTest.cs
+++ b/Engine/Source/Programs/Horde/Horde.Agent.Tests/WorkerServiceTest.cs
@@ -93,6 +93,7 @@ namespace Horde.Agent.Tests
 			});
 
 			_serviceCollection.AddSingleton<WorkerService>();
+			_serviceCollection.AddSingleton<StatusService>();
 
 			_serviceCollection.AddSingleton<JobHandler>();
 			_serviceCollection.AddSingleton<LeaseHandler>(sp => sp.GetRequiredService<JobHandler>());
@@ -242,7 +243,7 @@ namespace Horde.Agent.Tests
 			await using FakeHordeRpcServer fakeServer = new();
 			await using ISession session = FakeServerSessionFactory.CreateSession(fakeServer.GetConnection());
 
-			LeaseManager manager = new LeaseManager(session, null!, serviceProvider.GetRequiredService<IEnumerable<LeaseHandler>>(), serviceProvider.GetRequiredService<IOptions<AgentSettings>>(), NullLogger.Instance);
+			LeaseManager manager = new LeaseManager(session, null!, serviceProvider.GetRequiredService<StatusService>(), serviceProvider.GetRequiredService<IEnumerable<LeaseHandler>>(), serviceProvider.GetRequiredService<IOptions<AgentSettings>>(), NullLogger.Instance);
 
 			Task handleSessionTask = Task.Run(() => manager.RunAsync(false, cts.Token), cts.Token);
 			await fakeServer.UpdateSessionReceived.Task.WaitAsync(cts.Token);
diff --git a/Engine/Source/Programs/Horde/Horde.Agent.TrayApp/Horde.Agent.TrayApp.csproj b/Engine/Source/Programs/Horde/Horde.Agent.TrayApp/Horde.Agent.TrayApp.csproj
new file mode 100644
--- /dev/null
+++ b/Engine/Source/Programs/Horde/Horde.Agent.TrayApp/Horde.Agent.TrayApp.csproj
@@ -0,0 +1,35 @@
+﻿<Project Sdk="Microsoft.NET.Sdk">
+
+  <PropertyGroup>
+    <OutputType>WinExe</OutputType>
+    <TargetFramework>net6.0-windows</TargetFramework>
+    <Nullable>enable</Nullable>
+    <UseWindowsForms>true</UseWindowsForms>
+    <ImplicitUsings>enable</ImplicitUsings>
+    <NeutralLanguage>en-US</NeutralLanguage>
+  </PropertyGroup>
+
+  <ItemGroup>
+    <Compile Include="..\Horde.Agent\IpcMessages.cs" Link="IpcMessages.cs" />
+  </ItemGroup>
+
+  <ItemGroup>
+    <ProjectReference Include="..\..\Shared\EpicGames.Core\EpicGames.Core.csproj" />
+  </ItemGroup>
+
+  <ItemGroup>
+    <Compile Update="Properties\Resources.Designer.cs">
+      <DesignTime>True</DesignTime>
+      <AutoGen>True</AutoGen>
+      <DependentUpon>Resources.resx</DependentUpon>
+    </Compile>
+  </ItemGroup>
+
+  <ItemGroup>
+    <EmbeddedResource Update="Properties\Resources.resx">
+      <Generator>ResXFileCodeGenerator</Generator>
+      <LastGenOutput>Resources.Designer.cs</LastGenOutput>
+    </EmbeddedResource>
+  </ItemGroup>
+
+</Project>
\ No newline at end of file
diff --git a/Engine/Source/Programs/Horde/Horde.Agent.TrayApp/Program.cs b/Engine/Source/Programs/Horde/Horde.Agent.TrayApp/Program.cs
new file mode 100644
--- /dev/null
+++ b/Engine/Source/Programs/Horde/Horde.Agent.TrayApp/Program.cs
@@ -0,0 +1,184 @@
+// Copyright Epic Games, Inc. All Rights Reserved.
+
+using System.Buffers;
+using System.Buffers.Binary;
+using System.Diagnostics;
+using System.IO.Pipes;
+using System.Text;
+using System.Text.Json;
+using System.Text.Json.Serialization;
+using EpicGames.Core;
+using Horde.Agent.TrayApp.Properties;
+
+namespace Horde.Agent.TrayApp
+{
+	static class Program
+	{
+		[STAThread]
+		static async Task Main()
+		{
+			ApplicationConfiguration.Initialize();
+
+			await using (CustomApplicationContext appContext = new CustomApplicationContext())
+			{
+				Application.Run(appContext);
+			}
+		}
+	}
+
+	public class CustomApplicationContext : ApplicationContext, IAsyncDisposable
+	{
+		readonly NotifyIcon _trayIcon;
+		readonly BackgroundTask _clientTask;
+		readonly Control _mainThreadInvokeTarget;
+
+		bool _disposed;
+
+		public CustomApplicationContext()
+		{
+			ToolStripMenuItem logsMenuItem = new ToolStripMenuItem("View logs");
+			logsMenuItem.Click += OnOpenLogs;
+
+			ToolStripMenuItem exitMenuItem = new ToolStripMenuItem("Exit");
+			exitMenuItem.Click += OnExit;
+
+			ContextMenuStrip menu = new ContextMenuStrip();
+			menu.Items.Add(logsMenuItem);
+			menu.Items.Add(new ToolStripSeparator());
+			menu.Items.Add(exitMenuItem);
+
+			_mainThreadInvokeTarget = new Control();
+			_mainThreadInvokeTarget.CreateControl();
+
+			_trayIcon = new NotifyIcon()
+			{
+				Icon = Resources.StatusNormal,
+				ContextMenuStrip = menu,
+				Visible = true
+			};
+
+			_clientTask = BackgroundTask.StartNew(StatusTask);
+		}
+
+		protected override void Dispose(bool disposing)
+		{
+			base.Dispose(disposing);
+
+			if (disposing)
+			{
+				_mainThreadInvokeTarget.Dispose();
+				_trayIcon.Dispose();
+				_disposed = true;
+			}
+		}
+
+		public async ValueTask DisposeAsync()
+		{
+			await _clientTask.DisposeAsync();
+
+			Dispose();
+			GC.SuppressFinalize(this);
+		}
+
+		private void OnOpenLogs(object? sender, EventArgs e)
+		{
+			DirectoryReference? programDataDir = DirectoryReference.GetSpecialFolder(Environment.SpecialFolder.CommonApplicationData);
+			if (programDataDir != null)
+			{
+				DirectoryReference logsDir = DirectoryReference.Combine(programDataDir, "HordeAgent");
+				Process.Start(new ProcessStartInfo { FileName = logsDir.FullName, UseShellExecute = true });
+			}
+		}
+
+		private void OnExit(object? sender, EventArgs e)
+		{
+			ExitThread();
+		}
+
+		void SetStatus(AgentStatusMessage status)
+		{
+			_mainThreadInvokeTarget.BeginInvoke(() => SetStatus_MainThread(status));
+		}
+
+		void SetStatus_MainThread(AgentStatusMessage status)
+		{
+			if (!_disposed)
+			{
+				if (!status.Healthy)
+				{
+					_trayIcon.Icon = Resources.StatusError;
+					_trayIcon.Text = String.IsNullOrEmpty(status.Detail) ? "Error. Check logs." : (status.Detail.Length > 100) ? status.Detail.Substring(0, 100) : status.Detail;
+				}
+				else if (status.NumLeases > 0)
+				{
+					_trayIcon.Icon = Resources.StatusBusy;
+					_trayIcon.Text = (status.NumLeases == 1) ? "Currently handling 1 lease" : $"Currently handling {status.NumLeases} leases";
+				}
+				else
+				{
+					_trayIcon.Icon = Resources.StatusNormal;
+					_trayIcon.Text = "Agent is operating normally";
+				}
+			}
+		}
+
+		async Task StatusTask(CancellationToken cancellationToken)
+		{
+			SetStatus(AgentStatusMessage.Starting);
+			for (; ; )
+			{
+				try
+				{
+					await PollForStatusUpdatesAsync(cancellationToken);
+				}
+				catch (OperationCanceledException) when (cancellationToken.IsCancellationRequested)
+				{
+					break;
+				}
+				catch
+				{
+					SetStatus(new AgentStatusMessage(false, 0, "Unable to connect to Agent. Check the service is running."));
+					await Task.Delay(TimeSpan.FromSeconds(5.0), cancellationToken);
+				}
+			}
+		}
+
+		async Task PollForStatusUpdatesAsync(CancellationToken cancellationToken)
+		{
+			AgentMessageBuffer message = new AgentMessageBuffer();
+			using (NamedPipeClientStream pipeClient = new NamedPipeClientStream(".", AgentMessagePipe.PipeName, PipeDirection.InOut))
+			{
+				SetStatus(new AgentStatusMessage(false, 0, "Connecting to agent..."));
+				await pipeClient.ConnectAsync(cancellationToken);
+
+				SetStatus(new AgentStatusMessage(false, 0, "Waiting for status update."));
+				for (; ; )
+				{
+					message.Set(AgentMessageType.GetStatusRequest);
+					await message.SendAsync(pipeClient, cancellationToken);
+
+					if (!await message.TryReadAsync(pipeClient, cancellationToken))
+					{
+						break;
+					}					
+					
+					switch (message.Type)
+					{
+						case AgentMessageType.GetStatusResponse:
+							AgentStatusMessage status = message.Parse<AgentStatusMessage>();
+							SetStatus(status);
+							break;
+					}
+
+					await Task.Delay(TimeSpan.FromSeconds(5.0), cancellationToken);
+				}
+			}
+		}
+
+		void Exit(object sender, EventArgs e)
+		{
+			_trayIcon.Visible = false;
+			Application.Exit();
+		}
+	}
+}
diff --git a/Engine/Source/Programs/Horde/Horde.Agent/Leases/LeaseManager.cs b/Engine/Source/Programs/Horde/Horde.Agent/Leases/LeaseManager.cs
--- a/Engine/Source/Programs/Horde/Horde.Agent/Leases/LeaseManager.cs
+++ b/Engine/Source/Programs/Horde/Horde.Agent/Leases/LeaseManager.cs
@@ -118,14 +118,16 @@ namespace Horde.Agent.Leases
 
 		readonly ISession _session;
 		readonly CapabilitiesService _capabilitiesService;
+		readonly StatusService _statusService;
 		readonly Dictionary<string, LeaseHandler> _typeUrlToLeaseHandler;
 		readonly AgentSettings _settings;
 		readonly ILogger _logger;
 
-		public LeaseManager(ISession session, CapabilitiesService capabilitiesService, IEnumerable<LeaseHandler> leaseHandlers, IOptions<AgentSettings> settings, ILogger logger)
+		public LeaseManager(ISession session, CapabilitiesService capabilitiesService, StatusService statusService, IEnumerable<LeaseHandler> leaseHandlers, IOptions<AgentSettings> settings, ILogger logger)
 		{
 			_session = session;
 			_capabilitiesService = capabilitiesService;
+			_statusService = statusService;
 			_typeUrlToLeaseHandler = leaseHandlers.ToDictionary(x => x.LeaseType, x => x);
 			_settings = settings.Value;
 			_logger = logger;
@@ -134,7 +136,7 @@ namespace Horde.Agent.Leases
 		}
 
 		public LeaseManager(ISession session, IServiceProvider serviceProvider)
-			: this(session, serviceProvider.GetRequiredService<CapabilitiesService>(), serviceProvider.GetRequiredService<IEnumerable<LeaseHandler>>(), serviceProvider.GetRequiredService<IOptions<AgentSettings>>(), serviceProvider.GetRequiredService<ILogger<LeaseManager>>())
+			: this(session, serviceProvider.GetRequiredService<CapabilitiesService>(), serviceProvider.GetRequiredService<StatusService>(), serviceProvider.GetRequiredService<IEnumerable<LeaseHandler>>(), serviceProvider.GetRequiredService<IOptions<AgentSettings>>(), serviceProvider.GetRequiredService<ILogger<LeaseManager>>())
 		{
 		}
 
@@ -365,6 +367,9 @@ namespace Horde.Agent.Leases
 							return _sessionResult;
 						}
 					}
+
+					// Update the current status
+					_statusService.Set(true, _activeLeases.Count, (_activeLeases.Count == 0) ? "Waiting for work" : $"Executing {_activeLeases.Count} lease(s)");
 				}
 
 				// Update the historical update times
diff --git a/Engine/Source/Programs/Horde/Horde.Agent/Program.cs b/Engine/Source/Programs/Horde/Horde.Agent/Program.cs
--- a/Engine/Source/Programs/Horde/Horde.Agent/Program.cs
+++ b/Engine/Source/Programs/Horde/Horde.Agent/Program.cs
@@ -192,6 +192,8 @@ namespace Horde.Agent
 			services.AddSingleton<IJobExecutorFactory, TestExecutorFactory>();
 
 			services.AddSingleton<JobHandler>();
+			services.AddSingleton<StatusService>();
+			services.AddHostedService<StatusService>(sp => sp.GetRequiredService<StatusService>());
 
 			services.AddSingleton<LeaseHandler, ComputeHandler>();
 			services.AddSingleton<LeaseHandler, ConformHandler>();
diff --git a/Engine/Source/Programs/Horde/Horde.Agent/Services/SessionFactory.cs b/Engine/Source/Programs/Horde/Horde.Agent/Services/SessionFactory.cs
--- a/Engine/Source/Programs/Horde/Horde.Agent/Services/SessionFactory.cs
+++ b/Engine/Source/Programs/Horde/Horde.Agent/Services/SessionFactory.cs
@@ -131,13 +131,7 @@ namespace Horde.Agent.Services
 		/// <summary>
 		/// Creates a new agent session
 		/// </summary>
-		/// <param name="capabilitiesService"></param>
-		/// <param name="grpcService"></param>
-		/// <param name="settings"></param>
-		/// <param name="logger"></param>
-		/// <param name="cancellationToken">Indicates that the service is trying to stop</param>
-		/// <returns>Async task</returns>
-		public static async Task<Session> CreateAsync(CapabilitiesService capabilitiesService, GrpcService grpcService, IOptions<AgentSettings> settings, ILogger logger, CancellationToken cancellationToken)
+		public static async Task<Session> CreateAsync(CapabilitiesService capabilitiesService, GrpcService grpcService, StatusService statusService, IOptions<AgentSettings> settings, ILogger logger, CancellationToken cancellationToken)
 		{
 			AgentSettings currentSettings = settings.Value;
 
@@ -259,21 +253,23 @@ namespace Horde.Agent.Services
 	{
 		readonly CapabilitiesService _capabilitiesService;
 		readonly GrpcService _grpcService;
+		readonly StatusService _statusService;
 		readonly IOptions<AgentSettings> _settings;
 		readonly ILogger _logger;
 
 		/// <summary>
 		/// Constructor
 		/// </summary>
-		public SessionFactory(CapabilitiesService capabilitiesService, GrpcService grpcService, IOptions<AgentSettings> settings, ILogger<SessionFactory> logger)
+		public SessionFactory(CapabilitiesService capabilitiesService, GrpcService grpcService, StatusService statusService, IOptions<AgentSettings> settings, ILogger<SessionFactory> logger)
 		{
 			_capabilitiesService = capabilitiesService;
 			_grpcService = grpcService;
+			_statusService = statusService;
 			_settings = settings;
 			_logger = logger;
 		}
 
 		/// <inheritdoc/>
-		public async Task<ISession> CreateAsync(CancellationToken cancellationToken) => await Session.CreateAsync(_capabilitiesService, _grpcService, _settings, _logger, cancellationToken);
+		public async Task<ISession> CreateAsync(CancellationToken cancellationToken) => await Session.CreateAsync(_capabilitiesService, _grpcService, _statusService, _settings, _logger, cancellationToken);
 	}
 }
diff --git a/Engine/Source/Programs/Horde/Horde.Agent/Services/StatusService.cs b/Engine/Source/Programs/Horde/Horde.Agent/Services/StatusService.cs
new file mode 100644
--- /dev/null
+++ b/Engine/Source/Programs/Horde/Horde.Agent/Services/StatusService.cs
@@ -0,0 +1,167 @@
+// Copyright Epic Games, Inc. All Rights Reserved.
+
+using System;
+using System.Collections.Generic;
+using System.IO.Pipes;
+using System.Runtime.InteropServices;
+using System.Runtime.Versioning;
+using System.Security.AccessControl;
+using System.Security.Principal;
+using System.Threading;
+using System.Threading.Tasks;
+using EpicGames.Core;
+using Microsoft.Extensions.Hosting;
+using Microsoft.Extensions.Logging;
+
+namespace Horde.Agent.Services
+{
+	/// <summary>
+	/// Tracks the current status of the agent
+	/// </summary>
+	class StatusService : IHostedService, IAsyncDisposable
+	{
+		const int NumPipes = 10;
+
+		private AgentStatusMessage _current;
+
+		readonly BackgroundTask _task;
+		readonly ILogger _logger;
+
+		/// <summary>
+		/// The current agent status
+		/// </summary>
+		public AgentStatusMessage Current => _current;
+
+		/// <summary>
+		/// Constructor
+		/// </summary>
+		public StatusService(ILogger<StatusService> logger)
+		{
+			_current = AgentStatusMessage.Starting;
+			_task = new BackgroundTask(RunPipeServer);
+			_logger = logger;
+		}
+
+		/// <inheritdoc/>
+		public async ValueTask DisposeAsync()
+		{
+			await _task.DisposeAsync();
+		}
+
+		/// <inheritdoc/>
+		public Task StartAsync(CancellationToken cancellationToken)
+		{
+			_task.Start();
+			return Task.CompletedTask;
+		}
+
+		/// <inheritdoc/>
+		public async Task StopAsync(CancellationToken cancellationToken)
+		{
+			await _task.StopAsync();
+		}
+
+		/// <summary>
+		/// Sets the current status
+		/// </summary>
+		public void Set(AgentStatusMessage status)
+		{
+			Interlocked.Exchange(ref _current, status);
+			_logger.LogInformation("Updating status: {@Status}", status);
+		}
+
+		/// <summary>
+		/// Sets the current status
+		/// </summary>
+		public void Set(bool healthy, int numLeases, string message) => Set(new AgentStatusMessage(healthy, numLeases, message));
+
+		/// <summary>
+		/// Sets a status description
+		/// </summary>
+		public void SetDescription(string description) => Update(status => new AgentStatusMessage(status.Healthy, status.NumLeases, description));
+
+		/// <summary>
+		/// Updates the status using a custom funciton
+		/// </summary>
+		/// <param name="updateFunc">Function to take the existing status and create an updated version</param>
+		public void Update(Func<AgentStatusMessage, AgentStatusMessage> updateFunc)
+		{
+			for (; ; )
+			{
+				AgentStatusMessage info = _current;
+				if (Interlocked.CompareExchange(ref _current, updateFunc(info), info) == info)
+				{
+					break;
+				}
+			}
+		}
+
+		private async Task RunPipeServer(CancellationToken cancellationToken)
+		{
+			if (RuntimeInformation.IsOSPlatform(OSPlatform.Windows))
+			{
+				List<Task> tasks = new List<Task>();
+				for (int idx = 0; idx < NumPipes; idx++)
+				{
+					tasks.Add(RunSinglePipeServer(cancellationToken));
+				}
+				await Task.WhenAll(tasks);
+			}
+		}
+
+		[SupportedOSPlatform("windows")]
+		private async Task RunSinglePipeServer(CancellationToken cancellationToken)
+		{
+			while (!cancellationToken.IsCancellationRequested)
+			{
+				try
+				{
+					await RunPipeServerInternalAsync(cancellationToken);
+				}
+				catch (OperationCanceledException) when (cancellationToken.IsCancellationRequested)
+				{
+					break;
+				}
+				catch (Exception ex)
+				{
+					_logger.LogError(ex, "Exception while running pipe server: {Message}", ex.Message);
+				}
+
+				await Task.Delay(TimeSpan.FromSeconds(30), cancellationToken);
+			}
+		}
+
+		[SupportedOSPlatform("windows")]
+		private async Task RunPipeServerInternalAsync(CancellationToken cancellationToken)
+		{
+			AgentMessageBuffer request = new AgentMessageBuffer();
+			AgentMessageBuffer response = new AgentMessageBuffer();
+
+			_logger.LogDebug("Creating pipe for status updates");
+
+			PipeSecurity pipeSecurity = new PipeSecurity();
+			pipeSecurity.AddAccessRule(new PipeAccessRule(WindowsIdentity.GetCurrent().Name, PipeAccessRights.FullControl, AccessControlType.Allow));
+			pipeSecurity.AddAccessRule(new PipeAccessRule("Users", PipeAccessRights.ReadWrite, AccessControlType.Allow));
+
+			using (NamedPipeServerStream pipeServer = NamedPipeServerStreamAcl.Create(AgentMessagePipe.PipeName, PipeDirection.InOut, NumPipes, PipeTransmissionMode.Byte, PipeOptions.Asynchronous, 0, 0, pipeSecurity))
+			{
+				await pipeServer.WaitForConnectionAsync(cancellationToken);
+				_logger.LogDebug("Received pipe connection");
+
+				while (await request.TryReadAsync(pipeServer, cancellationToken))
+				{
+					switch (request.Type)
+					{
+						case AgentMessageType.GetStatusRequest:
+							response.Set(AgentMessageType.GetStatusResponse, Current);
+							break;
+						default:
+							response.Set(AgentMessageType.InvalidResponse);
+							break;
+					}
+					await response.SendAsync(pipeServer, cancellationToken);
+				}
+			}
+		}
+	}
+}
diff --git a/Engine/Source/Programs/Horde/Horde.Agent/Services/WorkerService.cs b/Engine/Source/Programs/Horde/Horde.Agent/Services/WorkerService.cs
--- a/Engine/Source/Programs/Horde/Horde.Agent/Services/WorkerService.cs
+++ b/Engine/Source/Programs/Horde/Horde.Agent/Services/WorkerService.cs
@@ -24,6 +24,7 @@ namespace Horde.Agent.Services
 		readonly IOptions<AgentSettings> _settings;
 		readonly ISessionFactory _sessionFactory;
 		readonly CapabilitiesService _capabilitiesService;
+		readonly StatusService _statusService;
 		readonly LeaseHandler[] _leaseHandlers;
 		readonly IServiceProvider _serviceProvider;
 
@@ -38,11 +39,12 @@ namespace Horde.Agent.Services
 		/// <summary>
 		/// Constructor. Registers with the server and starts accepting connections.
 		/// </summary>
-		public WorkerService(IOptions<AgentSettings> settings, ISessionFactory sessionFactory, CapabilitiesService capabilitiesService, IEnumerable<LeaseHandler> leaseHandlers, IServiceProvider serviceProvider, ILogger<WorkerService> logger)
+		public WorkerService(IOptions<AgentSettings> settings, ISessionFactory sessionFactory, CapabilitiesService capabilitiesService, StatusService statusService, IEnumerable<LeaseHandler> leaseHandlers, IServiceProvider serviceProvider, ILogger<WorkerService> logger)
 		{
 			_settings = settings;
 			_sessionFactory = sessionFactory;
 			_capabilitiesService = capabilitiesService;
+			_statusService = statusService;
 			_logger = logger;
 			_leaseHandlers = leaseHandlers.ToArray();
 			_serviceProvider = serviceProvider;
@@ -82,6 +84,7 @@ namespace Horde.Agent.Services
 			while (!stoppingToken.IsCancellationRequested)
 			{
 				SessionResult? result = null;
+				_statusService.Set(AgentStatusMessage.Starting);
 
 				Stopwatch sessionTime = Stopwatch.StartNew();
 
@@ -94,7 +97,7 @@ namespace Horde.Agent.Services
 
 						await using (ISession session = await _sessionFactory.CreateAsync(stoppingToken))
 						{
-							LeaseManager leaseManager = new LeaseManager(session, _capabilitiesService, _leaseHandlers, _settings, _logger);
+							LeaseManager leaseManager = new LeaseManager(session, _capabilitiesService, _statusService, _leaseHandlers, _settings, _logger);
 							result = await leaseManager.RunAsync(false, stoppingToken);
 						}
 
@@ -117,6 +120,7 @@ namespace Horde.Agent.Services
 
 						TimeSpan backOffTime = s_sessionBackOffTime[Math.Min(failureCount - 1, s_sessionBackOffTime.Length - 1)];
 						_logger.LogInformation("Session failure #{FailureNum}. Waiting {Time} and restarting. ({Message})", failureCount, backOffTime, ex.Message);
+						_statusService.Set(false, 0, $"Unable to start session: {ex.Message}");
 						await Task.Delay(backOffTime, stoppingToken);
 					}
 				}
