Horde: Add retrying policy for HTTP client used in JobRpcCommon

diff --git a/Engine/Source/Programs/Horde/Horde.Server/Horde.Server.csproj b/Engine/Source/Programs/Horde/Horde.Server/Horde.Server.csproj
--- a/Engine/Source/Programs/Horde/Horde.Server/Horde.Server.csproj
+++ b/Engine/Source/Programs/Horde/Horde.Server/Horde.Server.csproj
@@ -42,6 +42,7 @@
     <PackageReference Include="Microsoft.Extensions.Caching.Memory" Version="6.0.1" />
     <PackageReference Include="Microsoft.Extensions.Hosting.WindowsServices" Version="6.0.1" />
     <PackageReference Include="Microsoft.Extensions.Http" Version="6.0.0" />
+    <PackageReference Include="Microsoft.Extensions.Http.Polly" Version="6.0.16" />
     <PackageReference Include="MongoDB.Driver" Version="2.18.0" />
     <PackageReference Include="OpenTelemetry" Version="1.4.0" />
     <PackageReference Include="OpenTelemetry.Exporter.Console" Version="1.4.0" />
diff --git a/Engine/Source/Programs/Horde/Horde.Server/Startup.cs b/Engine/Source/Programs/Horde/Horde.Server/Startup.cs
--- a/Engine/Source/Programs/Horde/Horde.Server/Startup.cs
+++ b/Engine/Source/Programs/Horde/Horde.Server/Startup.cs
@@ -96,6 +96,8 @@ using StatusCode = Grpc.Core.StatusCode;
 using Horde.Server.Artifacts;
 using Horde.Server.Compute;
 using OpenTelemetry.Trace;
+using Polly;
+using Polly.Extensions.Http;
 
 namespace Horde.Server
 {
@@ -379,7 +381,7 @@ namespace Horde.Server
 			});
 			services.AddGrpcReflection();
 
-			services.AddHttpClient<JobRpcCommon>();
+			services.AddHttpClient<JobRpcCommon>().AddPolicyHandler(GetDefaultHttpRetryPolicy());
 			services.AddScoped<JobRpcCommon>();
 
 			services.AddSingleton<IAgentCollection, AgentCollection>();
@@ -878,6 +880,13 @@ namespace Horde.Server
 				}
 		*/
 
+		private static IAsyncPolicy<HttpResponseMessage> GetDefaultHttpRetryPolicy()
+		{
+			return HttpPolicyExtensions
+				.HandleTransientHttpError()
+				.WaitAndRetryAsync(new[] { TimeSpan.FromSeconds(1), TimeSpan.FromSeconds(5), TimeSpan.FromSeconds(10), TimeSpan.FromSeconds(10) });
+		}
+		
 		public sealed class HostIdBsonSerializer : SerializerBase<HostId>
 		{
 			/// <inheritdoc/>
