Horde: Add a separate step for creating an agent before starting a session.

- Creating an agent returns a token that the agent can use for creating sessions in the future.
- Token is saved to the user application directory for future use.

diff --git a/Engine/Source/Programs/Horde/Horde.Agent.Tests/WorkerServiceTest.cs b/Engine/Source/Programs/Horde/Horde.Agent.Tests/WorkerServiceTest.cs
--- a/Engine/Source/Programs/Horde/Horde.Agent.Tests/WorkerServiceTest.cs
+++ b/Engine/Source/Programs/Horde/Horde.Agent.Tests/WorkerServiceTest.cs
@@ -442,7 +442,7 @@ namespace Horde.Agent.Tests
 		public CreateSessionResponse OnCreateSessionRequest(CreateSessionRequest request)
 		{
 			CreateSessionReceived.TrySetResult(true);
-			_logger.LogInformation("OnCreateSessionRequest: {Name} {Status}", request.Name, request.Status);
+			_logger.LogInformation("OnCreateSessionRequest: {AgentId} {Status}", request.Id, request.Status);
 			CreateSessionResponse response = new()
 			{
 				AgentId = "bogusAgentId",
diff --git a/Engine/Source/Programs/Horde/Horde.Agent/Program.cs b/Engine/Source/Programs/Horde/Horde.Agent/Program.cs
--- a/Engine/Source/Programs/Horde/Horde.Agent/Program.cs
+++ b/Engine/Source/Programs/Horde/Horde.Agent/Program.cs
@@ -9,6 +9,8 @@ using System.Net.Http;
 using System.Net.Sockets;
 using System.Reflection;
 using System.Runtime.InteropServices;
+using System.Text.Json;
+using System.Text.Json.Serialization;
 using System.Threading;
 using System.Threading.Tasks;
 using Datadog.Trace;
@@ -89,6 +91,11 @@ namespace Horde.Agent
 		/// </summary>
 		public static string Version { get; } = GetVersion();
 
+		/// <summary>
+		/// Default settings for json serialization
+		/// </summary>
+		public static JsonSerializerOptions DefaultJsonSerializerOptions { get; } = new JsonSerializerOptions { AllowTrailingCommas = true, PropertyNameCaseInsensitive = true, PropertyNamingPolicy = JsonNamingPolicy.CamelCase, DefaultIgnoreCondition = JsonIgnoreCondition.WhenWritingNull };
+
 		/// <summary>
 		/// Entry point
 		/// </summary>
diff --git a/Engine/Source/Programs/Horde/Horde.Agent/Services/SessionFactory.cs b/Engine/Source/Programs/Horde/Horde.Agent/Services/SessionFactory.cs
--- a/Engine/Source/Programs/Horde/Horde.Agent/Services/SessionFactory.cs
+++ b/Engine/Source/Programs/Horde/Horde.Agent/Services/SessionFactory.cs
@@ -5,6 +5,7 @@ using System.Collections.Generic;
 using System.Linq;
 using System.Reflection;
 using System.Runtime.InteropServices;
+using System.Text.Json;
 using System.Threading;
 using System.Threading.Tasks;
 using EpicGames.Core;
@@ -128,6 +129,12 @@ namespace Horde.Agent.Services
 			_logger = logger;
 		}
 
+		class AgentRegistration
+		{
+			public string? Id { get; set; }
+			public string? Token { get; set; }
+		}
+
 		/// <summary>
 		/// Creates a new agent session
 		/// </summary>
@@ -182,15 +189,18 @@ namespace Horde.Agent.Services
 				}
 			}
 
+			// Read the registration settings
+			AgentRegistration registrationInfo = await GetAgentRegistrationAsync(grpcService, currentSettings, logger, cancellationToken);
+
 			// Create the session
 			CreateSessionResponse createSessionResponse;
-			using (GrpcChannel channel = grpcService.CreateGrpcChannel(serverProfile.Token))
+			using (GrpcChannel channel = grpcService.CreateGrpcChannel(registrationInfo.Token))
 			{
 				HordeRpc.HordeRpcClient rpcClient = new HordeRpc.HordeRpcClient(channel);
 
 				// Create the session information
 				CreateSessionRequest sessionRequest = new CreateSessionRequest();
-				sessionRequest.Name = currentSettings.GetAgentName();
+				sessionRequest.Id = registrationInfo.Id;
 				sessionRequest.Status = AgentStatus.Ok;
 				sessionRequest.Capabilities = capabilities;
 				sessionRequest.Version = Program.Version;
@@ -218,6 +228,55 @@ namespace Horde.Agent.Services
 			await RpcConnection.DisposeAsync();
 		}
 
+		/// <summary>
+		/// Registers the agent with the server
+		/// </summary>
+		static async Task<AgentRegistration> GetAgentRegistrationAsync(GrpcService grpcService, AgentSettings currentSettings, ILogger logger, CancellationToken cancellationToken)
+		{
+			// Get the location of the registration file
+			DirectoryReference? settingsDir = DirectoryReference.GetSpecialFolder(Environment.SpecialFolder.LocalApplicationData);
+			if (settingsDir == null)
+			{
+				settingsDir = DirectoryReference.GetCurrentDirectory();
+			}
+			else
+			{
+				settingsDir = DirectoryReference.Combine(settingsDir, "Horde.Agent");
+			}
+
+			FileReference settingsFile = FileReference.Combine(settingsDir, "registration.json");
+
+			// Read existing settings if possible
+			AgentRegistration? registrationInfo = null;
+			if (FileReference.Exists(settingsFile))
+			{
+				byte[] settingsData = await FileReference.ReadAllBytesAsync(settingsFile, cancellationToken);
+				registrationInfo = JsonSerializer.Deserialize<AgentRegistration>(settingsData, Program.DefaultJsonSerializerOptions);
+				logger.LogInformation("Read agent registration settings from {SettingsFile}", settingsFile);
+			}
+
+			// If they aren't valid, create a new agent registration
+			if (registrationInfo == null || registrationInfo.Id == null || registrationInfo.Token == null)
+			{
+				ServerProfile serverProfile = currentSettings.GetCurrentServerProfile();
+				using (GrpcChannel channel = grpcService.CreateGrpcChannel(serverProfile.Token))
+				{
+					HordeRpc.HordeRpcClient rpcClient = new HordeRpc.HordeRpcClient(channel);
+
+					CreateAgentRequest createAgentRequest = new CreateAgentRequest();
+					createAgentRequest.Name = currentSettings.GetAgentName();
+
+					CreateAgentResponse createAgentResponse = await rpcClient.CreateAgentAsync(createAgentRequest, null, null, cancellationToken);
+					registrationInfo = new AgentRegistration { Id = createAgentResponse.Id, Token = createAgentResponse.Token };
+					JsonSerializer.SerializeToUtf8Bytes(registrationInfo, Program.DefaultJsonSerializerOptions);
+
+					logger.LogInformation("Created agent (Id={AgentId}). Settings saved to {File}.", createAgentResponse.Id, settingsFile);
+				}
+			}
+
+			return registrationInfo;
+		}
+
 		/// <inheritdoc/>
 		public Task TerminateProcessesAsync(ILogger logger, CancellationToken cancellationToken)
 		{
diff --git a/Engine/Source/Programs/Horde/Horde.Server.Tests/RpcServiceTest.cs b/Engine/Source/Programs/Horde/Horde.Server.Tests/RpcServiceTest.cs
--- a/Engine/Source/Programs/Horde/Horde.Server.Tests/RpcServiceTest.cs
+++ b/Engine/Source/Programs/Horde/Horde.Server.Tests/RpcServiceTest.cs
@@ -296,7 +296,7 @@ namespace Horde.Server.Tests
 			CreateSessionRequest req = new CreateSessionRequest();
 			await Assert.ThrowsExceptionAsync<StructuredRpcException>(() => RpcService.CreateSession(req, _adminContext));
 
-			req.Name = "MyName";
+			req.Id = new AgentId("MyName").ToString();
 			await Assert.ThrowsExceptionAsync<StructuredRpcException>(() => RpcService.CreateSession(req, _adminContext));
 
 			req.Capabilities = new AgentCapabilities();
@@ -309,7 +309,7 @@ namespace Horde.Server.Tests
 		[TestMethod]
 		public async Task AgentJoinsPoolThroughProperties()
 		{
-			CreateSessionRequest req = new () { Name = "bogusAgentName", Capabilities = new AgentCapabilities() };
+			CreateSessionRequest req = new () { Id = new AgentId("bogusAgentName").ToString(), Capabilities = new AgentCapabilities() };
 			req.Capabilities.Properties.Add($"{KnownPropertyNames.RequestedPools}=fooPool,barPool");
 			CreateSessionResponse res = await RpcService.CreateSession(req, _adminContext);
 
@@ -317,7 +317,7 @@ namespace Horde.Server.Tests
 			CollectionAssert.AreEquivalent(new List<PoolId> { new ("fooPool"), new ("barPool") }, agent.GetPools().ToList());
 			
 			// Connect a second time, when the agent has already been created
-			req = new () { Name = "bogusAgentName", Capabilities = new AgentCapabilities() };
+			req = new () { Id = new AgentId("bogusAgentName").ToString(), Capabilities = new AgentCapabilities() };
 			req.Capabilities.Properties.Add($"{KnownPropertyNames.RequestedPools}=bazPool");
 			res = await RpcService.CreateSession(req, _adminContext);
 
@@ -328,7 +328,7 @@ namespace Horde.Server.Tests
 		[TestMethod]
 		public async Task PropertiesFromAgentCapabilities()
 		{
-			CreateSessionRequest req = new () { Name = "bogusAgentName", Capabilities = new AgentCapabilities() };
+			CreateSessionRequest req = new () { Id = new AgentId("bogusAgentName").ToString(), Capabilities = new AgentCapabilities() };
 			req.Capabilities.Properties.Add("fooKey=barValue");
 			CreateSessionResponse res = await RpcService.CreateSession(req, _adminContext);
 			IAgent agent = (await AgentService.GetAgentAsync(new AgentId(res.AgentId)))!;
@@ -338,7 +338,7 @@ namespace Horde.Server.Tests
 		[TestMethod]
 		public async Task PropertiesFromDeviceCapabilities()
 		{
-			CreateSessionRequest req = new () { Name = "bogusAgentName", Capabilities = new AgentCapabilities() };
+			CreateSessionRequest req = new () { Id = new AgentId("bogusAgentName").ToString(), Capabilities = new AgentCapabilities() };
 			req.Capabilities.Devices.Add(new DeviceCapabilities { Handle = "someHandle", Properties = { "foo=bar" } });
 			CreateSessionResponse res = await RpcService.CreateSession(req, _adminContext);
 			IAgent agent = (await AgentService.GetAgentAsync(new AgentId(res.AgentId)))!;
@@ -348,7 +348,7 @@ namespace Horde.Server.Tests
 		[TestMethod]
 		public async Task KnownPropertiesAreSetAsResources()
 		{
-			CreateSessionRequest req = new () { Name = "bogusAgentName", Capabilities = new AgentCapabilities() };
+			CreateSessionRequest req = new () { Id = new AgentId("bogusAgentName").ToString(), Capabilities = new AgentCapabilities() };
 			req.Capabilities.Devices.Add(new DeviceCapabilities { Handle = "someHandle", Properties = { $"{KnownPropertyNames.LogicalCores}=10" }});
 			CreateSessionResponse res = await RpcService.CreateSession(req, _adminContext);
 			IAgent agent = (await AgentService.GetAgentAsync(new AgentId(res.AgentId)))!;
@@ -360,7 +360,7 @@ namespace Horde.Server.Tests
 		{
 			CreateSessionRequest createReq = new CreateSessionRequest
 			{
-				Name = "UpdateSessionTest1", Capabilities = new AgentCapabilities()
+				Id = new AgentId("UpdateSessionTest1").ToString(), Capabilities = new AgentCapabilities()
 			};
 			CreateSessionResponse createRes = await RpcService.CreateSession(createReq, _adminContext);
 			string agentId = createRes.AgentId;
@@ -405,7 +405,7 @@ namespace Horde.Server.Tests
 		{
 			CreateSessionRequest createReq = new CreateSessionRequest
 			{
-				Name = "UpdateSessionTest1", Capabilities = new AgentCapabilities()
+				Id = new AgentId("UpdateSessionTest1").ToString(), Capabilities = new AgentCapabilities()
 			};
 			CreateSessionResponse createRes = await RpcService.CreateSession(createReq, _adminContext);
 			string agentId = createRes.AgentId;
diff --git a/Engine/Source/Programs/Horde/Horde.Server/Acls/AclService.cs b/Engine/Source/Programs/Horde/Horde.Server/Acls/AclService.cs
--- a/Engine/Source/Programs/Horde/Horde.Server/Acls/AclService.cs
+++ b/Engine/Source/Programs/Horde/Horde.Server/Acls/AclService.cs
@@ -84,7 +84,7 @@ namespace Horde.Server.Acls
 		/// <returns></returns>
 		public static AgentId? GetAgentId(ClaimsPrincipal user)
 		{
-			Claim? claim = user.Claims.FirstOrDefault(x => x.Type == HordeClaimTypes.AgentId);
+			Claim? claim = user.Claims.FirstOrDefault(x => x.Type == HordeClaimTypes.Agent);
 			if (claim == null)
 			{
 				return null;
@@ -103,6 +103,11 @@ namespace Horde.Server.Acls
 			return user.HasClaim(HordeClaims.AdminClaim.Type, HordeClaims.AdminClaim.Value);
 		}
 
+		public static bool HasAgentClaim(this ClaimsPrincipal user, AgentId agentId)
+		{
+			return user.HasClaim(HordeClaimTypes.Agent, agentId.ToString());
+		}
+
 		public static bool HasLeaseClaim(this ClaimsPrincipal user, LeaseId leaseId)
 		{
 			return user.HasClaim(HordeClaimTypes.Lease, leaseId.ToString());
diff --git a/Engine/Source/Programs/Horde/Horde.Server/Agents/AgentMessages.cs b/Engine/Source/Programs/Horde/Horde.Server/Agents/AgentMessages.cs
--- a/Engine/Source/Programs/Horde/Horde.Server/Agents/AgentMessages.cs
+++ b/Engine/Source/Programs/Horde/Horde.Server/Agents/AgentMessages.cs
@@ -12,53 +12,6 @@ using Horde.Server.Agents.Sessions;
 
 namespace Horde.Server.Agents
 {
-	/// <summary>
-	/// Parameters to register a new agent
-	/// </summary>
-	public class CreateAgentRequest
-	{
-		/// <summary>
-		/// Friendly name for the agent
-		/// </summary>
-		[Required]
-		public string Name { get; set; } = null!; // Enforced by [required] attribute
-
-		/// <summary>
-		/// Whether the agent is currently enabled
-		/// </summary>
-		public bool Enabled { get; set; } = true;
-
-		/// <summary>
-		/// Whether the agent is ephemeral (ie. should not be shown when inactive)
-		/// </summary>
-		public bool Ephemeral { get; set; }
-
-		/// <summary>
-		/// Pools for this agent
-		/// </summary>
-		public List<string>? Pools { get; set; }
-	}
-
-	/// <summary>
-	/// Response from creating an agent
-	/// </summary>
-	public class CreateAgentResponse
-	{
-		/// <summary>
-		/// Unique id for the new agent
-		/// </summary>
-		public string Id { get; set; }
-
-		/// <summary>
-		/// Constructor
-		/// </summary>
-		/// <param name="id">Unique id for this agent</param>
-		public CreateAgentResponse(string id)
-		{
-			Id = id;
-		}
-	}
-
 	/// <summary>
 	/// Parameters to update an agent
 	/// </summary>
diff --git a/Engine/Source/Programs/Horde/Horde.Server/Agents/AgentsController.cs b/Engine/Source/Programs/Horde/Horde.Server/Agents/AgentsController.cs
--- a/Engine/Source/Programs/Horde/Horde.Server/Agents/AgentsController.cs
+++ b/Engine/Source/Programs/Horde/Horde.Server/Agents/AgentsController.cs
@@ -39,25 +39,6 @@ namespace Horde.Server.Agents
 			_globalConfig = globalConfig;
 		}
 
-		/// <summary>
-		/// Register an agent to perform remote work.
-		/// </summary>
-		/// <param name="request">Request parameters</param>
-		/// <returns>Information about the registered agent</returns>
-		[HttpPost]
-		[Route("/api/v1/agents")]
-		public async Task<ActionResult<CreateAgentResponse>> CreateAgentAsync([FromBody] CreateAgentRequest request)
-		{
-			if(!_globalConfig.Value.Authorize(AgentAclAction.CreateAgent, User))
-			{
-				return Forbid(AgentAclAction.CreateAgent);
-			}
-
-			IAgent agent = await _agentService.CreateAgentAsync(request.Name, request.Enabled, request.Pools?.ConvertAll(x => new PoolId(x)));
-
-			return new CreateAgentResponse(agent.Id.ToString());
-		}
-
 		/// <summary>
 		/// Finds the agents matching specified criteria.
 		/// </summary>
diff --git a/Engine/Source/Programs/Horde/Horde.Server/Server/RpcService.cs b/Engine/Source/Programs/Horde/Horde.Server/Server/RpcService.cs
--- a/Engine/Source/Programs/Horde/Horde.Server/Server/RpcService.cs
+++ b/Engine/Source/Programs/Horde/Horde.Server/Server/RpcService.cs
@@ -6,6 +6,7 @@ using System.Collections.Generic;
 using System.IO;
 using System.Linq;
 using System.Net;
+using System.Security.Claims;
 using System.Threading;
 using System.Threading.Tasks;
 using Amazon.EC2.Model;
@@ -56,13 +57,14 @@ namespace Horde.Server.Server
 		readonly ConformTaskSource _conformTaskSource;
 		readonly JobRpcCommon _jobRpcCommon;
 		readonly IToolCollection _toolCollection;
+		readonly AclService _aclService;
 		readonly IOptionsSnapshot<GlobalConfig> _globalConfig;
 		readonly ILogger _logger;
 
 		/// <summary>
 		/// Constructor
 		/// </summary>
-		public RpcService(AgentService agentService, PoolService poolService, LifetimeService lifetimeService, ITelemetrySink telemetrySink, ConformTaskSource conformTaskSource, JobRpcCommon jobRpcCommon, IToolCollection toolCollection, IOptionsSnapshot<GlobalConfig> globalConfig, ILogger<RpcService> logger)
+		public RpcService(AgentService agentService, PoolService poolService, LifetimeService lifetimeService, ITelemetrySink telemetrySink, ConformTaskSource conformTaskSource, JobRpcCommon jobRpcCommon, IToolCollection toolCollection, AclService aclService, IOptionsSnapshot<GlobalConfig> globalConfig, ILogger<RpcService> logger)
 		{
 			_agentService = agentService;
 			_poolService = poolService;
@@ -71,6 +73,7 @@ namespace Horde.Server.Server
 			_conformTaskSource = conformTaskSource;
 			_jobRpcCommon = jobRpcCommon;
 			_toolCollection = toolCollection;
+			_aclService = aclService;
 			_globalConfig = globalConfig;
 			_logger = logger;
 		}
@@ -244,6 +247,39 @@ namespace Horde.Server.Server
 			CopyPropertyToResource(KnownPropertyNames.Ram, properties, resources);
 		}
 
+		/// <summary>
+		/// Creates a new agent
+		/// </summary>
+		/// <param name="request">Request to create a new agent</param>
+		/// <param name="context">Context for the RPC call</param>
+		/// <returns>Information about the new agent</returns>
+		public override async Task<CreateAgentResponse> CreateAgent(CreateAgentRequest request, ServerCallContext context)
+		{
+			using IDisposable scope = _logger.BeginScope("CreateAgent({AgentId})", request.Name.ToString());
+
+			if (!_globalConfig.Value.Authorize(AgentAclAction.CreateAgent, context.GetHttpContext().User))
+			{
+				throw new StructuredRpcException(StatusCode.PermissionDenied, "User is not authenticated to create new agents");
+			}
+
+			// TODO: Do not allow overwriting existing agents unless explicitly asked to. We may be reusing an IP or hostname, and we can't trust the agent is what it says it is.
+			IAgent? agent = await _agentService.GetAgentAsync(new AgentId(request.Name));
+			if (agent == null)
+			{
+				// TODO: Once we have this flow nailed down, do not enable agents until approved for joining the farm.
+				agent = await _agentService.CreateAgentAsync(request.Name, true, null);
+			}
+
+			List<AclClaimConfig> claims = new List<AclClaimConfig>();
+			claims.Add(new AclClaimConfig(HordeClaimTypes.Agent, agent.Id.ToString()));
+
+			CreateAgentResponse response = new CreateAgentResponse();
+			response.Id = agent.Id.ToString();
+			response.Token = await _aclService.IssueBearerTokenAsync(claims, null);
+
+			return response;
+		}
+
 		/// <summary>
 		/// Creates a new session
 		/// </summary>
@@ -257,7 +293,7 @@ namespace Horde.Server.Server
 				throw new StructuredRpcException(StatusCode.InvalidArgument, "Capabilities may not be null");
 			}
 
-			AgentId agentId = new AgentId(request.Name);
+			AgentId agentId = new AgentId(request.Id);
 			using IDisposable scope = _logger.BeginScope("CreateSession({AgentId})", agentId.ToString());
 
 			GlobalConfig globalConfig = _globalConfig.Value;
@@ -271,13 +307,14 @@ namespace Horde.Server.Server
 					throw new StructuredRpcException(StatusCode.PermissionDenied, "User is not authenticated to create new agents");
 				}
 
-				agent = await _agentService.CreateAgentAsync(request.Name, true, null);
+				agent = await _agentService.CreateAgentAsync(request.Id, true, null);
 			}
 
 			// Make sure we're allowed to create sessions on this agent
-			if (!globalConfig.Authorize(SessionAclAction.CreateSession, context.GetHttpContext().User))
+			ClaimsPrincipal user = context.GetHttpContext().User;
+			if (!globalConfig.Authorize(SessionAclAction.CreateSession, user) && !user.HasAgentClaim(agentId))
 			{
-				throw new StructuredRpcException(StatusCode.PermissionDenied, "User is not authenticated to create session for {AgentId}", request.Name);
+				throw new StructuredRpcException(StatusCode.PermissionDenied, "User is not authenticated to create session for {AgentId}", agentId);
 			}
 
 			// Get the known properties for this agent
@@ -287,7 +324,7 @@ namespace Horde.Server.Server
 			agent = await _agentService.CreateSessionAsync(agent, request.Status, properties, resources, request.Version);
 			if (agent == null)
 			{
-				throw new StructuredRpcException(StatusCode.NotFound, "Agent {AgentId} not found", request.Name);
+				throw new StructuredRpcException(StatusCode.NotFound, "Agent {AgentId} not found", agentId);
 			}
 
 			// Create the response
diff --git a/Engine/Source/Programs/Horde/Horde.Server/Utilities/HordeClaimTypes.cs b/Engine/Source/Programs/Horde/Horde.Server/Utilities/HordeClaimTypes.cs
--- a/Engine/Source/Programs/Horde/Horde.Server/Utilities/HordeClaimTypes.cs
+++ b/Engine/Source/Programs/Horde/Horde.Server/Utilities/HordeClaimTypes.cs
@@ -22,9 +22,9 @@ namespace Horde.Server.Utilities
 		public const string Role = Prefix + "role";
 
 		/// <summary>
-		/// Claim for a particular session
+		/// Claim for a particular agent
 		/// </summary>
-		public const string AgentId = Prefix + "agent";
+		public const string Agent = Prefix + "agent";
 
 		/// <summary>
 		/// Claim for a particular session
diff --git a/Engine/Source/Programs/Horde/Horde.Server/Utilities/HordeClaims.cs b/Engine/Source/Programs/Horde/Horde.Server/Utilities/HordeClaims.cs
--- a/Engine/Source/Programs/Horde/Horde.Server/Utilities/HordeClaims.cs
+++ b/Engine/Source/Programs/Horde/Horde.Server/Utilities/HordeClaims.cs
@@ -53,7 +53,7 @@ namespace Horde.Server.Utilities
 		/// <returns>New claim instance</returns>
 		public static AclClaimConfig GetAgentClaim(AgentId agentId)
 		{
-			return new AclClaimConfig(HordeClaimTypes.AgentId, agentId.ToString());
+			return new AclClaimConfig(HordeClaimTypes.Agent, agentId.ToString());
 		}
 
 		/// <summary>
