Horde: Fix test error.

diff --git a/Engine/Source/Programs/Horde/Horde.Agent.Tests/WorkerServiceTest.cs b/Engine/Source/Programs/Horde/Horde.Agent.Tests/WorkerServiceTest.cs
--- a/Engine/Source/Programs/Horde/Horde.Agent.Tests/WorkerServiceTest.cs
+++ b/Engine/Source/Programs/Horde/Horde.Agent.Tests/WorkerServiceTest.cs
@@ -235,7 +235,7 @@ namespace Horde.Agent.Tests
 			});
 
 			_serviceCollection.AddSingleton<IJobExecutorFactory>(x => new SimpleTestExecutorFactory(executor));
-			using ServiceProvider serviceProvider = _serviceCollection.BuildServiceProvider();
+			await using ServiceProvider serviceProvider = _serviceCollection.BuildServiceProvider();
 
 			using CancellationTokenSource cts = new ();
 			cts.CancelAfter(20000);
