Horde: Query the server for the HTTP2 port to use via a HTTP1 REST call if using unencrypted sockets. This allows specifying the regular server URL when configuring agents.

Also add support for multiple server profiles in the saved credentials list.

diff --git a/Engine/Source/Programs/Horde/Horde.Agent/Commands/Execution/ExecuteJobCommand.cs b/Engine/Source/Programs/Horde/Horde.Agent/Commands/Execution/ExecuteJobCommand.cs
--- a/Engine/Source/Programs/Horde/Horde.Agent/Commands/Execution/ExecuteJobCommand.cs
+++ b/Engine/Source/Programs/Horde/Horde.Agent/Commands/Execution/ExecuteJobCommand.cs
@@ -47,7 +47,7 @@ namespace Horde.Agent.Commands.Execution
 			ServerProfile serverProfile = _settings.GetCurrentServerProfile();
 			ExecuteJobTask executeTask = ExecuteJobTask.Parser.ParseFrom(Convert.FromBase64String(Task));
 
-			await using RpcConnection rpcConnection = new RpcConnection(() => _grpcService.CreateGrpcChannel(executeTask.Token), logger);
+			await using RpcConnection rpcConnection = new RpcConnection(ctx => _grpcService.CreateGrpcChannelAsync(executeTask.Token, ctx), logger);
 			await using Session session = new Session(serverProfile.Url, AgentId, SessionId, executeTask.Token, rpcConnection, WorkingDir, _settings.ProcessNamesToTerminate, logger);
 
 			await _jobHandler.ExecuteInternalAsync(session, LeaseId, executeTask, CancellationToken.None);
diff --git a/Engine/Source/Programs/Horde/Horde.Agent/Services/GrpcService.cs b/Engine/Source/Programs/Horde/Horde.Agent/Services/GrpcService.cs
--- a/Engine/Source/Programs/Horde/Horde.Agent/Services/GrpcService.cs
+++ b/Engine/Source/Programs/Horde/Horde.Agent/Services/GrpcService.cs
@@ -4,6 +4,9 @@ using System;
 using System.Net.Http;
 using System.Net.Http.Headers;
 using System.Net.Security;
+using System.Net.Http.Json;
+using System.Threading;
+using System.Threading.Tasks;
 using Grpc.Net.Client;
 using Horde.Agent.Utility;
 using Microsoft.Extensions.Logging;
@@ -16,10 +19,20 @@ namespace Horde.Agent.Services
 	/// </summary>
 	class GrpcService
 	{
+		class GetPortsResponse
+		{
+			public int UnencryptedHttp2 { get; set; }
+		}
+
 		private readonly ServerProfile _serverProfile;
 		private readonly ILogger _logger;
 		private readonly ILoggerFactory _loggerFactory;
 
+		/// <summary>
+		/// The current server profile
+		/// </summary>
+		public ServerProfile ServerProfile => _serverProfile;
+
 		/// <summary>
 		/// Constructor
 		/// </summary>
@@ -37,16 +50,16 @@ namespace Horde.Agent.Services
 		/// Create a GRPC channel with a default token
 		/// </summary>
 		/// <returns>New grpc channel</returns>
-		public GrpcChannel CreateGrpcChannel()
+		public Task<GrpcChannel> CreateGrpcChannelAsync(CancellationToken cancellationToken)
 		{
-			return CreateGrpcChannel(_serverProfile.Token);
+			return CreateGrpcChannelAsync(_serverProfile.Token, cancellationToken);
 		}
 		
 		/// <summary>
 		/// Create a GRPC channel with the given bearer token
 		/// </summary>
 		/// <returns>New grpc channel</returns>
-		public GrpcChannel CreateGrpcChannel(string? bearerToken)
+		public async Task<GrpcChannel> CreateGrpcChannelAsync(string? bearerToken, CancellationToken cancellationToken)
 		{
 #pragma warning disable CA2000 // Dispose objects before losing scope
 			// HTTP handler is disposed by GrpcChannel below
@@ -69,8 +82,26 @@ namespace Horde.Agent.Services
 
 			httpClient.Timeout = TimeSpan.FromSeconds(210); // Need to make sure this doesn't cancel any long running gRPC streaming calls (eg. session update)
 
-			_logger.LogInformation("Connecting to rpc server {BaseUrl}", _serverProfile.Url);
-			return GrpcChannel.ForAddress(_serverProfile.Url, new GrpcChannelOptions
+			// Get the server URL for gRPC traffic. If we're using an unencrpyted connection we need to use a different port for http/2, so 
+			// send a http1 request to the server to query it.
+			Uri serverUri = _serverProfile.Url;
+			if (serverUri.Scheme.Equals("http", StringComparison.Ordinal))
+			{
+				_logger.LogInformation("Querying server {BaseUrl} for rpc port", serverUri);
+				using (HttpResponseMessage response = await httpClient.GetAsync(new Uri(serverUri, "api/v1/server/ports"), cancellationToken))
+				{
+					GetPortsResponse? ports = await response.Content.ReadFromJsonAsync<GetPortsResponse>(Program.DefaultJsonSerializerOptions, cancellationToken);
+					if (ports != null && ports.UnencryptedHttp2 != 0)
+					{
+						UriBuilder builder = new UriBuilder(serverUri);
+						builder.Port = ports.UnencryptedHttp2;
+						serverUri = builder.Uri;
+					}
+				}
+			}
+
+			_logger.LogInformation("Connecting to rpc server {BaseUrl}", serverUri);
+			return GrpcChannel.ForAddress(serverUri, new GrpcChannelOptions
 			{
 				// Required payloads coming from CAS service can be large
 				MaxReceiveMessageSize = 1024 * 1024 * 1024, // 1 GB
diff --git a/Engine/Source/Programs/Horde/Horde.Agent/Services/SessionFactory.cs b/Engine/Source/Programs/Horde/Horde.Agent/Services/SessionFactory.cs
--- a/Engine/Source/Programs/Horde/Horde.Agent/Services/SessionFactory.cs
+++ b/Engine/Source/Programs/Horde/Horde.Agent/Services/SessionFactory.cs
@@ -3,6 +3,8 @@
 using System;
 using System.Collections.Generic;
 using System.Linq;
+using System.Net.Http;
+using System.Net.Http.Json;
 using System.Reflection;
 using System.Runtime.InteropServices;
 using System.Text.Json;
@@ -129,16 +131,17 @@ namespace Horde.Agent.Services
 			_logger = logger;
 		}
 
-		class AgentRegistration
+		public class AgentRegistrationList
 		{
-			public string? Id { get; set; }
-			public string? Token { get; set; }
+			public List<AgentRegistration> Entries { get; set; } = new List<AgentRegistration>();
 		}
 
+		public record class AgentRegistration(Uri Server, string Id, string Token);
+
 		/// <summary>
 		/// Creates a new agent session
 		/// </summary>
-		public static async Task<Session> CreateAsync(CapabilitiesService capabilitiesService, GrpcService grpcService, StatusService statusService, IOptions<AgentSettings> settings, ILogger logger, CancellationToken cancellationToken)
+		public static async Task<Session> CreateAsync(CapabilitiesService capabilitiesService, GrpcService grpcService, IOptions<AgentSettings> settings, ILogger logger, CancellationToken cancellationToken)
 		{
 			AgentSettings currentSettings = settings.Value;
 
@@ -194,7 +197,7 @@ namespace Horde.Agent.Services
 
 			// Create the session
 			CreateSessionResponse createSessionResponse;
-			using (GrpcChannel channel = grpcService.CreateGrpcChannel(registrationInfo.Token))
+			using (GrpcChannel channel = await grpcService.CreateGrpcChannelAsync(registrationInfo.Token, cancellationToken))
 			{
 				HordeRpc.HordeRpcClient rpcClient = new HordeRpc.HordeRpcClient(channel);
 
@@ -210,11 +213,11 @@ namespace Horde.Agent.Services
 				logger.LogInformation("Created session. AgentName={AgentName} SessionId={SessionId}", currentSettings.GetAgentName(), createSessionResponse.SessionId);
 			}
 
-			Func<GrpcChannel> createGrpcChannel = () => grpcService.CreateGrpcChannel(createSessionResponse.Token);
+			Func<CancellationToken, Task<GrpcChannel>> createGrpcChannelAsync = ctx => grpcService.CreateGrpcChannelAsync(createSessionResponse.Token, ctx);
 
 			// Open a connection to the server
 #pragma warning disable CA2000 // False positive; ownership is transferred to new Session object.
-			IRpcConnection rpcConnection = new RpcConnection(createGrpcChannel, logger);
+			IRpcConnection rpcConnection = new RpcConnection(createGrpcChannelAsync, logger);
 			return new Session(serverProfile.Url, createSessionResponse.AgentId, createSessionResponse.SessionId, createSessionResponse.Token, rpcConnection, workingDir, currentSettings.ProcessNamesToTerminate, logger);
 #pragma warning restore CA2000
 		}
@@ -244,22 +247,25 @@ namespace Horde.Agent.Services
 				settingsDir = DirectoryReference.Combine(settingsDir, "Horde.Agent");
 			}
 
-			FileReference settingsFile = FileReference.Combine(settingsDir, "registration.json");
+			FileReference settingsFile = FileReference.Combine(settingsDir, "servers.json");
 
 			// Read existing settings if possible
-			AgentRegistration? registrationInfo = null;
+			AgentRegistrationList? registrationList = null;
 			if (FileReference.Exists(settingsFile))
 			{
 				byte[] settingsData = await FileReference.ReadAllBytesAsync(settingsFile, cancellationToken);
-				registrationInfo = JsonSerializer.Deserialize<AgentRegistration>(settingsData, Program.DefaultJsonSerializerOptions);
+				registrationList = JsonSerializer.Deserialize<AgentRegistrationList>(settingsData, Program.DefaultJsonSerializerOptions);
+				registrationList?.Entries.RemoveAll(x => x.Server == null || x.Id == null || x.Token == null);
 				logger.LogInformation("Read agent registration settings from {SettingsFile}", settingsFile);
 			}
+			registrationList ??= new AgentRegistrationList();
 
 			// If they aren't valid, create a new agent registration
-			if (registrationInfo == null || registrationInfo.Id == null || registrationInfo.Token == null)
+			AgentRegistration? registration = registrationList.Entries.FirstOrDefault(x => x.Server == grpcService.ServerProfile.Url);
+			if (registration == null)
 			{
 				ServerProfile serverProfile = currentSettings.GetCurrentServerProfile();
-				using (GrpcChannel channel = grpcService.CreateGrpcChannel(serverProfile.Token))
+				using (GrpcChannel channel = await grpcService.CreateGrpcChannelAsync(serverProfile.Token, cancellationToken))
 				{
 					HordeRpc.HordeRpcClient rpcClient = new HordeRpc.HordeRpcClient(channel);
 
@@ -267,14 +273,18 @@ namespace Horde.Agent.Services
 					createAgentRequest.Name = currentSettings.GetAgentName();
 
 					CreateAgentResponse createAgentResponse = await rpcClient.CreateAgentAsync(createAgentRequest, null, null, cancellationToken);
-					registrationInfo = new AgentRegistration { Id = createAgentResponse.Id, Token = createAgentResponse.Token };
-					JsonSerializer.SerializeToUtf8Bytes(registrationInfo, Program.DefaultJsonSerializerOptions);
+					registration = new AgentRegistration(grpcService.ServerProfile.Url, createAgentResponse.Id, createAgentResponse.Token);
+					registrationList.Entries.Add(registration);
+
+					byte[] data = JsonSerializer.SerializeToUtf8Bytes(registrationList, new JsonSerializerOptions(Program.DefaultJsonSerializerOptions) { WriteIndented = true });
+					DirectoryReference.CreateDirectory(settingsDir);
+					await FileReference.WriteAllBytesAsync(settingsFile, data, cancellationToken);
 
 					logger.LogInformation("Created agent (Id={AgentId}). Settings saved to {File}.", createAgentResponse.Id, settingsFile);
 				}
 			}
 
-			return registrationInfo;
+			return registration;
 		}
 
 		/// <inheritdoc/>
@@ -312,23 +322,21 @@ namespace Horde.Agent.Services
 	{
 		readonly CapabilitiesService _capabilitiesService;
 		readonly GrpcService _grpcService;
-		readonly StatusService _statusService;
 		readonly IOptions<AgentSettings> _settings;
 		readonly ILogger _logger;
 
 		/// <summary>
 		/// Constructor
 		/// </summary>
-		public SessionFactory(CapabilitiesService capabilitiesService, GrpcService grpcService, StatusService statusService, IOptions<AgentSettings> settings, ILogger<SessionFactory> logger)
+		public SessionFactory(CapabilitiesService capabilitiesService, GrpcService grpcService, IOptions<AgentSettings> settings, ILogger<SessionFactory> logger)
 		{
 			_capabilitiesService = capabilitiesService;
 			_grpcService = grpcService;
-			_statusService = statusService;
 			_settings = settings;
 			_logger = logger;
 		}
 
 		/// <inheritdoc/>
-		public async Task<ISession> CreateAsync(CancellationToken cancellationToken) => await Session.CreateAsync(_capabilitiesService, _grpcService, _statusService, _settings, _logger, cancellationToken);
+		public async Task<ISession> CreateAsync(CancellationToken cancellationToken) => await Session.CreateAsync(_capabilitiesService, _grpcService, _settings, _logger, cancellationToken);
 	}
 }
diff --git a/Engine/Source/Programs/Horde/Horde.Agent/Services/TelemetryService.cs b/Engine/Source/Programs/Horde/Horde.Agent/Services/TelemetryService.cs
--- a/Engine/Source/Programs/Horde/Horde.Agent/Services/TelemetryService.cs
+++ b/Engine/Source/Programs/Horde/Horde.Agent/Services/TelemetryService.cs
@@ -288,7 +288,7 @@ class TelemetryService : BackgroundService
 			return false;
 		}
 		
-		using GrpcChannel channel = _grpcService.CreateGrpcChannel();
+		using GrpcChannel channel = await _grpcService.CreateGrpcChannelAsync(stoppingToken);
 		HordeRpc.HordeRpcClient client = new (channel);
 		
 		while (!stoppingToken.IsCancellationRequested)
diff --git a/Engine/Source/Programs/Horde/Horde.Agent/Utility/RpcConnection.cs b/Engine/Source/Programs/Horde/Horde.Agent/Utility/RpcConnection.cs
--- a/Engine/Source/Programs/Horde/Horde.Agent/Utility/RpcConnection.cs
+++ b/Engine/Source/Programs/Horde/Horde.Agent/Utility/RpcConnection.cs
@@ -365,7 +365,7 @@ namespace Horde.Agent.Utility
 			public Task DisposingTask => _subConnection.DisposingTask;
 		}
 
-		private readonly Func<GrpcChannel> _createGrpcChannel;
+		private readonly Func<CancellationToken, Task<GrpcChannel>> _createGrpcChannelAsync;
 		private readonly TaskCompletionSource<bool> _stoppingTaskSource = new TaskCompletionSource<bool>();
 		private TaskCompletionSource<RpcSubConnection> _subConnectionTaskSource = new TaskCompletionSource<RpcSubConnection>();
 		private Task? _backgroundTask;
@@ -376,11 +376,11 @@ namespace Horde.Agent.Utility
 		/// <summary>
 		/// Constructor
 		/// </summary>
-		/// <param name="createGrpcChannel">Factory method for creating new GRPC channels</param>
+		/// <param name="createGrpcChannelAsync">Factory method for creating new GRPC channels</param>
 		/// <param name="logger">Logger instance</param>
-		public RpcConnection(Func<GrpcChannel> createGrpcChannel, ILogger logger)
+		public RpcConnection(Func<CancellationToken, Task<GrpcChannel>> createGrpcChannelAsync, ILogger logger)
 		{
-			_createGrpcChannel = createGrpcChannel;
+			_createGrpcChannelAsync = createGrpcChannelAsync;
 			_logger = logger;
 
 			_backgroundTask = Task.Run(() => ExecuteAsync());
@@ -531,7 +531,7 @@ namespace Horde.Agent.Utility
 		/// <returns>Async task</returns>
 		async Task HandleConnectionInternalAsync(int connectionId, TaskCompletionSource<bool> reconnectTaskSource)
 		{
-			using (GrpcChannel channel = _createGrpcChannel())
+			using (GrpcChannel channel = await _createGrpcChannelAsync(CancellationToken.None))
 			{
 				HordeRpc.HordeRpcClient client = new HordeRpc.HordeRpcClient(channel);
 
diff --git a/Engine/Source/Programs/Horde/Horde.Server/Commands/ServerCommand.cs b/Engine/Source/Programs/Horde/Horde.Server/Commands/ServerCommand.cs
--- a/Engine/Source/Programs/Horde/Horde.Server/Commands/ServerCommand.cs
+++ b/Engine/Source/Programs/Horde/Horde.Server/Commands/ServerCommand.cs
@@ -49,11 +49,15 @@ namespace Horde.Server.Commands
 
 		static IHostBuilder CreateHostBuilderWithCert(string[] args, IConfiguration config, ServerSettings serverSettings, X509Certificate2? sslCert)
 		{
+			AppContext.SetSwitch("System.Net.Http.SocketsHttpHandler.Http2UnencryptedSupport", true);
+
 			IHostBuilder hostBuilder = Host.CreateDefaultBuilder(args)
 				.UseSerilog()
 				.ConfigureAppConfiguration(builder => builder.AddConfiguration(config))
 				.ConfigureWebHostDefaults(webBuilder =>
 				{
+					webBuilder.UseUrls(); // Disable default URLs; we will configure each port directly.
+
 					webBuilder.ConfigureKestrel(options =>
 					{
 						options.Limits.MaxRequestBodySize = 100 * 1024 * 1024;
@@ -66,11 +70,15 @@ namespace Horde.Server.Commands
 						if (serverSettings.HttpsPort != 0)
 						{
 							options.ListenAnyIP(serverSettings.HttpsPort, configure => 
-							{ 
-								if (sslCert != null) 
-								{ 
-									configure.UseHttps(sslCert); 
-								} 
+							{
+								if (sslCert != null)
+								{
+									configure.UseHttps(sslCert);
+								}
+								else
+								{
+									configure.UseHttps();
+								}
 							});
 						}
 
diff --git a/Engine/Source/Programs/Horde/Horde.Server/Server/ServerController.cs b/Engine/Source/Programs/Horde/Horde.Server/Server/ServerController.cs
--- a/Engine/Source/Programs/Horde/Horde.Server/Server/ServerController.cs
+++ b/Engine/Source/Programs/Horde/Horde.Server/Server/ServerController.cs
@@ -7,6 +7,8 @@ using Horde.Server.Agents;
 using Horde.Server.Tools;
 using HordeCommon;
 using Microsoft.AspNetCore.Authorization;
+using Microsoft.AspNetCore.Hosting.Server;
+using Microsoft.AspNetCore.Hosting.Server.Features;
 using Microsoft.AspNetCore.Mvc;
 using Microsoft.Extensions.Options;
 
@@ -69,5 +71,22 @@ namespace Horde.Server.Server
 
 			return new GetServerInfoResponse(agentVersion);
 		}
+
+		/// <summary>
+		/// Gets ports used by the server
+		/// </summary>
+		[HttpGet]
+		[AllowAnonymous]
+		[Route("/api/v1/server/ports")]
+		public ActionResult<GetPortsResponse> GetPorts()
+		{
+			ServerSettings serverSettings = _globalConfig.Value.ServerSettings;
+
+			GetPortsResponse response = new GetPortsResponse();
+			response.Http = serverSettings.HttpPort;
+			response.Https = serverSettings.HttpsPort;
+			response.UnencryptedHttp2 = serverSettings.Http2Port;
+			return response;
+		}
 	}
 }
diff --git a/Engine/Source/Programs/Horde/Horde.Server/Server/ServerMessages.cs b/Engine/Source/Programs/Horde/Horde.Server/Server/ServerMessages.cs
--- a/Engine/Source/Programs/Horde/Horde.Server/Server/ServerMessages.cs
+++ b/Engine/Source/Programs/Horde/Horde.Server/Server/ServerMessages.cs
@@ -38,5 +38,26 @@ namespace Horde.Server.Server
 			OsDescription = RuntimeInformation.OSDescription;			
 		}
 	}
+
+	/// <summary>
+	/// Gets ports configured for this server
+	/// </summary>
+	public class GetPortsResponse
+	{
+		/// <summary>
+		/// Port for HTTP communication
+		/// </summary>
+		public int? Http { get; set; }
+
+		/// <summary>
+		/// Port number for HTTPS communication
+		/// </summary>
+		public int? Https { get; set; }
+
+		/// <summary>
+		/// Port number for unencrpyted HTTPS communication
+		/// </summary>
+		public int? UnencryptedHttp2 { get; set; }
+	}
 }
 
diff --git a/Engine/Source/Programs/Horde/Horde.Server/ServerSettings.cs b/Engine/Source/Programs/Horde/Horde.Server/ServerSettings.cs
--- a/Engine/Source/Programs/Horde/Horde.Server/ServerSettings.cs
+++ b/Engine/Source/Programs/Horde/Horde.Server/ServerSettings.cs
@@ -409,12 +409,12 @@ namespace Horde.Server
 		/// <summary>
 		/// Main port for serving HTTP. Uses the default Kestrel port (5000) if not specified.
 		/// </summary>
-		public int HttpPort { get; set; }
+		public int HttpPort { get; set; } = 5000;
 
 		/// <summary>
 		/// Port for serving HTTP with TLS enabled. Uses the default Kestrel port (5001) if not specified.
 		/// </summary>
-		public int HttpsPort { get; set; }
+		public int HttpsPort { get; set; } = 5001;
 
 		/// <summary>
 		/// Dedicated port for serving only HTTP/2.
