Horde: Remove flag for unencrypted HTTP 2 support. This is not needed in NET 6.

diff --git a/Engine/Source/Programs/Horde/Horde.Agent/Program.cs b/Engine/Source/Programs/Horde/Horde.Agent/Program.cs
--- a/Engine/Source/Programs/Horde/Horde.Agent/Program.cs
+++ b/Engine/Source/Programs/Horde/Horde.Agent/Program.cs
@@ -139,9 +139,6 @@ namespace Horde.Agent
 			services.AddCommandsFromAssembly(Assembly.GetExecutingAssembly());
 			services.AddSingleton(loggerFactory);
 
-			// Enable unencrypted HTTP/2 for gRPC channel without TLS
-			AppContext.SetSwitch("System.Net.Http.SocketsHttpHandler.Http2UnencryptedSupport", true);
-
 			if (RuntimeInformation.IsOSPlatform(OSPlatform.Windows))
 			{
 				// Prioritize agent execution time over any job its running.
