Horde: Cap the maximum length of Perforce changelist descriptions at 128k characters.

diff --git a/Engine/Source/Programs/Horde/Horde.Server/Perforce/PerforceServiceCache.cs b/Engine/Source/Programs/Horde/Horde.Server/Perforce/PerforceServiceCache.cs
--- a/Engine/Source/Programs/Horde/Horde.Server/Perforce/PerforceServiceCache.cs
+++ b/Engine/Source/Programs/Horde/Horde.Server/Perforce/PerforceServiceCache.cs
@@ -135,6 +135,13 @@ namespace Horde.Server.Perforce
 				AuthorId = commit.AuthorId;
 				OwnerId = commit.OwnerId;
 				Description = commit.Description;
+
+				const int MaxDescriptionLength = 128 * 1024; // 128k characters is a pretty big limit, but we've seen automated processes submit changes with larger :(
+				if (Description.Length > MaxDescriptionLength)
+				{
+					Description = Description.Substring(0, MaxDescriptionLength);
+				}
+
 				BasePath = commit.BasePath;
 				DateUtc = commit.DateUtc;
 
