Horde: Force refresh of HTTP client connections including DNS name

By default, the pooled connection lifespan is infinite.

diff --git a/Engine/Source/Programs/Horde/Horde.Agent/Services/GrpcService.cs b/Engine/Source/Programs/Horde/Horde.Agent/Services/GrpcService.cs
--- a/Engine/Source/Programs/Horde/Horde.Agent/Services/GrpcService.cs
+++ b/Engine/Source/Programs/Horde/Horde.Agent/Services/GrpcService.cs
@@ -66,6 +66,13 @@ namespace Horde.Agent.Services
 			SocketsHttpHandler httpHandler = new()
 			{
 				KeepAlivePingDelay = TimeSpan.FromSeconds(50),
+				ConnectTimeout = TimeSpan.FromSeconds(20),
+				
+				// Certain load balancers, like AWS ALB, can rotate IPs quickly.
+				// By setting a short lifetime, new connections will be re-created and subsequently have the DNS name refreshed.
+				// This prevents using outdated IP addresses.
+				PooledConnectionLifetime = TimeSpan.FromSeconds(60),
+				
 				SslOptions = new SslClientAuthenticationOptions
 				{
 					RemoteCertificateValidationCallback = (sender, cert, chain, errors) => CertificateHelper.CertificateValidationCallBack(_logger, sender, cert, chain, errors, _serverProfile)
