Horde: Enable multiple HTTP/2 connections for agent's gRPC client

diff --git a/Engine/Source/Programs/Horde/Horde.Agent/Services/GrpcService.cs b/Engine/Source/Programs/Horde/Horde.Agent/Services/GrpcService.cs
--- a/Engine/Source/Programs/Horde/Horde.Agent/Services/GrpcService.cs
+++ b/Engine/Source/Programs/Horde/Horde.Agent/Services/GrpcService.cs
@@ -66,7 +66,9 @@ namespace Horde.Agent.Services
 			SocketsHttpHandler httpHandler = new()
 			{
 				KeepAlivePingDelay = TimeSpan.FromSeconds(50),
-				ConnectTimeout = TimeSpan.FromSeconds(20),
+				ConnectTimeout = TimeSpan.FromSeconds(10),
+				MaxConnectionsPerServer = 20,
+				EnableMultipleHttp2Connections = true,
 				
 				// Certain load balancers, like AWS ALB, can rotate IPs quickly.
 				// By setting a short lifetime, new connections will be re-created and subsequently have the DNS name refreshed.
