Horde: Tweak Kestrel limits to avoid gRPC errors in agent to server communication

diff --git a/Engine/Source/Programs/Horde/Horde.Server/Commands/ServerCommand.cs b/Engine/Source/Programs/Horde/Horde.Server/Commands/ServerCommand.cs
--- a/Engine/Source/Programs/Horde/Horde.Server/Commands/ServerCommand.cs
+++ b/Engine/Source/Programs/Horde/Horde.Server/Commands/ServerCommand.cs
@@ -61,6 +61,14 @@ namespace Horde.Server.Commands
 					webBuilder.ConfigureKestrel(options =>
 					{
 						options.Limits.MaxRequestBodySize = 100 * 1024 * 1024;
+						options.Limits.Http2.KeepAlivePingDelay = TimeSpan.FromSeconds(30);
+						
+						// When agents are saturated with work (CPU or I/O), slow sending of gRPC data can happen.
+						// Kestrel protects against this behavior by default as it's commonly used for malicious attacks.
+						// Setting a more generous data rate should prevent incoming HTTP connections from being closed prematurely.
+						options.Limits.MinRequestBodyDataRate = new MinDataRate(10, TimeSpan.FromSeconds(60));
+
+						options.Limits.KeepAliveTimeout = TimeSpan.FromSeconds(220); // 10 seconds more than agent's timeout
 
 						if (serverSettings.HttpPort != 0)
 						{
