Horde: Add a gRPC call invoker that can enrich outgoing requests with agent version and name

diff --git a/Engine/Source/Programs/Horde/Horde.Agent/Services/GrpcService.cs b/Engine/Source/Programs/Horde/Horde.Agent/Services/GrpcService.cs
--- a/Engine/Source/Programs/Horde/Horde.Agent/Services/GrpcService.cs
+++ b/Engine/Source/Programs/Horde/Horde.Agent/Services/GrpcService.cs
@@ -7,6 +7,8 @@ using System.Net.Security;
 using System.Net.Http.Json;
 using System.Threading;
 using System.Threading.Tasks;
+using Grpc.Core;
+using Grpc.Core.Interceptors;
 using Grpc.Net.Client;
 using Horde.Agent.Utility;
 using Microsoft.Extensions.Logging;
@@ -24,6 +26,7 @@ namespace Horde.Agent.Services
 			public int UnencryptedHttp2 { get; set; }
 		}
 
+		private readonly IOptions<AgentSettings> _settings;
 		private readonly ServerProfile _serverProfile;
 		private readonly ILogger _logger;
 		private readonly ILoggerFactory _loggerFactory;
@@ -41,6 +44,7 @@ namespace Horde.Agent.Services
 		/// <param name="loggerFactory"></param>
 		public GrpcService(IOptions<AgentSettings> settings, ILogger<GrpcService> logger, ILoggerFactory loggerFactory)
 		{
+			_settings = settings;
 			_serverProfile = settings.Value.GetCurrentServerProfile();
 			_logger = logger;
 			_loggerFactory = loggerFactory;
@@ -120,5 +124,23 @@ namespace Horde.Agent.Services
 				DisposeHttpClient = true
 			});
 		}
+
+		/// <summary>
+		/// Get a gRPC call invoker for the given channel with extra metadata attached,
+		/// such as current version and name
+		/// </summary>
+		/// <param name="channel">gRPC channel to use</param>
+		/// <returns>A call invoker</returns>
+		public CallInvoker GetInvoker(GrpcChannel channel)
+		{
+			CallInvoker invoker = channel.Intercept(headers =>
+			{
+				headers.Add("Horde-Agent-Version", Program.Version);
+				headers.Add("Horde-Agent-Name", _settings.Value.GetAgentName());
+				return headers;
+			});
+			
+			return invoker;
+		}
 	}
 }
diff --git a/Engine/Source/Programs/Horde/Horde.Agent/Services/TelemetryService.cs b/Engine/Source/Programs/Horde/Horde.Agent/Services/TelemetryService.cs
--- a/Engine/Source/Programs/Horde/Horde.Agent/Services/TelemetryService.cs
+++ b/Engine/Source/Programs/Horde/Horde.Agent/Services/TelemetryService.cs
@@ -289,7 +289,8 @@ class TelemetryService : BackgroundService
 		}
 		
 		using GrpcChannel channel = await _grpcService.CreateGrpcChannelAsync(stoppingToken);
-		HordeRpc.HordeRpcClient client = new (channel);
+		CallInvoker invoker = _grpcService.GetInvoker(channel);
+		HordeRpc.HordeRpcClient client = new (invoker);
 		
 		while (!stoppingToken.IsCancellationRequested)
 		{
