Horde: Remove the 'Alias' entry for exports in bundles.

This is an application-level thing, and can be added separately through the API without baking into the data. Not widely implemented yet.

diff --git a/Engine/Source/Programs/Horde/Horde.Server.Tests/StorageServiceTests.cs b/Engine/Source/Programs/Horde/Horde.Server.Tests/StorageServiceTests.cs
--- a/Engine/Source/Programs/Horde/Horde.Server.Tests/StorageServiceTests.cs
+++ b/Engine/Source/Programs/Horde/Horde.Server.Tests/StorageServiceTests.cs
@@ -36,14 +36,18 @@ namespace Horde.Server.Tests
 
 			List<BundleExport> exports = new List<BundleExport>();
 
-			exports.Add(new BundleExport(0, hash1, data1.Length, Array.Empty<int>(), "foo"));
-			exports.Add(new BundleExport(0, hash1, data1.Length, Array.Empty<int>(), "foo"));
-			exports.Add(new BundleExport(0, hash2, data2.Length, Array.Empty<int>(), "bar"));
+			exports.Add(new BundleExport(0, hash1, data1.Length, Array.Empty<int>()));
+			exports.Add(new BundleExport(0, hash1, data1.Length, Array.Empty<int>()));
+			exports.Add(new BundleExport(0, hash2, data2.Length, Array.Empty<int>()));
 
 			BundleHeader header = new BundleHeader(BundleCompressionFormat.None, types, Array.Empty<BundleImport>(), exports, Array.Empty<BundlePacket>());
 			Bundle bundle = new Bundle(header, Array.Empty<ReadOnlyMemory<byte>>());
 			BlobLocator locator = await client.WriteBundleAsync(bundle);
 
+			await client.AddAliasAsync("foo", new NodeHandle(hash1, locator, 0));
+			await client.AddAliasAsync("foo", new NodeHandle(hash1, locator, 1));
+			await client.AddAliasAsync("bar", new NodeHandle(hash2, locator, 2));
+
 			List<NodeHandle> handles;
 			
 			handles = await client.FindNodesAsync("foo").ToListAsync();
diff --git a/Engine/Source/Programs/Horde/Horde.Server/Commands/Bundles/BundleWriteTest.cs b/Engine/Source/Programs/Horde/Horde.Server/Commands/Bundles/BundleWriteTest.cs
--- a/Engine/Source/Programs/Horde/Horde.Server/Commands/Bundles/BundleWriteTest.cs
+++ b/Engine/Source/Programs/Horde/Horde.Server/Commands/Bundles/BundleWriteTest.cs
@@ -22,6 +22,8 @@ namespace Horde.Server.Commands.Bundles
 			public Task DeleteRefAsync(RefName name, CancellationToken cancellationToken = default) => Task.CompletedTask;
 			public Task<Stream> ReadBlobAsync(BlobLocator locator, CancellationToken cancellationToken = default) => throw new NotImplementedException();
 			public Task<Stream> ReadBlobRangeAsync(BlobLocator locator, int offset, int length, CancellationToken cancellationToken = default) => throw new NotImplementedException();
+			public Task AddAliasAsync(Utf8String name, NodeHandle handle, CancellationToken cancellationToken = default) => throw new NotImplementedException();
+			public Task RemoveAliasAsync(Utf8String name, NodeHandle handle, CancellationToken cancellationToken = default) => throw new NotImplementedException();
 			public IAsyncEnumerable<NodeHandle> FindNodesAsync(Utf8String alias, CancellationToken cancellationToken = default) => throw new NotImplementedException();
 			public Task<NodeHandle?> TryReadRefTargetAsync(RefName name, DateTime cacheTime = default, CancellationToken cancellationToken = default) => throw new NotImplementedException();
 			public Task<BlobLocator> WriteBlobAsync(Stream stream, Utf8String prefix = default, CancellationToken cancellationToken = default) => Task.FromResult(BlobLocator.Create(HostId.Empty));
diff --git a/Engine/Source/Programs/Horde/Horde.Server/Storage/StorageService.cs b/Engine/Source/Programs/Horde/Horde.Server/Storage/StorageService.cs
--- a/Engine/Source/Programs/Horde/Horde.Server/Storage/StorageService.cs
+++ b/Engine/Source/Programs/Horde/Horde.Server/Storage/StorageService.cs
@@ -162,18 +162,8 @@ namespace Horde.Server.Storage
 						memoryStream.Position = 0;
 						BundleHeader header = await BundleHeader.FromStreamAsync(memoryStream, cancellationToken);
 
-						List<ExportInfo> exports = new List<ExportInfo>();
-						for (int idx = 0; idx < header.Exports.Count; idx++)
-						{
-							BundleExport export = header.Exports[idx];
-							if (!export.Alias.IsEmpty)
-							{
-								exports.Add(new ExportInfo(export.Alias.ToString(), export.Hash, idx));
-							}
-						}
-
 						// Add the blob record
-						locator = await _outer.AddBlobAsync(NamespaceId, prefix, exports, cancellationToken);
+						locator = await _outer.AddBlobAsync(NamespaceId, prefix, null, cancellationToken);
 
 						// Write it to the backend
 						string path = GetBlobPath(locator.BlobId);
@@ -223,6 +213,12 @@ namespace Horde.Server.Storage
 
 			#region Nodes
 
+			/// <inheritdoc/>
+			public override Task AddAliasAsync(Utf8String name, NodeHandle locator, CancellationToken cancellationToken = default) => _outer.AddAliasAsync(NamespaceId, name, locator, cancellationToken);
+		
+			/// <inheritdoc/>
+			public override Task RemoveAliasAsync(Utf8String name, NodeHandle locator, CancellationToken cancellationToken = default) => _outer.RemoveAliasAsync(NamespaceId, name, locator, cancellationToken);
+
 			/// <inheritdoc/>
 			public override IAsyncEnumerable<NodeHandle> FindNodesAsync(Utf8String alias, CancellationToken cancellationToken = default) => _outer.FindNodesAsync(NamespaceId, alias, cancellationToken);
 
@@ -678,6 +674,45 @@ namespace Horde.Server.Storage
 
 		#region Nodes
 
+		/// <summary>
+		/// Adds a node alias
+		/// </summary>
+		/// <param name="namespaceId">Namespace to search</param>
+		/// <param name="alias">Alias for the node</param>
+		/// <param name="target">Target node for the alias</param>
+		/// <param name="cancellationToken">Cancellation token for the operation</param>
+		/// <returns>Sequence of thandles</returns>
+		async Task AddAliasAsync(NamespaceId namespaceId, Utf8String alias, NodeHandle target, CancellationToken cancellationToken = default)
+		{
+			BlobInfo? blobInfo = await _blobCollection.Find(x => x.NamespaceId == namespaceId && x.BlobId == target.Locator.Blob.BlobId).FirstOrDefaultAsync(cancellationToken);
+			if (blobInfo == null)
+			{
+				throw new KeyNotFoundException($"Missing blob {target.Locator.Blob}");
+			}
+
+			if (blobInfo.Exports != null && blobInfo.Exports.Any(x => x.Alias == alias && x.Index == target.Locator.ExportIdx))
+			{
+				return;
+			}
+
+			FilterDefinition<BlobInfo> filter = Builders<BlobInfo>.Filter.Expr(x => x.NamespaceId == blobInfo.NamespaceId && x.BlobId == target.Locator.Blob.BlobId);
+			UpdateDefinition<BlobInfo> update = Builders<BlobInfo>.Update.Push(x => x.Exports, new ExportInfo(alias.ToString(), target.Hash, target.Locator.ExportIdx));
+			await _blobCollection.UpdateOneAsync(filter, update, cancellationToken: cancellationToken);
+		}
+
+		/// <summary>
+		/// Removes a node alias
+		/// </summary>
+		/// <param name="namespaceId">Namespace to search</param>
+		/// <param name="alias">Alias for the node</param>
+		/// <param name="target">Target node for the alias</param>
+		/// <param name="cancellationToken">Cancellation token for the operation</param>
+		/// <returns>Sequence of thandles</returns>
+		Task RemoveAliasAsync(NamespaceId namespaceId, Utf8String alias, NodeHandle target, CancellationToken cancellationToken = default)
+		{
+			throw new NotSupportedException();
+		}
+
 		/// <summary>
 		/// Finds nodes with the given type and hash
 		/// </summary>
diff --git a/Engine/Source/Programs/Shared/EpicGames.Horde/Compute/ComputeStorageClient.cs b/Engine/Source/Programs/Shared/EpicGames.Horde/Compute/ComputeStorageClient.cs
--- a/Engine/Source/Programs/Shared/EpicGames.Horde/Compute/ComputeStorageClient.cs
+++ b/Engine/Source/Programs/Shared/EpicGames.Horde/Compute/ComputeStorageClient.cs
@@ -66,6 +66,18 @@ namespace EpicGames.Horde.Compute
 
 		#region Nodes
 
+		/// <inheritdoc/>
+		public override Task AddAliasAsync(Utf8String name, NodeHandle locator, CancellationToken cancellationToken = default)
+		{
+			throw new NotSupportedException();
+		}
+
+		/// <inheritdoc/>
+		public override Task RemoveAliasAsync(Utf8String name, NodeHandle locator, CancellationToken cancellationToken = default)
+		{
+			throw new NotSupportedException();
+		}
+
 		/// <inheritdoc/>
 		public override IAsyncEnumerable<NodeHandle> FindNodesAsync(Utf8String name, CancellationToken cancellationToken = default)
 		{
diff --git a/Engine/Source/Programs/Shared/EpicGames.Horde/Storage/Backends/FileStorageClient.cs b/Engine/Source/Programs/Shared/EpicGames.Horde/Storage/Backends/FileStorageClient.cs
--- a/Engine/Source/Programs/Shared/EpicGames.Horde/Storage/Backends/FileStorageClient.cs
+++ b/Engine/Source/Programs/Shared/EpicGames.Horde/Storage/Backends/FileStorageClient.cs
@@ -74,7 +74,19 @@ namespace EpicGames.Horde.Storage.Backends
 
 		#endregion
 
-		#region Nodes
+		#region Aliases
+
+		/// <inheritdoc/>
+		public override Task AddAliasAsync(Utf8String name, NodeHandle locator, CancellationToken cancellationToken = default)
+		{
+			throw new NotSupportedException("File storage client does not currently support aliases.");
+		}
+
+		/// <inheritdoc/>
+		public override Task RemoveAliasAsync(Utf8String name, NodeHandle locator, CancellationToken cancellationToken = default)
+		{
+			throw new NotSupportedException("File storage client does not currently support aliases.");
+		}
 
 		/// <inheritdoc/>
 		public override IAsyncEnumerable<NodeHandle> FindNodesAsync(Utf8String alias, CancellationToken cancellationToken = default)
diff --git a/Engine/Source/Programs/Shared/EpicGames.Horde/Storage/Backends/HttpStorageClient.cs b/Engine/Source/Programs/Shared/EpicGames.Horde/Storage/Backends/HttpStorageClient.cs
--- a/Engine/Source/Programs/Shared/EpicGames.Horde/Storage/Backends/HttpStorageClient.cs
+++ b/Engine/Source/Programs/Shared/EpicGames.Horde/Storage/Backends/HttpStorageClient.cs
@@ -191,6 +191,18 @@ namespace EpicGames.Horde.Storage.Backends
 
 		#region Nodes
 
+		/// <inheritdoc/>
+		public override Task AddAliasAsync(Utf8String name, NodeHandle locator, CancellationToken cancellationToken = default)
+		{
+			throw new NotSupportedException("Http storage client does not currently support aliases.");
+		}
+
+		/// <inheritdoc/>
+		public override Task RemoveAliasAsync(Utf8String name, NodeHandle locator, CancellationToken cancellationToken = default)
+		{
+			throw new NotSupportedException("Http storage client does not currently support aliases.");
+		}
+
 		/// <inheritdoc/>
 		public override async IAsyncEnumerable<NodeHandle> FindNodesAsync(Utf8String alias, [EnumeratorCancellation] CancellationToken cancellationToken = default)
 		{
diff --git a/Engine/Source/Programs/Shared/EpicGames.Horde/Storage/Backends/MemoryStorageClient.cs b/Engine/Source/Programs/Shared/EpicGames.Horde/Storage/Backends/MemoryStorageClient.cs
--- a/Engine/Source/Programs/Shared/EpicGames.Horde/Storage/Backends/MemoryStorageClient.cs
+++ b/Engine/Source/Programs/Shared/EpicGames.Horde/Storage/Backends/MemoryStorageClient.cs
@@ -72,22 +72,25 @@ namespace EpicGames.Horde.Storage.Backends
 			Bundle bundle = await Bundle.FromStreamAsync(stream, cancellationToken);
 			_blobs[locator] = bundle;
 
-			for (int idx = 0; idx < bundle.Header.Exports.Count; idx++)
-			{
-				BundleExport export = bundle.Header.Exports[idx];
-				if(!export.Alias.IsEmpty)
-				{
-					NodeHandle handle = new NodeHandle(export.Hash, locator, idx);
-					_exports.AddOrUpdate(export.Alias, _ => new ExportEntry(handle, null), (_, entry) => new ExportEntry(handle, entry));
-				}
-			}
-
 			return locator;
 		}
 
 		#endregion
 
-		#region Nodes
+		#region Aliases
+
+		/// <inheritdoc/>
+		public override Task AddAliasAsync(Utf8String name, NodeHandle handle, CancellationToken cancellationToken = default)
+		{
+			_exports.AddOrUpdate(name, _ => new ExportEntry(handle, null), (_, entry) => new ExportEntry(handle, entry));
+			return Task.CompletedTask;
+		}
+
+		/// <inheritdoc/>
+		public override Task RemoveAliasAsync(Utf8String name, NodeHandle handle, CancellationToken cancellationToken = default)
+		{
+			throw new NotSupportedException();
+		}
 
 		/// <inheritdoc/>
 		public override async IAsyncEnumerable<NodeHandle> FindNodesAsync(Utf8String alias, [EnumeratorCancellation] CancellationToken cancellationToken = default)
diff --git a/Engine/Source/Programs/Shared/EpicGames.Horde/Storage/Bundle.cs b/Engine/Source/Programs/Shared/EpicGames.Horde/Storage/Bundle.cs
--- a/Engine/Source/Programs/Shared/EpicGames.Horde/Storage/Bundle.cs
+++ b/Engine/Source/Programs/Shared/EpicGames.Horde/Storage/Bundle.cs
@@ -30,10 +30,15 @@ namespace EpicGames.Horde.Storage
 		Initial = 0,
 
 		/// <summary>
-		/// Added the <see cref="BundleExport.Alias"/> property
+		/// Added the BundleExport.Alias property
 		/// </summary>
 		ExportAliases = 1,
 
+		/// <summary>
+		/// Back out change to include aliases. Will likely do this through an API rather than baked into the data. 
+		/// </summary>
+		RemoveAliases = 2,
+
 		/// <summary>
 		/// Last item in the enum. Used for <see cref="Latest"/>
 		/// </summary>
@@ -523,21 +528,15 @@ namespace EpicGames.Horde.Storage
         /// </summary>
 		public IReadOnlyList<int> References { get; }
 
-		/// <summary>
-		/// Alias for the content of this node. Used for content addressing.
-		/// </summary>
-		public Utf8String Alias { get; }
-
 		/// <summary>
 		/// Constructor
 		/// </summary>
-		public BundleExport(int typeIdx, IoHash hash, int length, IReadOnlyList<int> references, Utf8String alias = default)
+		public BundleExport(int typeIdx, IoHash hash, int length, IReadOnlyList<int> references)
 		{
 			TypeIdx = typeIdx;
 			Hash = hash;
 			Length = length;
 			References = references;
-			Alias = alias;
 		}
 
 		/// <summary>
@@ -566,9 +565,9 @@ namespace EpicGames.Horde.Storage
 				References = references;
 			}
 
-			if (version >= BundleVersion.ExportAliases)
+			if (version == BundleVersion.ExportAliases)
 			{
-				Alias = reader.ReadUtf8String();
+				_ = reader.ReadUtf8String();
 			}
 		}
 
@@ -587,8 +586,6 @@ namespace EpicGames.Horde.Storage
 			{
 				writer.WriteUnsignedVarInt(References[idx]);
 			}
-
-			writer.WriteUtf8String(Alias);
 		}
 	}
 
diff --git a/Engine/Source/Programs/Shared/EpicGames.Horde/Storage/IStorageClient.cs b/Engine/Source/Programs/Shared/EpicGames.Horde/Storage/IStorageClient.cs
--- a/Engine/Source/Programs/Shared/EpicGames.Horde/Storage/IStorageClient.cs
+++ b/Engine/Source/Programs/Shared/EpicGames.Horde/Storage/IStorageClient.cs
@@ -109,7 +109,23 @@ namespace EpicGames.Horde.Storage
 
 		#endregion
 
-		#region Nodes
+		#region Aliases
+
+		/// <summary>
+		/// Adds an alias to a given node
+		/// </summary>
+		/// <param name="name">Alias for the node</param>
+		/// <param name="handle">Locator for the node</param>
+		/// <param name="cancellationToken">Cancellation token for the operation</param>
+		Task AddAliasAsync(Utf8String name, NodeHandle handle, CancellationToken cancellationToken = default);
+
+		/// <summary>
+		/// Removes an alias from a node
+		/// </summary>
+		/// <param name="name">Name of the alias</param>
+		/// <param name="handle">Locator for the node</param>
+		/// <param name="cancellationToken">Cancellation token for the operation</param>
+		Task RemoveAliasAsync(Utf8String name, NodeHandle handle, CancellationToken cancellationToken = default);
 
 		/// <summary>
 		/// Finds nodes with the given alias. Unlike refs, aliases do not serve as GC roots.
@@ -209,7 +225,13 @@ namespace EpicGames.Horde.Storage
 
 			#endregion
 
-			#region Nodes
+			#region Aliases
+
+			/// <inheritdoc/>
+			public Task AddAliasAsync(Utf8String name, NodeHandle handle, CancellationToken cancellationToken = default) => _inner.AddAliasAsync(name, handle, cancellationToken);
+
+			/// <inheritdoc/>
+			public Task RemoveAliasAsync(Utf8String name, NodeHandle handle, CancellationToken cancellationToken = default) => _inner.RemoveAliasAsync(name, handle, cancellationToken);
 
 			/// <inheritdoc/>
 			public IAsyncEnumerable<NodeHandle> FindNodesAsync(Utf8String name, CancellationToken cancellationToken = default) => _inner.FindNodesAsync(name, cancellationToken);
diff --git a/Engine/Source/Programs/Shared/EpicGames.Horde/Storage/StorageClientBase.cs b/Engine/Source/Programs/Shared/EpicGames.Horde/Storage/StorageClientBase.cs
--- a/Engine/Source/Programs/Shared/EpicGames.Horde/Storage/StorageClientBase.cs
+++ b/Engine/Source/Programs/Shared/EpicGames.Horde/Storage/StorageClientBase.cs
@@ -36,6 +36,12 @@ namespace EpicGames.Horde.Storage
 
 		#region Nodes
 
+		/// <inheritdoc/>
+		public abstract Task AddAliasAsync(Utf8String name, NodeHandle locator, CancellationToken cancellationToken = default);
+
+		/// <inheritdoc/>
+		public abstract Task RemoveAliasAsync(Utf8String name, NodeHandle locator, CancellationToken cancellationToken = default);
+
 		/// <inheritdoc/>
 		public abstract IAsyncEnumerable<NodeHandle> FindNodesAsync(Utf8String name, CancellationToken cancellationToken = default);
 
