Horde: Disable logging of cancellation exceptions from UpdateSession calls.

diff --git a/Engine/Source/Programs/Horde/Horde.Server/Server/RpcService.cs b/Engine/Source/Programs/Horde/Horde.Server/Server/RpcService.cs
--- a/Engine/Source/Programs/Horde/Horde.Server/Server/RpcService.cs
+++ b/Engine/Source/Programs/Horde/Horde.Server/Server/RpcService.cs
@@ -10,6 +10,7 @@ using System.Security.Claims;
 using System.Threading;
 using System.Threading.Tasks;
 using Amazon.EC2.Model;
+using EpicGames.Core;
 using Google.Protobuf;
 using Google.Protobuf.WellKnownTypes;
 using Grpc.Core;
@@ -412,9 +413,16 @@ namespace Horde.Server.Server
 				}
 
 				// Wait for the client to close the stream
-				while (await nextRequestTask)
+				try
+				{
+					while (await nextRequestTask)
+					{
+						nextRequestTask = reader.MoveNext();
+					}
+				}
+				catch (OperationCanceledException ex)
 				{
-					nextRequestTask = reader.MoveNext();
+					_logger.LogDebug(ex, "Ignoring cancellation exception while finishing UpdateSession request.");
 				}
 			}
 		}
