Horde: Fix error when caching a commit containing no files.

diff --git a/Engine/Source/Programs/Horde/Horde.Server/Perforce/PerforceService.cs b/Engine/Source/Programs/Horde/Horde.Server/Perforce/PerforceService.cs
--- a/Engine/Source/Programs/Horde/Horde.Server/Perforce/PerforceService.cs
+++ b/Engine/Source/Programs/Horde/Horde.Server/Perforce/PerforceService.cs
@@ -685,7 +685,7 @@ namespace Horde.Server.Perforce
 			Commit commit = await CreateCommitInternalAsync(streamConfig, describeRecord.Number, describeRecord.User, describeRecord.Description, describeRecord.Path, timeUtc, cancellationToken);
 
 			List<string> files = await perforce.GetStreamFilesAsync(streamConfig, describeRecord, cancellationToken);
-			if (files.Count == 0)
+			if (files.Count == 0 && describeRecord.Files.Count > 0)
 			{
 				throw new PerforceException($"Changelist {commit.Number} does not contain any files in {streamConfig.Id}");
 			}
