Horde: Disable HTTP/2 keep-alive pings in Kestrel

They're not supported by AWS ALB.

diff --git a/Engine/Source/Programs/Horde/Horde.Server/Commands/ServerCommand.cs b/Engine/Source/Programs/Horde/Horde.Server/Commands/ServerCommand.cs
--- a/Engine/Source/Programs/Horde/Horde.Server/Commands/ServerCommand.cs
+++ b/Engine/Source/Programs/Horde/Horde.Server/Commands/ServerCommand.cs
@@ -61,7 +61,6 @@ namespace Horde.Server.Commands
 					webBuilder.ConfigureKestrel(options =>
 					{
 						options.Limits.MaxRequestBodySize = 100 * 1024 * 1024;
-						options.Limits.Http2.KeepAlivePingDelay = TimeSpan.FromSeconds(30);
 						
 						// When agents are saturated with work (CPU or I/O), slow sending of gRPC data can happen.
 						// Kestrel protects against this behavior by default as it's commonly used for malicious attacks.
