Horde: Close the tray app automatically when uninstalling Horde.

diff --git a/Engine/Source/Programs/Horde/Horde.Agent.TrayApp/Program.cs b/Engine/Source/Programs/Horde/Horde.Agent.TrayApp/Program.cs
--- a/Engine/Source/Programs/Horde/Horde.Agent.TrayApp/Program.cs
+++ b/Engine/Source/Programs/Horde/Horde.Agent.TrayApp/Program.cs
@@ -1,12 +1,7 @@
 // Copyright Epic Games, Inc. All Rights Reserved.
 
-using System.Buffers;
-using System.Buffers.Binary;
 using System.Diagnostics;
 using System.IO.Pipes;
-using System.Text;
-using System.Text.Json;
-using System.Text.Json.Serialization;
 using EpicGames.Core;
 using Horde.Agent.TrayApp.Properties;
 
@@ -14,27 +9,87 @@ namespace Horde.Agent.TrayApp
 {
 	static class Program
 	{
+		const string MutexName = "Horde.Agent.TrayApp-Mutex";
+		const string EventName = "Horde.Agent.TrayApp-Exit";
+
 		[STAThread]
-		static async Task Main()
+		public static int Main(string[] args)
+		{
+			using EventWaitHandle closeEvent = new EventWaitHandle(false, EventResetMode.AutoReset, EventName);
+			using SingleInstanceMutex mutex = new SingleInstanceMutex(MutexName);
+
+			if (args.Any(x => x.Equals("-close", StringComparison.OrdinalIgnoreCase)))
+			{
+				closeEvent.Set();
+				return mutex.Wait(5000)? 0 : 1;
+			}
+
+			if (!mutex.Wait(0))
+			{
+				return 1;
+			}
+
+			MainAsync(closeEvent).GetAwaiter().GetResult();
+			return 0;
+		}
+
+		static async Task MainAsync(EventWaitHandle closeEvent)
 		{
 			ApplicationConfiguration.Initialize();
 
-			await using (CustomApplicationContext appContext = new CustomApplicationContext())
+			await using (CustomApplicationContext appContext = new CustomApplicationContext(closeEvent))
 			{
 				Application.Run(appContext);
 			}
 		}
 	}
 
+	class SingleInstanceMutex : IDisposable
+	{
+		readonly Mutex _mutex;
+		bool _locked;
+
+		public SingleInstanceMutex(string name)
+		{
+			_mutex = new Mutex(true, name);
+		}
+
+		public bool Wait(int timeout)
+		{
+			if (!_locked)
+			{
+				try
+				{
+					_locked = _mutex.WaitOne(timeout);
+				}
+				catch (AbandonedMutexException)
+				{
+					_locked = true;
+				}
+			}
+			return _locked;
+		}
+
+		public void Dispose()
+		{
+			if (_locked)
+			{
+				_mutex.ReleaseMutex();
+			}
+			_mutex.Dispose();
+		}
+	}
+
 	public class CustomApplicationContext : ApplicationContext, IAsyncDisposable
 	{
 		readonly NotifyIcon _trayIcon;
 		readonly BackgroundTask _clientTask;
+		readonly BackgroundTask _waitForExitTask;
 		readonly Control _mainThreadInvokeTarget;
 
 		bool _disposed;
 
-		public CustomApplicationContext()
+		public CustomApplicationContext(EventWaitHandle eventHandle)
 		{
 			ToolStripMenuItem logsMenuItem = new ToolStripMenuItem("View logs");
 			logsMenuItem.Click += OnOpenLogs;
@@ -58,6 +113,7 @@ namespace Horde.Agent.TrayApp
 			};
 
 			_clientTask = BackgroundTask.StartNew(StatusTask);
+			_waitForExitTask = BackgroundTask.StartNew(ctx => WaitForExitTask(eventHandle, ctx));
 		}
 
 		protected override void Dispose(bool disposing)
@@ -74,6 +130,7 @@ namespace Horde.Agent.TrayApp
 
 		public async ValueTask DisposeAsync()
 		{
+			await _waitForExitTask.DisposeAsync();
 			await _clientTask.DisposeAsync();
 
 			Dispose();
@@ -143,6 +200,12 @@ namespace Horde.Agent.TrayApp
 			}
 		}
 
+		async Task WaitForExitTask(EventWaitHandle eventHandle, CancellationToken cancellationToken)
+		{
+			await eventHandle.WaitOneAsync(cancellationToken);
+			_mainThreadInvokeTarget.BeginInvoke(() => Exit_MainThread());
+		}
+
 		async Task PollForStatusUpdatesAsync(CancellationToken cancellationToken)
 		{
 			AgentMessageBuffer message = new AgentMessageBuffer();
@@ -175,10 +238,18 @@ namespace Horde.Agent.TrayApp
 			}
 		}
 
+		void Exit_MainThread()
+		{
+			if (!_disposed)
+			{
+				_trayIcon.Visible = false;
+				Application.Exit();
+			}
+		}
+
 		void Exit(object sender, EventArgs e)
 		{
-			_trayIcon.Visible = false;
-			Application.Exit();
+			Exit_MainThread();
 		}
 	}
 }
