Horde: Prevent sending Slack messages to empty channel names.

diff --git a/Engine/Source/Programs/Horde/Horde.Server/Notifications/Sinks/SlackNotificationSink.cs b/Engine/Source/Programs/Horde/Horde.Server/Notifications/Sinks/SlackNotificationSink.cs
--- a/Engine/Source/Programs/Horde/Horde.Server/Notifications/Sinks/SlackNotificationSink.cs
+++ b/Engine/Source/Programs/Horde/Horde.Server/Notifications/Sinks/SlackNotificationSink.cs
@@ -803,7 +803,7 @@ namespace Horde.Server.Notifications.Sinks
 				if (workflowId != null)
 				{
 					StreamConfig? streamConfig;
-					if (globalConfig.TryGetStream(span.StreamId, out streamConfig) && streamConfig.TryGetWorkflow(workflowId.Value, out workflow) && workflow.TriageChannel != null)
+					if (globalConfig.TryGetStream(span.StreamId, out streamConfig) && streamConfig.TryGetWorkflow(workflowId.Value, out workflow) && !String.IsNullOrEmpty(workflow.TriageChannel))
 					{
 						await CreateOrUpdateWorkflowThreadAsync(workflow.TriageChannel, issue, span, details.Spans, workflow);
 						notifyOwner = notifySuspects = false;
@@ -849,11 +849,11 @@ namespace Horde.Server.Notifications.Sinks
 				if (globalConfig.TryGetStream(span.StreamId, out streamConfig))
 				{
 					TemplateRefConfig? templateRefConfig;
-					if (streamConfig.TryGetTemplate(span.TemplateRefId, out templateRefConfig) && templateRefConfig.TriageChannel != null)
+					if (streamConfig.TryGetTemplate(span.TemplateRefId, out templateRefConfig) && !String.IsNullOrEmpty(templateRefConfig.TriageChannel))
 					{
 						channels.Add(templateRefConfig.TriageChannel);
 					}
-					else if (streamConfig.TriageChannel != null)
+					else if (!String.IsNullOrEmpty(streamConfig.TriageChannel))
 					{
 						channels.Add(streamConfig.TriageChannel);
 					}
@@ -1306,7 +1306,7 @@ namespace Horde.Server.Notifications.Sinks
 			if (workflowId != null)
 			{
 				StreamConfig? streamConfig;
-				if(globalConfig.TryGetStream(span.StreamId, out streamConfig) && streamConfig.TryGetWorkflow(workflowId.Value, out WorkflowConfig? workflow) && workflow.TriageChannel != null && workflow.AllowMentions)
+				if(globalConfig.TryGetStream(span.StreamId, out streamConfig) && streamConfig.TryGetWorkflow(workflowId.Value, out WorkflowConfig? workflow) && !String.IsNullOrEmpty(workflow.TriageChannel) && workflow.AllowMentions)
 				{
 					MessageStateDocument? state = await GetMessageStateAsync(workflow.TriageChannel, GetTriageThreadEventId(issue.Id));
 					if (state != null)
@@ -1892,7 +1892,7 @@ namespace Horde.Server.Notifications.Sinks
 			string prefix = GetSeverityPrefix(issue.Severity);
 			StringBuilder body = new StringBuilder($"{prefix}Issue *<{issueUrl}|{issue.Id}>");
 
-			if (triageChannel != null)
+			if (!String.IsNullOrEmpty(triageChannel))
 			{
 				MessageStateDocument? state = await GetMessageStateAsync(triageChannel, GetTriageThreadEventId(issue.Id));
 				if (state != null && state.Permalink != null)
@@ -2425,7 +2425,7 @@ namespace Horde.Server.Notifications.Sinks
 			{
 				return null;
 			}
-			if (workflow.TriageChannel == null || workflow.EscalateAlias == null || workflow.EscalateTimes.Count == 0)
+			if (String.IsNullOrEmpty(workflow.TriageChannel) || workflow.EscalateAlias == null || workflow.EscalateTimes.Count == 0)
 			{
 				return null;
 			}
