Horde: Another attempt to suppress cancellation exceptions during UpdateSession calls.

diff --git a/Engine/Source/Programs/Horde/Horde.Server/Server/RpcService.cs b/Engine/Source/Programs/Horde/Horde.Server/Server/RpcService.cs
--- a/Engine/Source/Programs/Horde/Horde.Server/Server/RpcService.cs
+++ b/Engine/Source/Programs/Horde/Horde.Server/Server/RpcService.cs
@@ -365,7 +365,7 @@ namespace Horde.Server.Server
 				nextRequestTask = nextRequestTask.ContinueWith(task =>
 				{
 					cancellationSource.Cancel();
-					return task.Result;
+					return task.IsCanceled? false : task.Result;
 				}, TaskScheduler.Current);
 
 				// Get the current agent state
@@ -420,9 +420,9 @@ namespace Horde.Server.Server
 						nextRequestTask = reader.MoveNext();
 					}
 				}
-				catch (OperationCanceledException ex)
+				catch (Exception ex)
 				{
-					_logger.LogDebug(ex, "Ignoring cancellation exception while finishing UpdateSession request.");
+					_logger.LogDebug(ex, "Ignoring exception while finishing UpdateSession request.");
 				}
 			}
 		}
