Horde: Disable agent telemetry service

diff --git a/Engine/Source/Programs/Horde/Horde.Agent/Program.cs b/Engine/Source/Programs/Horde/Horde.Agent/Program.cs
--- a/Engine/Source/Programs/Horde/Horde.Agent/Program.cs
+++ b/Engine/Source/Programs/Horde/Horde.Agent/Program.cs
@@ -187,8 +187,9 @@ namespace Horde.Agent
 				});
 
 			services.AddSingleton<GrpcService>();
-			services.AddSingleton<TelemetryService>();
-			services.AddHostedService(sp => sp.GetRequiredService<TelemetryService>());
+			// Disabled due to HTTP client timeouts potentially caused by the TelemetryService's gRPC requests
+			// services.AddSingleton<TelemetryService>();
+			// services.AddHostedService(sp => sp.GetRequiredService<TelemetryService>());
 
 			services.AddSingleton<IJobExecutorFactory, PerforceExecutorFactory>();
 			services.AddSingleton<IJobExecutorFactory, WorkspaceExecutorFactory>();
