Horde: Fix AuditLogAwsInstanceTypeChanges test.

This intermittently fails due to a race condition between the task which is manually being called to update the message queue and a background task which is updating it, resulting in AuditLog.FlushMessagesInternalAsync() returning 0 because AuditLog.WriteMessagesAsync() has already processed it.

Added a FlushAsync() method to flush the message queue explicitly, and removed the hack to manipulate the queue through an internal method.

diff --git a/Engine/Source/Programs/Horde/Horde.Server.Tests/AgentServiceTest.cs b/Engine/Source/Programs/Horde/Horde.Server.Tests/AgentServiceTest.cs
--- a/Engine/Source/Programs/Horde/Horde.Server.Tests/AgentServiceTest.cs
+++ b/Engine/Source/Programs/Horde/Horde.Server.Tests/AgentServiceTest.cs
@@ -135,9 +135,9 @@ public class AgentServiceTest : TestSetup
 		Fixture fixture = await CreateFixtureAsync();
 		IAuditLogChannel<AgentId> agentLogger = AgentCollection.GetLogger(fixture.Agent1.Id);
 
-		async Task<bool> AuditLogContains(string text, int waitForNumMessages = -1)
+		async Task<bool> AuditLogContains(string text)
 		{
-			await FlushAuditLogsAsync(waitForNumMessages);
+			await agentLogger.FlushAsync();
 			return await agentLogger.FindAsync().AnyAsync(x => x.Data.Contains(text, StringComparison.Ordinal));
 		}
 		
@@ -147,33 +147,7 @@ public class AgentServiceTest : TestSetup
 		
 		props = new () { $"{KnownPropertyNames.AwsInstanceType}=c6.xlarge" };
 		agent = await AgentService.CreateSessionAsync(agent, AgentStatus.Ok, props, new Dictionary<string, int>(), "test");
-		Assert.IsTrue(await AuditLogContains("AWS EC2 instance type changed", 1));
-	}
-
-	private async Task FlushAuditLogsAsync(int waitForNumMessages = -1, int maxWaitMs = 3000)
-	{
-		IAuditLog<AgentId> auditLog = ServiceProvider.GetRequiredService<IAuditLog<AgentId>>();
-		AuditLog<AgentId> log = (AuditLog<AgentId>)auditLog;
-		using CancellationTokenSource cts = new (maxWaitMs);
-		
-		if (waitForNumMessages >= 1)
-		{
-			int numMessagesFlushed = 0;
-			for (;;)
-			{
-				numMessagesFlushed += await log.FlushMessagesInternalAsync();
-				if (numMessagesFlushed >= waitForNumMessages || cts.Token.IsCancellationRequested)
-				{
-					break;
-				}
-				
-				await Task.Delay(100, cts.Token);
-			}
-		}
-		else
-		{
-			await log.FlushMessagesInternalAsync();	
-		}
+		Assert.IsTrue(await AuditLogContains("AWS EC2 instance type changed"));
 	}
 	
 	[TestMethod]
diff --git a/Engine/Source/Programs/Horde/Horde.Server/Auditing/AuditLog.cs b/Engine/Source/Programs/Horde/Horde.Server/Auditing/AuditLog.cs
--- a/Engine/Source/Programs/Horde/Horde.Server/Auditing/AuditLog.cs
+++ b/Engine/Source/Programs/Horde/Horde.Server/Auditing/AuditLog.cs
@@ -2,9 +2,12 @@
 
 using System;
 using System.Collections.Generic;
+using System.Runtime.CompilerServices;
 using System.Text;
+using System.Threading;
 using System.Threading.Channels;
 using System.Threading.Tasks;
+using Amazon.S3.Model;
 using EpicGames.Core;
 using Horde.Server.Server;
 using Horde.Server.Utilities;
@@ -154,13 +157,17 @@ namespace Horde.Server.Auditing
 			public IAsyncEnumerable<IAuditLogMessage> FindAsync(DateTime? minTime, DateTime? maxTime, int? index, int? count) => Outer.FindAsync(Subject, minTime, maxTime, index, count);
 
 			public Task<long> DeleteAsync(DateTime? minTime, DateTime? maxTime) => Outer.DeleteAsync(Subject, minTime, maxTime);
+
+			public Task FlushAsync(CancellationToken cancellationToken) => Outer.FlushAsync(cancellationToken);
 		}
 
 		readonly IMongoCollection<AuditLogMessage> _messages;
-		readonly Channel<AuditLogMessage> _messageChannel;
+		readonly Channel<AuditLogMessage?> _messageChannel;
 		readonly string _subjectProperty;
 		readonly ILogger _logger;
 		readonly Task _backgroundTask;
+		
+		TaskCompletionSource? _flushEvent;
 
 		public IAuditLogChannel<TSubject> this[TSubject subject] => new AuditLogChannel(this, subject);
 
@@ -171,7 +178,7 @@ namespace Horde.Server.Auditing
 
 			_messages = mongoService.GetCollection<AuditLogMessage>(collectionName, indexes);
 
-			_messageChannel = Channel.CreateUnbounded<AuditLogMessage>();
+			_messageChannel = Channel.CreateUnbounded<AuditLogMessage?>();
 			_subjectProperty = subjectProperty;
 			_logger = logger;
 
@@ -198,7 +205,10 @@ namespace Horde.Server.Auditing
 			List<AuditLogMessage> newMessages = new ();
 			while (_messageChannel.Reader.TryRead(out AuditLogMessage? newMessage))
 			{
-				newMessages.Add(newMessage);
+				if (newMessage != null)
+				{
+					newMessages.Add(newMessage);
+				}
 			}
 			if (newMessages.Count > 0)
 			{
@@ -212,10 +222,38 @@ namespace Horde.Server.Auditing
 		{
 			while (await _messageChannel.Reader.WaitToReadAsync())
 			{
-				await FlushMessagesInternalAsync();
+				TaskCompletionSource? flushEvent = Interlocked.Exchange(ref _flushEvent, null);
+
+				List<AuditLogMessage> newMessages = new();
+				while (_messageChannel.Reader.TryRead(out AuditLogMessage? newMessage))
+				{
+					if (newMessage != null)
+					{
+						newMessages.Add(newMessage);
+					}
+				}
+				if (newMessages.Count > 0)
+				{
+					await _messages.InsertManyAsync(newMessages);
+				}
+
+				flushEvent?.TrySetResult();
 			}
 		}
 
+		public async Task FlushAsync(CancellationToken cancellationToken)
+		{
+			// Get the existing task completion source, or create a new one
+			TaskCompletionSource newTcs = new TaskCompletionSource();
+			TaskCompletionSource? tcs = Interlocked.CompareExchange(ref _flushEvent, newTcs, null) ?? newTcs;
+
+			// Force the background task to run once 
+			await _messageChannel.Writer.WriteAsync(null, cancellationToken);
+
+			// Wait for the event to trigger
+			await tcs.Task;
+		}
+
 		async IAsyncEnumerable<IAuditLogMessage<TSubject>> FindAsync(TSubject subject, DateTime? minTime = null, DateTime? maxTime = null, int? index = null, int? count = null)
 		{
 			FilterDefinition<AuditLogMessage> filter = Builders<AuditLogMessage>.Filter.Eq(x => x.Subject, subject);
diff --git a/Engine/Source/Programs/Horde/Horde.Server/Auditing/IAuditLog.cs b/Engine/Source/Programs/Horde/Horde.Server/Auditing/IAuditLog.cs
--- a/Engine/Source/Programs/Horde/Horde.Server/Auditing/IAuditLog.cs
+++ b/Engine/Source/Programs/Horde/Horde.Server/Auditing/IAuditLog.cs
@@ -4,6 +4,7 @@ using System;
 using System.Collections.Generic;
 using System.IO.Pipelines;
 using System.Text;
+using System.Threading;
 using System.Threading.Tasks;
 using Microsoft.Extensions.Logging;
 
@@ -52,6 +53,12 @@ namespace Horde.Server.Auditing
 		/// <param name="maxTime">Maximum time to remove</param>
 		/// <returns>Async task</returns>
 		Task<long> DeleteAsync(DateTime? minTime = null, DateTime? maxTime = null);
+
+		/// <summary>
+		/// Flush all writes to this log
+		/// </summary>
+		/// <param name="cancellationToken">Cancellation token for the operation</param>
+		Task FlushAsync(CancellationToken cancellationToken = default);
 	}
 
 	/// <summary>
