Horde: Remove recent gRPC client tweaks for agent

Recent issues with disk I/O saturation in Epic's build farm is the reason for at least temporarily reverting. These settings are likely more sound but we need sort the I/O issues out first.

diff --git a/Engine/Source/Programs/Horde/Horde.Agent/Services/GrpcService.cs b/Engine/Source/Programs/Horde/Horde.Agent/Services/GrpcService.cs
--- a/Engine/Source/Programs/Horde/Horde.Agent/Services/GrpcService.cs
+++ b/Engine/Source/Programs/Horde/Horde.Agent/Services/GrpcService.cs
@@ -69,16 +69,6 @@ namespace Horde.Agent.Services
 			// HTTP handler is disposed by GrpcChannel below
 			SocketsHttpHandler httpHandler = new()
 			{
-				KeepAlivePingDelay = TimeSpan.FromSeconds(50),
-				ConnectTimeout = TimeSpan.FromSeconds(10),
-				MaxConnectionsPerServer = 20,
-				EnableMultipleHttp2Connections = true,
-				
-				// Certain load balancers, like AWS ALB, can rotate IPs quickly.
-				// By setting a short lifetime, new connections will be re-created and subsequently have the DNS name refreshed.
-				// This prevents using outdated IP addresses.
-				PooledConnectionLifetime = TimeSpan.FromSeconds(60),
-				
 				SslOptions = new SslClientAuthenticationOptions
 				{
 					RemoteCertificateValidationCallback = (sender, cert, chain, errors) => CertificateHelper.CertificateValidationCallBack(_logger, sender, cert, chain, errors, _serverProfile)
