Horde: Instrument gRPC client to separately capture DNS resolve and TCP connect timings

diff --git a/Engine/Source/Programs/Horde/Horde.Agent/Services/GrpcService.cs b/Engine/Source/Programs/Horde/Horde.Agent/Services/GrpcService.cs
--- a/Engine/Source/Programs/Horde/Horde.Agent/Services/GrpcService.cs
+++ b/Engine/Source/Programs/Horde/Horde.Agent/Services/GrpcService.cs
@@ -1,10 +1,12 @@
 // Copyright Epic Games, Inc. All Rights Reserved.
 
 using System;
+using System.Net;
 using System.Net.Http;
 using System.Net.Http.Headers;
 using System.Net.Security;
 using System.Net.Http.Json;
+using System.Net.Sockets;
 using System.Threading;
 using System.Threading.Tasks;
 using Grpc.Core;
@@ -13,6 +15,8 @@ using Grpc.Net.Client;
 using Horde.Agent.Utility;
 using Microsoft.Extensions.Logging;
 using Microsoft.Extensions.Options;
+using OpenTracing;
+using OpenTracing.Util;
 
 namespace Horde.Agent.Services
 {
@@ -69,6 +73,36 @@ namespace Horde.Agent.Services
 			// HTTP handler is disposed by GrpcChannel below
 			SocketsHttpHandler httpHandler = new()
 			{
+				ConnectCallback = async (context, connectCt) =>
+				{
+					using IScope scope = GlobalTracer.Instance
+						.BuildSpan($"{nameof(GrpcService)}.ConnectCallback")
+						.WithResourceName(context.DnsEndPoint.Host)
+						.StartActive();
+
+					IPAddress[] ips;
+					using (IScope _ = GlobalTracer.Instance.BuildSpan("DnsResolve").StartActive())
+					{
+						ips = await Dns.GetHostAddressesAsync(context.DnsEndPoint.Host, connectCt);
+					}
+
+					IPEndPoint ipEndpoint = new (ips[0], context.DnsEndPoint.Port);
+					using (IScope _ = GlobalTracer.Instance.BuildSpan("TcpConnect").StartActive())
+					{
+						Socket socket = new(SocketType.Stream, ProtocolType.Tcp);
+						try
+						{
+							await socket.ConnectAsync(ipEndpoint, connectCt);
+							return new NetworkStream(socket, ownsSocket: true);
+						}
+						catch
+						{
+							socket.Dispose();
+							throw;
+						}
+					}
+				},
+				
 				SslOptions = new SslClientAuthenticationOptions
 				{
 					RemoteCertificateValidationCallback = (sender, cert, chain, errors) => CertificateHelper.CertificateValidationCallBack(_logger, sender, cert, chain, errors, _serverProfile)
