Horde: Log when the event loop gets blocked for an extended time

diff --git a/Engine/Source/Programs/Horde/Horde.Agent.Tests/Services/TelemetryServiceTest.cs b/Engine/Source/Programs/Horde/Horde.Agent.Tests/Services/TelemetryServiceTest.cs
new file mode 100644
--- /dev/null
+++ b/Engine/Source/Programs/Horde/Horde.Agent.Tests/Services/TelemetryServiceTest.cs
@@ -0,0 +1,64 @@
+// Copyright Epic Games, Inc. All Rights Reserved.
+
+using System;
+using System.Threading;
+using System.Threading.Tasks;
+using Horde.Agent.Services;
+using Microsoft.Extensions.Logging.Abstractions;
+using Microsoft.Extensions.Options;
+using Microsoft.VisualStudio.TestTools.UnitTesting;
+
+namespace Horde.Agent.Tests.Services;
+
+[TestClass]
+public sealed class TelemetryServiceTest : IDisposable
+{
+	private readonly TelemetryService _telemetryService;
+	private readonly NullLoggerFactory _loggerFactory = new ();
+
+	public TelemetryServiceTest()
+	{
+		AgentSettings settings = new() { Server = "Test", ServerProfiles = { new ServerProfile() { Name = "Test", Url = new Uri("http://localhost:1234") } }};
+		GrpcService grpcService = new (new OptionsWrapper<AgentSettings>(settings), NullLogger<GrpcService>.Instance, _loggerFactory);
+		_telemetryService = new TelemetryService(grpcService, new OptionsWrapper<AgentSettings>(settings), NullLogger<TelemetryService>.Instance);
+	}
+
+	[TestMethod]
+	public async Task NormalEventLoopTiming()
+	{
+		DateTime now = DateTime.UtcNow;
+		_telemetryService.GetUtcNow = () => now;
+		Task<(bool onTime, TimeSpan diff)> task = _telemetryService.IsEventLoopOnTimeAsync(TimeSpan.FromSeconds(1), TimeSpan.FromMilliseconds(500), CancellationToken.None);
+		_telemetryService.GetUtcNow = () => now + TimeSpan.FromMilliseconds(1000 + 12);
+		(bool onTime, TimeSpan diff) = await task;
+		Assert.IsTrue(onTime);
+	}
+	
+	[TestMethod]
+	public async Task SlowEventLoopTooEarly()
+	{
+		DateTime now = DateTime.UtcNow;
+		_telemetryService.GetUtcNow = () => now;
+		Task<(bool onTime, TimeSpan diff)> task = _telemetryService.IsEventLoopOnTimeAsync(TimeSpan.FromSeconds(1), TimeSpan.FromMilliseconds(100), CancellationToken.None);
+		_telemetryService.GetUtcNow = () => now + TimeSpan.FromMilliseconds(50);
+		(bool onTime, TimeSpan diff) = await task;
+		Assert.IsFalse(onTime);
+	}
+	
+	[TestMethod]
+	public async Task SlowEventLoopTooLate()
+	{
+		DateTime now = DateTime.UtcNow;
+		_telemetryService.GetUtcNow = () => now;		
+		Task<(bool onTime, TimeSpan diff)> task = _telemetryService.IsEventLoopOnTimeAsync(TimeSpan.FromSeconds(1), TimeSpan.FromMilliseconds(100), CancellationToken.None);
+		_telemetryService.GetUtcNow = () => now + TimeSpan.FromMilliseconds(2500);
+		(bool onTime, TimeSpan diff) = await task;
+		Assert.IsFalse(onTime);
+	}
+
+	public void Dispose()
+	{
+		_telemetryService.Dispose();
+		_loggerFactory.Dispose();
+	}
+}
diff --git a/Engine/Source/Programs/Horde/Horde.Agent/Services/TelemetryService.cs b/Engine/Source/Programs/Horde/Horde.Agent/Services/TelemetryService.cs
--- a/Engine/Source/Programs/Horde/Horde.Agent/Services/TelemetryService.cs
+++ b/Engine/Source/Programs/Horde/Horde.Agent/Services/TelemetryService.cs
@@ -203,6 +203,9 @@ public sealed class WindowsSystemMetrics : ISystemMetrics
 /// </summary>
 class TelemetryService : BackgroundService
 {
+	private readonly TimeSpan _heartbeatInterval = TimeSpan.FromSeconds(60);
+	private readonly TimeSpan _heartbeatMaxAllowedDiff = TimeSpan.FromSeconds(5);
+	
 	private readonly GrpcService _grpcService;
 	private readonly AgentSettings _agentSettings;
 	private readonly ILogger<TelemetryService> _logger;
@@ -212,6 +215,9 @@ class TelemetryService : BackgroundService
 	private ISystemMetrics? _systemMetrics;
 	private DateTime _lastTimeAgentMetadataSent = DateTime.UnixEpoch;
 
+	private CancellationTokenSource? _eventLoopHeartbeatCts;
+	internal Func<DateTime> GetUtcNow { get; set; } = () => DateTime.UtcNow;
+
 	/// <summary>
 	/// Constructor
 	/// </summary>
@@ -231,6 +237,7 @@ class TelemetryService : BackgroundService
 	{
 		base.Dispose();
 		_systemMetrics?.Dispose();
+		_eventLoopHeartbeatCts?.Dispose();
 	}
  
 	/// <inheritdoc />
@@ -251,10 +258,61 @@ class TelemetryService : BackgroundService
 		{
 			_logger.LogInformation("System metric collection only implemented on Windows");
 		}
+
+		_eventLoopHeartbeatCts = CancellationTokenSource.CreateLinkedTokenSource(cancellationToken);
+#pragma warning disable CS4014 // Call not awaited
+		EventLoopHeartbeatAsync(_eventLoopHeartbeatCts.Token);
+#pragma warning restore CS4014
 		
 		return base.StartAsync(cancellationToken);
 	}
 
+	public override Task StopAsync(CancellationToken cancellationToken)
+	{
+		_eventLoopHeartbeatCts?.Cancel();
+		return base.StopAsync(cancellationToken);
+	}
+
+	/// <summary>
+	/// Continuously run the heartbeat for the event loop
+	/// </summary>
+	/// <param name="cancellationToken">Cancellation token</param>
+	private async Task EventLoopHeartbeatAsync(CancellationToken cancellationToken)
+	{
+		while (!cancellationToken.IsCancellationRequested)
+		{
+			(bool onTime, TimeSpan diff) = await IsEventLoopOnTimeAsync(_heartbeatInterval, _heartbeatMaxAllowedDiff, cancellationToken);
+			if (!onTime)
+			{
+				_logger.LogWarning("Event loop heartbeat was not on time. Diff {Diff} ms", diff.TotalMilliseconds);
+			}
+		}
+	}
+	
+	/// <summary>
+	/// Checks if the async event loop is on time
+	/// </summary>
+	/// <param name="wait">Time to wait</param>
+	/// <param name="maxDiff">Max allowed diff</param>
+	/// <param name="cancellationToken">Cancellation token</param>
+	/// <returns>A tuple with the result</returns>
+	public async Task<(bool onTime, TimeSpan diff)> IsEventLoopOnTimeAsync(TimeSpan wait, TimeSpan maxDiff, CancellationToken cancellationToken)
+	{
+		try
+		{
+			DateTime before = GetUtcNow();
+			await Task.Delay(wait, cancellationToken);
+			DateTime after = GetUtcNow();
+			TimeSpan diff = after - before - wait;
+			return (diff.Duration() < maxDiff, diff);
+		}
+		catch (TaskCanceledException)
+		{
+			// Ignore
+			return (true, TimeSpan.Zero);
+		}
+	}
+
 	/// <inheritdoc />
 	protected override async Task ExecuteAsync(CancellationToken stoppingToken)
 	{
