Horde: Add support for creating and extracting bundles from servers.

To create a bundle:

    horde bundle create -input=C:\InputDir -server=https://localhost:5001/api/v1/storage/namespace-name -token=bearer-token -ref=my-ref-name

To extract a bundle:

    horde bundle extract -output=C:\OutputDir -server=https://localhost:5001/api/v1/storage/namespace-name -token=bearer-token -ref=my-ref-name

diff --git a/Engine/Source/Programs/Horde/Horde/Commands/Bundles/BundleCreate.cs b/Engine/Source/Programs/Horde/Horde/Commands/Bundles/BundleCreate.cs
--- a/Engine/Source/Programs/Horde/Horde/Commands/Bundles/BundleCreate.cs
+++ b/Engine/Source/Programs/Horde/Horde/Commands/Bundles/BundleCreate.cs
@@ -1,22 +1,22 @@
 // Copyright Epic Games, Inc. All Rights Reserved.
 
 using System.Diagnostics;
-using System.Threading;
-using System.Threading.Tasks;
 using EpicGames.Core;
 using EpicGames.Horde.Storage;
 using EpicGames.Horde.Storage.Backends;
 using EpicGames.Horde.Storage.Nodes;
-using JetBrains.Annotations;
 using Microsoft.Extensions.Logging;
 
 namespace Horde.Commands.Bundles
 {
 	[Command("bundle", "create", "Creates a bundle from a folder on the local hard drive")]
-	class BundleCreate : Command
+	class BundleCreate : StorageCommandBase
 	{
-		[CommandLine("-Output=", Required = true, Description = "Output file. Should have a .ref extension.")]
-		public FileReference OutputFile { get; set; } = null!;
+		[CommandLine("-File=", Description = "Output file for the bundle ref")]
+		public FileReference? File { get; set; }
+
+		[CommandLine("-Ref=")]
+		public string? Ref { get; set; }
 
 		[CommandLine("-Input=", Required = true, Description = "Input file or directory")]
 		public string Input { get; set; } = null!;
@@ -24,15 +24,28 @@ namespace Horde.Commands.Bundles
 		[CommandLine("-Filter=", Description = "Filter for files to include, in P4 syntax (eg. Foo/...).")]
 		public string Filter { get; set; } = "...";
 
-		[CommandLine("-Blobs=", Description = "Directory to store blobs")]
-		public DirectoryReference? BlobDir { get; set; }
-
 		public override async Task<int> ExecuteAsync(ILogger logger)
 		{
+			// Create the storage client
+			IStorageClient store;
+			if (File != null)
+			{
+				store = new FileStorageClient(File.Directory, logger);
+			}
+			else if (!String.IsNullOrEmpty(Ref))
+			{
+				store = CreateStorageClient(logger);
+			}
+			else
+			{
+				throw new CommandLineArgumentException("Either -File=... or -Ref=... must be specified.");
+			}
+
+			// Gather the input files
 			DirectoryReference baseDir;
 			List<FileReference> files = new List<FileReference>();
 
-			if (File.Exists(Input))
+			if (System.IO.File.Exists(Input))
 			{
 				FileReference file = new FileReference(Input);
 				baseDir = file.Directory;
@@ -50,11 +63,7 @@ namespace Horde.Commands.Bundles
 				return 1;
 			}
 
-			DirectoryReference blobDir = BlobDir ?? OutputFile.Directory;
-			DirectoryReference.CreateDirectory(blobDir);
-
-			FileStorageClient store = new FileStorageClient(blobDir, logger);
-
+			// Create the bundle
 			Stopwatch timer = Stopwatch.StartNew();
 
 			using (TreeWriter writer = new TreeWriter(store))
@@ -65,9 +74,16 @@ namespace Horde.Commands.Bundles
 				await node.CopyFilesAsync(baseDir, files, options, writer, null, CancellationToken.None);
 
 				NodeHandle handle = await writer.FlushAsync(node);
-
-				logger.LogInformation("Writing {File}", OutputFile);
-				await FileReference.WriteAllTextAsync(OutputFile, handle.ToString());
+				if (File != null)
+				{
+					logger.LogInformation("Writing {File}", File);
+					await FileReference.WriteAllTextAsync(File, handle.ToString());
+				}
+				else
+				{
+					logger.LogInformation("Writing ref {Ref}", Ref);
+					await store.WriteRefTargetAsync(new RefName(Ref!), handle);
+				}
 			}
 
 			logger.LogInformation("Time: {Time}", timer.Elapsed.TotalSeconds);
diff --git a/Engine/Source/Programs/Horde/Horde/Commands/Bundles/BundleExtract.cs b/Engine/Source/Programs/Horde/Horde/Commands/Bundles/BundleExtract.cs
--- a/Engine/Source/Programs/Horde/Horde/Commands/Bundles/BundleExtract.cs
+++ b/Engine/Source/Programs/Horde/Horde/Commands/Bundles/BundleExtract.cs
@@ -1,8 +1,6 @@
 // Copyright Epic Games, Inc. All Rights Reserved.
 
 using System.Diagnostics;
-using System.Threading;
-using System.Threading.Tasks;
 using EpicGames.Core;
 using EpicGames.Horde.Storage;
 using EpicGames.Horde.Storage.Backends;
@@ -13,23 +11,35 @@ using Microsoft.Extensions.Logging;
 namespace Horde.Commands.Bundles
 {
 	[Command("bundle", "extract", "Extracts data from a bundle to the local hard drive")]
-	internal class BundleExtract : Command
+	internal class BundleExtract : StorageCommandBase
 	{
-		[CommandLine("-Input=")]
-		public FileReference Input { get; set; } = null!;
+		[CommandLine("-File=")]
+		public FileReference? File { get; set; }
 
-		[CommandLine("-BlobDir=")]
-		public DirectoryReference? BlobDir { get; set; }
+		[CommandLine("-Ref=")]
+		public string? Ref { get; set; }
 
 		[CommandLine("-OutputDir=", Required = true)]
 		public DirectoryReference OutputDir { get; set; } = null!;
 
 		public override async Task<int> ExecuteAsync(ILogger logger)
 		{
-			string Text = await FileReference.ReadAllTextAsync(Input);
-			NodeHandle handle = NodeHandle.Parse(Text);
-
-			FileStorageClient store = new FileStorageClient(BlobDir ?? Input.Directory, logger);
+			IStorageClient store;
+			NodeHandle handle;
+			if (File != null)
+			{
+				store = new FileStorageClient(File.Directory, logger);
+				handle = NodeHandle.Parse(await FileReference.ReadAllTextAsync(File));
+			}
+			else if (Ref != null)
+			{
+				store = CreateStorageClient(logger);
+				handle = await store.ReadRefTargetAsync(new RefName(Ref));
+			}
+			else
+			{
+				throw new CommandLineArgumentException("Either -File=... or -Ref=... must be specified");
+			}
 
 			using MemoryCache cache = new MemoryCache(new MemoryCacheOptions());
 			TreeReader reader = new TreeReader(store, cache, logger);
diff --git a/Engine/Source/Programs/Horde/Horde/Commands/StorageCommandBase.cs b/Engine/Source/Programs/Horde/Horde/Commands/StorageCommandBase.cs
new file mode 100644
--- /dev/null
+++ b/Engine/Source/Programs/Horde/Horde/Commands/StorageCommandBase.cs
@@ -0,0 +1,62 @@
+// Copyright Epic Games, Inc. All Rights Reserved.
+
+using System.Net.Http.Headers;
+using EpicGames.Core;
+using EpicGames.Horde.Storage;
+using EpicGames.Horde.Storage.Backends;
+using Microsoft.Extensions.Logging;
+
+namespace Horde.Commands
+{
+	/// <summary>
+	/// Base class for commands that require a configured storage client
+	/// </summary>
+	abstract class StorageCommandBase : Command
+	{
+		/// <summary>
+		/// Base URI to upload to
+		/// </summary>
+		[CommandLine("-Server=", Description = "Base address of the server to upload to. Should include the base path of the endpoint to use (eg. https://foo/api/v1/storage/default-namespace). Not required if an output file is specified.")]
+		public string? Server { get; set; }
+
+		/// <summary>
+		/// Auth token to use for the client
+		/// </summary>
+		[CommandLine("-Token=", Description = "Bearer token for communication with the storage server")]
+		public string? Token { get; set; }
+
+		/// <summary>
+		/// Creates a new client instance
+		/// </summary>
+		/// <param name="logger">Logger for output messages</param>
+		public IStorageClient CreateStorageClient(ILogger logger)
+		{
+			if (Server == null)
+			{
+				throw new CommandLineArgumentException("Missing -Server=... argument for using a comm");
+			}
+			if (!Server.EndsWith("/", StringComparison.Ordinal))
+			{
+				Server += "/";
+			}
+
+			return new HttpStorageClient(CreateDefaultHttpClient, () => new HttpClient(), logger);
+		}
+
+		/// <summary>
+		/// Test whether it's possible to create a storage client
+		/// </summary>
+		public bool CanCreateStorageClient() => Server != null;
+
+		HttpClient CreateDefaultHttpClient()
+		{
+			HttpClient client = new HttpClient();
+			client.BaseAddress = new Uri(Server!);
+			if (!String.IsNullOrEmpty(Token))
+			{
+				client.DefaultRequestHeaders.Authorization = new AuthenticationHeaderValue("Bearer", Token);
+			}
+			return client;
+		}
+	}
+}
