Horde: Log any loaded filter drivers known to be problematic on agent

diff --git a/Engine/Source/Programs/Horde/Horde.Agent.Tests/Services/TelemetryServiceTest.cs b/Engine/Source/Programs/Horde/Horde.Agent.Tests/Services/TelemetryServiceTest.cs
--- a/Engine/Source/Programs/Horde/Horde.Agent.Tests/Services/TelemetryServiceTest.cs
+++ b/Engine/Source/Programs/Horde/Horde.Agent.Tests/Services/TelemetryServiceTest.cs
@@ -1,6 +1,7 @@
 // Copyright Epic Games, Inc. All Rights Reserved.
 
 using System;
+using System.Collections.Generic;
 using System.Threading;
 using System.Threading.Tasks;
 using Horde.Agent.Services;
@@ -56,6 +57,41 @@ public sealed class TelemetryServiceTest : IDisposable
 		Assert.IsFalse(onTime);
 	}
 
+	[TestMethod]
+	public void ParseFltMcOutput()
+	{
+		string emptyOutput = @"
+Filter listing failed with error: 0x80070005
+Access is denied.
+";
+		
+		List<WindowsFilterDriverInfo>? noFilters = TelemetryService.ParseFltMcOutput(emptyOutput);
+		Assert.IsNull(noFilters);
+		
+		string output = @"
+Filter Name                     Num Instances    Altitude    Frame
+------------------------------  -------------  ------------  -----
+bindflt                                 1       409800         0
+CSAgent                                 9       321410         0
+storqosflt                              0       244000         0
+wcifs                                   0       189900         0
+PrjFlt                                  2       189800         0
+CldFlt                                  0       180451         0
+FileCrypt                               0       141100         0
+luafv                                   1       135000         0
+npsvctrig                               1        46000         0
+Wof                                     4        40700         0
+FileInfo                                7        40500         0
+
+";
+
+		List<WindowsFilterDriverInfo>? filters = TelemetryService.ParseFltMcOutput(output);
+		Assert.AreEqual("bindflt", filters![0].Name);
+		Assert.AreEqual(409800, filters![0].Altitude);
+		Assert.AreEqual("FileInfo", filters![10].Name);
+		Assert.AreEqual(40500, filters![10].Altitude);
+	}
+
 	public void Dispose()
 	{
 		_telemetryService.Dispose();
diff --git a/Engine/Source/Programs/Horde/Horde.Agent/Services/TelemetryService.cs b/Engine/Source/Programs/Horde/Horde.Agent/Services/TelemetryService.cs
--- a/Engine/Source/Programs/Horde/Horde.Agent/Services/TelemetryService.cs
+++ b/Engine/Source/Programs/Horde/Horde.Agent/Services/TelemetryService.cs
@@ -1,11 +1,16 @@
 // Copyright Epic Games, Inc. All Rights Reserved.
 
 using System;
+using System.Collections.Generic;
 using System.Threading;
 using System.Threading.Tasks;
 using Microsoft.Extensions.Hosting;
 using System.Diagnostics;
+using System.IO;
+using System.Linq;
 using System.Runtime.InteropServices;
+using System.Text;
+using EpicGames.Core;
 using Google.Protobuf.WellKnownTypes;
 using Grpc.Core;
 using Grpc.Net.Client;
@@ -198,6 +203,33 @@ public sealed class WindowsSystemMetrics : ISystemMetrics
 
 #pragma warning restore CA1416
 
+/// <summary>
+/// Info for a Windows filter driver as reported by fltmc.exe
+/// </summary>
+public class WindowsFilterDriverInfo
+{
+	/// <summary>
+	/// Name
+	/// </summary>
+	public string Name { get; init; }
+	
+	/// <summary>
+	/// Altitude (unique identifier for driver)
+	/// </summary>
+	public int Altitude { get; init; }
+
+	/// <summary>
+	/// Constructor
+	/// </summary>
+	/// <param name="name"></param>
+	/// <param name="altitude"></param>
+	public WindowsFilterDriverInfo(string name, int altitude)
+	{
+		Name = name;
+		Altitude = altitude;
+	}
+}
+
 /// <summary>
 /// Send telemetry events back to server at regular intervals
 /// </summary>
@@ -313,6 +345,125 @@ class TelemetryService : BackgroundService
 		}
 	}
 
+	/// <summary>
+	/// Get list of any filter drivers known to be problematic for builds
+	/// </summary>
+	/// <param name="logger">Logger to use</param>
+	/// <param name="cancellationToken">Cancellation token for the call</param>
+	public static async Task<List<string>> GetProblematicFilterDrivers(ILogger logger, CancellationToken cancellationToken)
+	{
+		try
+		{
+			List<string> problematicDrivers = new()
+			{
+				// PlayStation SDK related filter drivers known to have negative effects on file system performance
+				"cbfilter",
+				"cbfsfilter",
+				"cbfsconnect",
+				"sie-filemon",
+				
+				"csagent", // CrowdStrike
+			};
+			
+			if (RuntimeInformation.IsOSPlatform(OSPlatform.Windows))
+			{
+				string output = await ReadFltMcOutputAsync(cancellationToken);
+				List<WindowsFilterDriverInfo>? drivers = ParseFltMcOutput(output);
+				if (drivers == null)
+				{
+					logger.LogWarning("Unable to get loaded filter drivers");
+					return new List<string>();
+				}
+
+				List<string> loadedDrivers = drivers
+					.Where(x =>
+					{
+						foreach (string probDriverName in problematicDrivers)
+						{
+							if (x.Name.Contains(probDriverName, StringComparison.OrdinalIgnoreCase))
+							{
+								return true;
+							}
+						}
+						return false;
+					})
+					.Select(x => $"{x.Name} (altitude={x.Altitude})")
+					.ToList();
+				
+
+				return loadedDrivers;
+			}
+		}
+		catch (Exception e)
+		{
+			logger.LogError(e, "Error logging filter drivers");
+		}
+		
+		return new List<string>();
+	}
+	
+	/// <summary>
+	/// Log any filter drivers known to be problematic for builds
+	/// </summary>
+	/// <param name="logger">Logger to use</param>
+	/// <param name="cancellationToken">Cancellation token for the call</param>
+	public static async Task LogProblematicFilterDrivers(ILogger logger, CancellationToken cancellationToken)
+	{
+		List<string> loadedDrivers = await GetProblematicFilterDrivers(logger, cancellationToken);
+		if (loadedDrivers.Count > 0)
+		{
+			logger.LogWarning("Agent has problematic filter drivers loaded: {FilterDrivers}", String.Join(',', loadedDrivers));
+		}
+	}
+	
+	internal static async Task<string> ReadFltMcOutputAsync(CancellationToken cancellationToken)
+	{
+		string fltmcExePath = Path.Combine(Environment.SystemDirectory, "fltmc.exe");
+		using CancellationTokenSource cancellationSource = CancellationTokenSource.CreateLinkedTokenSource(cancellationToken);
+		cancellationSource.CancelAfter(10000);
+		using ManagedProcess process = new (null, fltmcExePath, "filters", null, null, null, ProcessPriorityClass.Normal);
+		StringBuilder sb = new(1000);
+		
+		while (!cancellationToken.IsCancellationRequested)
+		{
+			string? line = await process.ReadLineAsync(cancellationToken);
+			if (line == null)
+			{
+				break;
+			}
+
+			sb.AppendLine(line);
+		}
+
+		await process.WaitForExitAsync(cancellationToken);
+		return sb.ToString();
+	}
+
+	internal static List<WindowsFilterDriverInfo>? ParseFltMcOutput(string output)
+	{
+		if (output.Contains("access is denied", StringComparison.OrdinalIgnoreCase))
+		{
+			return null;
+		}
+		
+		List<WindowsFilterDriverInfo> filters = new ();
+		string[] lines = output.Split(new[] { "\r\n", "\r", "\n" }, StringSplitOptions.None);
+		
+		foreach (string line in lines)
+		{
+			if (line.Length < 5) continue;
+			if (line.StartsWith("---", StringComparison.Ordinal)) continue;
+			if (line.StartsWith("Filter", StringComparison.Ordinal)) continue;
+				
+			string[] parts = line.Split("   ", StringSplitOptions.RemoveEmptyEntries);
+			if (parts.Length < 3) continue;
+
+			filters.Add(new WindowsFilterDriverInfo(parts[0], Int32.Parse(parts[2])));
+		}
+
+		return filters;
+	}
+
 	/// <inheritdoc />
 	protected override async Task ExecuteAsync(CancellationToken stoppingToken)
 	{
diff --git a/Engine/Source/Programs/Horde/Horde.Agent/Services/WorkerService.cs b/Engine/Source/Programs/Horde/Horde.Agent/Services/WorkerService.cs
--- a/Engine/Source/Programs/Horde/Horde.Agent/Services/WorkerService.cs
+++ b/Engine/Source/Programs/Horde/Horde.Agent/Services/WorkerService.cs
@@ -79,6 +79,8 @@ namespace Horde.Agent.Services
 			// Print the server info
 			_logger.LogInformation("Arguments: {Arguments}", Environment.CommandLine);
 
+			await TelemetryService.LogProblematicFilterDrivers(_logger, stoppingToken);
+
 			// Keep trying to start an agent session with the server
 			int failureCount = 0;
 			while (!stoppingToken.IsCancellationRequested)
