Horde: More robust parsing of fltmc.exe output for listing filter drivers

diff --git a/Engine/Source/Programs/Horde/Horde.Agent.Tests/Services/TelemetryServiceTest.cs b/Engine/Source/Programs/Horde/Horde.Agent.Tests/Services/TelemetryServiceTest.cs
--- a/Engine/Source/Programs/Horde/Horde.Agent.Tests/Services/TelemetryServiceTest.cs
+++ b/Engine/Source/Programs/Horde/Horde.Agent.Tests/Services/TelemetryServiceTest.cs
@@ -65,12 +65,13 @@ Filter listing failed with error: 0x80070005
 Access is denied.
 ";
 		
-		List<WindowsFilterDriverInfo>? noFilters = TelemetryService.ParseFltMcOutput(emptyOutput);
+		List<string>? noFilters = TelemetryService.ParseFltMcOutput(emptyOutput);
 		Assert.IsNull(noFilters);
 		
 		string output = @"
 Filter Name                     Num Instances    Altitude    Frame
 ------------------------------  -------------  ------------  -----
+cbfsconnect2017                                             <Legacy>
 bindflt                                 1       409800         0
 CSAgent                                 9       321410         0
 storqosflt                              0       244000         0
@@ -85,11 +86,10 @@ FileInfo                                7        40500         0
 
 ";
 
-		List<WindowsFilterDriverInfo>? filters = TelemetryService.ParseFltMcOutput(output);
-		Assert.AreEqual("bindflt", filters![0].Name);
-		Assert.AreEqual(409800, filters![0].Altitude);
-		Assert.AreEqual("FileInfo", filters![10].Name);
-		Assert.AreEqual(40500, filters![10].Altitude);
+		List<string>? filters = TelemetryService.ParseFltMcOutput(output);
+		Assert.AreEqual("cbfsconnect2017", filters![0]);
+		Assert.AreEqual("bindflt", filters![1]);
+		Assert.AreEqual("FileInfo", filters![11]);
 	}
 
 	public void Dispose()
diff --git a/Engine/Source/Programs/Horde/Horde.Agent/Services/TelemetryService.cs b/Engine/Source/Programs/Horde/Horde.Agent/Services/TelemetryService.cs
--- a/Engine/Source/Programs/Horde/Horde.Agent/Services/TelemetryService.cs
+++ b/Engine/Source/Programs/Horde/Horde.Agent/Services/TelemetryService.cs
@@ -203,33 +203,6 @@ public sealed class WindowsSystemMetrics : ISystemMetrics
 
 #pragma warning restore CA1416
 
-/// <summary>
-/// Info for a Windows filter driver as reported by fltmc.exe
-/// </summary>
-public class WindowsFilterDriverInfo
-{
-	/// <summary>
-	/// Name
-	/// </summary>
-	public string Name { get; init; }
-	
-	/// <summary>
-	/// Altitude (unique identifier for driver)
-	/// </summary>
-	public int Altitude { get; init; }
-
-	/// <summary>
-	/// Constructor
-	/// </summary>
-	/// <param name="name"></param>
-	/// <param name="altitude"></param>
-	public WindowsFilterDriverInfo(string name, int altitude)
-	{
-		Name = name;
-		Altitude = altitude;
-	}
-}
-
 /// <summary>
 /// Send telemetry events back to server at regular intervals
 /// </summary>
@@ -368,7 +341,7 @@ class TelemetryService : BackgroundService
 			if (RuntimeInformation.IsOSPlatform(OSPlatform.Windows))
 			{
 				string output = await ReadFltMcOutputAsync(cancellationToken);
-				List<WindowsFilterDriverInfo>? drivers = ParseFltMcOutput(output);
+				List<string>? drivers = ParseFltMcOutput(output);
 				if (drivers == null)
 				{
 					logger.LogWarning("Unable to get loaded filter drivers");
@@ -380,14 +353,13 @@ class TelemetryService : BackgroundService
 					{
 						foreach (string probDriverName in problematicDrivers)
 						{
-							if (x.Name.Contains(probDriverName, StringComparison.OrdinalIgnoreCase))
+							if (x.Contains(probDriverName, StringComparison.OrdinalIgnoreCase))
 							{
 								return true;
 							}
 						}
 						return false;
 					})
-					.Select(x => $"{x.Name} (altitude={x.Altitude})")
 					.ToList();
 				
 
@@ -439,14 +411,14 @@ class TelemetryService : BackgroundService
 		return sb.ToString();
 	}
 
-	internal static List<WindowsFilterDriverInfo>? ParseFltMcOutput(string output)
+	internal static List<string>? ParseFltMcOutput(string output)
 	{
 		if (output.Contains("access is denied", StringComparison.OrdinalIgnoreCase))
 		{
 			return null;
 		}
 		
-		List<WindowsFilterDriverInfo> filters = new ();
+		List<string> filters = new ();
 		string[] lines = output.Split(new[] { "\r\n", "\r", "\n" }, StringSplitOptions.None);
 		
 		foreach (string line in lines)
@@ -456,9 +428,8 @@ class TelemetryService : BackgroundService
 			if (line.StartsWith("Filter", StringComparison.Ordinal)) continue;
 				
 			string[] parts = line.Split("   ", StringSplitOptions.RemoveEmptyEntries);
-			if (parts.Length < 3) continue;
-
-			filters.Add(new WindowsFilterDriverInfo(parts[0], Int32.Parse(parts[2])));
+			string filterName = parts[0];
+			filters.Add(filterName);
 		}
 
 		return filters;
