Horde: Add an endpoint for returning OIDC configuration from the server.

diff --git a/Engine/Source/Programs/Horde/Horde.Server/Server/ServerController.cs b/Engine/Source/Programs/Horde/Horde.Server/Server/ServerController.cs
--- a/Engine/Source/Programs/Horde/Horde.Server/Server/ServerController.cs
+++ b/Engine/Source/Programs/Horde/Horde.Server/Server/ServerController.cs
@@ -88,5 +88,16 @@ namespace Horde.Server.Server
 			response.UnencryptedHttp2 = serverSettings.Http2Port;
 			return response;
 		}
+
+		/// <summary>
+		/// Returns settings for automating auth against this server
+		/// </summary>
+		[HttpGet]
+		[AllowAnonymous]
+		[Route("/api/v1/server/auth")]
+		public ActionResult<GetAuthConfigResponse> GetAuthConfig()
+		{
+			return new GetAuthConfigResponse(_globalConfig.Value.ServerSettings);
+		}
 	}
 }
diff --git a/Engine/Source/Programs/Horde/Horde.Server/Server/ServerMessages.cs b/Engine/Source/Programs/Horde/Horde.Server/Server/ServerMessages.cs
--- a/Engine/Source/Programs/Horde/Horde.Server/Server/ServerMessages.cs
+++ b/Engine/Source/Programs/Horde/Horde.Server/Server/ServerMessages.cs
@@ -59,5 +59,31 @@ namespace Horde.Server.Server
 		/// </summary>
 		public int? UnencryptedHttp2 { get; set; }
 	}
+
+	/// <summary>
+	/// Describes the auth config for this server
+	/// </summary>
+	public class GetAuthConfigResponse
+	{
+		/// <inheritdoc cref="ServerSettings.AuthMethod"/>
+		public AuthMethod Method { get; }
+
+		/// <inheritdoc cref="ServerSettings.OidcAuthority"/>
+		public string? ServerUrl { get; }
+
+		/// <inheritdoc cref="ServerSettings.OidcClientId"/>
+		public string? ClientId { get; }
+
+		/// <inheritdoc cref="ServerSettings.OidcLocalRedirectUrls"/>
+		public string[]? LocalRedirectUrls { get; }
+
+		internal GetAuthConfigResponse(ServerSettings settings)
+		{
+			Method = settings.AuthMethod;
+			ServerUrl = settings.OidcAuthority;
+			ClientId = settings.OidcClientId;
+			LocalRedirectUrls = settings.OidcLocalRedirectUrls;
+		}
+	}
 }
 
diff --git a/Engine/Source/Programs/Horde/Horde.Server/ServerSettings.cs b/Engine/Source/Programs/Horde/Horde.Server/ServerSettings.cs
--- a/Engine/Source/Programs/Horde/Horde.Server/ServerSettings.cs
+++ b/Engine/Source/Programs/Horde/Horde.Server/ServerSettings.cs
@@ -487,6 +487,11 @@ namespace Horde.Server
 		/// </summary>
 		public string? OidcSigninRedirect { get; set; }
 
+		/// <summary>
+		/// Optional redirect url provided to OIDC login for external tools (typically to a local server)
+		/// </summary>
+		public string[]? OidcLocalRedirectUrls { get; set; }
+
 		/// <summary>
 		/// OpenID Connect scopes to request when signing in
 		/// </summary>
