Horde: Add a login command to the Horde tool, which queries the OIDC configuration from the server and creates an access token for it. This configuration is used for subsequent commands against that server, eg. uploading and downloading bundles.

Currently registers the credentials as the "Horde" OIDC provider. May want to revisit this in future to ensure name matches those configured in the branch.

diff --git a/Engine/Source/Programs/Horde/Horde/Commands/Bundles/BundleCreate.cs b/Engine/Source/Programs/Horde/Horde/Commands/Bundles/BundleCreate.cs
--- a/Engine/Source/Programs/Horde/Horde/Commands/Bundles/BundleCreate.cs
+++ b/Engine/Source/Programs/Horde/Horde/Commands/Bundles/BundleCreate.cs
@@ -34,7 +34,7 @@ namespace Horde.Commands.Bundles
 			}
 			else if (!String.IsNullOrEmpty(Ref))
 			{
-				store = CreateStorageClient(logger);
+				store = await CreateStorageClientAsync(logger);
 			}
 			else
 			{
diff --git a/Engine/Source/Programs/Horde/Horde/Commands/Bundles/BundleExtract.cs b/Engine/Source/Programs/Horde/Horde/Commands/Bundles/BundleExtract.cs
--- a/Engine/Source/Programs/Horde/Horde/Commands/Bundles/BundleExtract.cs
+++ b/Engine/Source/Programs/Horde/Horde/Commands/Bundles/BundleExtract.cs
@@ -33,7 +33,7 @@ namespace Horde.Commands.Bundles
 			}
 			else if (Ref != null)
 			{
-				store = CreateStorageClient(logger);
+				store = await CreateStorageClientAsync(logger);
 				handle = await store.ReadRefTargetAsync(new RefName(Ref));
 			}
 			else
diff --git a/Engine/Source/Programs/Horde/Horde/Commands/Login/LoginCommand.cs b/Engine/Source/Programs/Horde/Horde/Commands/Login/LoginCommand.cs
new file mode 100644
--- /dev/null
+++ b/Engine/Source/Programs/Horde/Horde/Commands/Login/LoginCommand.cs
@@ -0,0 +1,50 @@
+// Copyright Epic Games, Inc. All Rights Reserved.
+
+using System;
+using System.Collections.Generic;
+using System.Linq;
+using System.Text;
+using System.Threading.Tasks;
+using EpicGames.Core;
+using Microsoft.Extensions.Logging;
+
+namespace Horde.Commands.Login
+{
+	[Command("login", "Logs in to a Horde server")]
+	class LoginCommand : Command
+	{
+		class GetAuthConfigResponse
+		{
+			public string? Method { get; set; }
+			public string? ServerUrl { get; set; }
+			public string? ClientId { get; set; }
+			public string[]? RedirectUrls { get; set; }
+		}
+
+		[CommandLine("-Server=")]
+		public string? Server { get; set; }
+
+		/// <inheritdoc/>
+		public override async Task<int> ExecuteAsync(ILogger logger)
+		{
+			if (Server != null)
+			{
+				await Settings.SetServerAsync(Server);
+			}
+
+			if(await Settings.GetServerAsync() == null)
+			{
+				logger.LogError("No Horde server is configured. Specify -Server=... to configure one.");
+				return 1;
+			}
+
+			if (await Settings.GetAccessTokenAsync(logger) == null)
+			{
+				logger.LogError("Unable to log in to server");
+				return 1;
+			}
+
+			return 0;
+		}
+	}
+}
diff --git a/Engine/Source/Programs/Horde/Horde/Commands/StorageCommandBase.cs b/Engine/Source/Programs/Horde/Horde/Commands/StorageCommandBase.cs
--- a/Engine/Source/Programs/Horde/Horde/Commands/StorageCommandBase.cs
+++ b/Engine/Source/Programs/Horde/Horde/Commands/StorageCommandBase.cs
@@ -14,48 +14,54 @@ namespace Horde.Commands
 	abstract class StorageCommandBase : Command
 	{
 		/// <summary>
-		/// Base URI to upload to
+		/// Namespace to use
 		/// </summary>
-		[CommandLine("-Server=", Description = "Base address of the server to upload to. Should include the base path of the endpoint to use (eg. https://foo/api/v1/storage/default-namespace). Not required if an output file is specified.")]
-		public string? Server { get; set; }
+		[CommandLine("-Namespace=", Description = "Namespace for data to manipulate")]
+		public string Namespace { get; set; } = "default";
 
 		/// <summary>
-		/// Auth token to use for the client
+		/// Base URI to upload to
 		/// </summary>
-		[CommandLine("-Token=", Description = "Bearer token for communication with the storage server")]
-		public string? Token { get; set; }
+		[CommandLine("-Path=", Description = "Relative path on the server for the store to write to/from (eg. api/v1/storage/default)")]
+		public string? Path { get; set; }
 
 		/// <summary>
 		/// Creates a new client instance
 		/// </summary>
 		/// <param name="logger">Logger for output messages</param>
-		public IStorageClient CreateStorageClient(ILogger logger)
+		/// <param name="cancellationToken"></param>
+		public async Task<IStorageClient> CreateStorageClientAsync(ILogger logger, CancellationToken cancellationToken = default)
 		{
-			if (Server == null)
+			Uri? server = await Settings.GetServerAsync(cancellationToken);
+			if (server == null)
 			{
-				throw new CommandLineArgumentException("Missing -Server=... argument for using a comm");
+				throw new Exception("No server is configured. Run 'horde login -server=...' to set up.");
 			}
-			if (!Server.EndsWith("/", StringComparison.Ordinal))
+
+			string? token = await Settings.GetAccessTokenAsync(logger, cancellationToken);
+			if (token == null)
 			{
-				Server += "/";
+				throw new Exception("Unable to log in to server.");
 			}
 
-			return new HttpStorageClient(CreateDefaultHttpClient, () => new HttpClient(), logger);
-		}
+			if (String.IsNullOrEmpty(Path))
+			{
+				Path = $"api/v1/storage/{Namespace}/";
+			}
+			else if(!Path.EndsWith("/", StringComparison.Ordinal))
+			{
+				Path += "/";
+			}
 
-		/// <summary>
-		/// Test whether it's possible to create a storage client
-		/// </summary>
-		public bool CanCreateStorageClient() => Server != null;
+			server = new Uri(server, Path);
+			return new HttpStorageClient(() => CreateDefaultHttpClient(server, token), () => new HttpClient(), logger);
+		}
 
-		HttpClient CreateDefaultHttpClient()
+		static HttpClient CreateDefaultHttpClient(Uri server, string token)
 		{
 			HttpClient client = new HttpClient();
-			client.BaseAddress = new Uri(Server!);
-			if (!String.IsNullOrEmpty(Token))
-			{
-				client.DefaultRequestHeaders.Authorization = new AuthenticationHeaderValue("Bearer", Token);
-			}
+			client.BaseAddress = server;
+			client.DefaultRequestHeaders.Authorization = new AuthenticationHeaderValue("Bearer", token);
 			return client;
 		}
 	}
diff --git a/Engine/Source/Programs/Horde/Horde/Settings.cs b/Engine/Source/Programs/Horde/Horde/Settings.cs
new file mode 100644
--- /dev/null
+++ b/Engine/Source/Programs/Horde/Horde/Settings.cs
@@ -0,0 +1,142 @@
+// Copyright Epic Games, Inc. All Rights Reserved.
+
+using System.Text.Json;
+using EpicGames.Core;
+using EpicGames.OIDC;
+using Microsoft.Extensions.Configuration;
+using Microsoft.Extensions.Logging;
+
+namespace Horde
+{
+	/// <summary>
+	/// Utility methods for CLI settings
+	/// </summary>
+	static class Settings
+	{
+		class SettingsData
+		{
+			public Uri Server { get; set; } = new Uri("http://localhost:5000/");
+		}
+
+		class GetAuthConfigResponse
+		{
+			public string? Method { get; set; }
+			public string? ServerUrl { get; set; }
+			public string? ClientId { get; set; }
+			public string[]? LocalRedirectUrls { get; set; }
+		}
+
+		static FileReference GetSettingsFile()
+		{
+			DirectoryReference? baseDir = DirectoryReference.GetSpecialFolder(Environment.SpecialFolder.UserProfile);
+			baseDir ??= DirectoryReference.GetSpecialFolder(Environment.SpecialFolder.LocalApplicationData);
+			baseDir ??= DirectoryReference.GetCurrentDirectory();
+			return FileReference.Combine(baseDir, ".horde.json");
+		}
+
+		static SettingsData? s_settings;
+
+		static async Task<SettingsData> ReadSettingsAsync(CancellationToken cancellationToken)
+		{
+			if (s_settings == null)
+			{
+				FileReference file = GetSettingsFile();
+				if (FileReference.Exists(file))
+				{
+					byte[] data = await FileReference.ReadAllBytesAsync(GetSettingsFile(), cancellationToken);
+					s_settings = JsonSerializer.Deserialize<SettingsData>(data, new JsonSerializerOptions { AllowTrailingCommas = true, PropertyNameCaseInsensitive = true, ReadCommentHandling = JsonCommentHandling.Skip });
+				}
+				s_settings ??= new SettingsData();
+			}
+			return s_settings;
+		}
+
+		static async Task WriteSettingsAsync(CancellationToken cancellationToken)
+		{
+			if (s_settings != null)
+			{
+				byte[] data = JsonSerializer.SerializeToUtf8Bytes(s_settings, new JsonSerializerOptions { PropertyNamingPolicy = JsonNamingPolicy.CamelCase, WriteIndented = true });
+
+				FileReference file = GetSettingsFile();
+				DirectoryReference.CreateDirectory(file.Directory);
+
+				await FileReference.WriteAllBytesAsync(file, data, cancellationToken);
+			}
+		}
+
+		public static async Task SetServerAsync(string server, CancellationToken cancellationToken = default)
+		{
+			if (!server.EndsWith("/", StringComparison.Ordinal))
+			{
+				server += "/";
+			}
+
+			SettingsData settings = await ReadSettingsAsync(cancellationToken);
+			settings.Server = new Uri(server);
+			await WriteSettingsAsync(cancellationToken);
+		}
+
+		public static async Task<Uri?> GetServerAsync(CancellationToken cancellationToken = default)
+		{
+			SettingsData settings = await ReadSettingsAsync(cancellationToken);
+			return settings.Server;
+		}
+
+		public static async Task<string?> GetAccessTokenAsync(ILogger logger, CancellationToken cancellationToken = default)
+		{
+			SettingsData settings = await ReadSettingsAsync(cancellationToken);
+			if (settings.Server == null)
+			{
+				throw new InvalidOperationException("No Horde server is configured. Run 'horde login -server=...' to configure.");
+			}
+
+			logger.LogInformation("Getting access token for {Server}", settings.Server);
+
+			GetAuthConfigResponse authConfig;
+			using (HttpClient httpClient = new HttpClient())
+			{
+				Uri uri = new Uri(settings.Server, "api/v1/server/auth");
+				authConfig = await httpClient.GetAsync<GetAuthConfigResponse>(uri, cancellationToken);
+			}
+
+			if (String.Equals(authConfig.Method, "Anonymous", StringComparison.OrdinalIgnoreCase))
+			{
+				return null;
+			}
+				
+			const string OidcProvider = "Horde";
+
+			Dictionary<string, string?> values = new Dictionary<string, string?>();
+			values[$"Providers:{OidcProvider}:DisplayName"] = "Horde";
+			values[$"Providers:{OidcProvider}:ServerUri"] = authConfig.ServerUrl;
+			values[$"Providers:{OidcProvider}:ClientId"] = authConfig.ClientId;
+			values[$"Providers:{OidcProvider}:RedirectUri"] = authConfig.LocalRedirectUrls?.FirstOrDefault();
+
+			ConfigurationBuilder builder = new ConfigurationBuilder();
+			builder.AddInMemoryCollection(values);
+
+			IConfiguration configuration = builder.Build();
+
+			using ITokenStore tokenStore = TokenStoreFactory.CreateTokenStore();
+			OidcTokenManager oidcTokenManager = OidcTokenManager.CreateTokenManager(configuration, tokenStore, new List<string>() { OidcProvider });
+
+			OidcTokenInfo result;
+			try
+			{
+				result = await oidcTokenManager.GetAccessToken(OidcProvider, cancellationToken);
+			}
+			catch (NotLoggedInException)
+			{
+				result = await oidcTokenManager.Login(OidcProvider, cancellationToken);
+			}
+
+			if (result.AccessToken == null)
+			{
+				throw new Exception($"Unable to get access token for {settings.Server}");
+			}
+
+			logger.LogInformation("Received bearer token for {Server}", settings.Server);
+			return result.AccessToken;
+		}
+	}
+}
