Horde: Fix some job steps being added a second time, even when their first run succeeded. Caused by logic during batch update which would not correctly skip over nodes that have already run.

diff --git a/Engine/Source/Programs/Horde/Horde.Server.Tests/JobCollectionTests.cs b/Engine/Source/Programs/Horde/Horde.Server.Tests/JobCollectionTests.cs
--- a/Engine/Source/Programs/Horde/Horde.Server.Tests/JobCollectionTests.cs
+++ b/Engine/Source/Programs/Horde/Horde.Server.Tests/JobCollectionTests.cs
@@ -352,67 +352,71 @@ namespace Horde.Server.Tests
 			Assert.AreEqual(JobStepState.Skipped, job.Batches[0].Steps[1].State);
 			Assert.AreEqual(JobStepOutcome.Failure, job.Batches[0].Steps[1].Outcome);
 
-			// Retry the failed step
+			// Retry the failed step (expect: new batch containing the retried step is created, skipped steps from original batch are removed)
 			job = Deref(await JobCollection.TryUpdateStepAsync(job, graph, job.Batches[0].Id, job.Batches[0].Steps[0].Id, newRetryByUserId: UserId.Anonymous));
-			Assert.AreEqual(1, job.Batches.Count);
-			Assert.AreEqual(3, job.Batches[0].Steps.Count);
+			Assert.AreEqual(2, job.Batches.Count);
+			Assert.AreEqual(1, job.Batches[0].Steps.Count);
+			Assert.AreEqual(2, job.Batches[1].Steps.Count);
 
 			IJobStep step = job.Batches[0].Steps[0];
 			Assert.AreEqual(0, step.NodeIdx);
 			Assert.AreEqual(JobStepState.Completed, step.State);
 			Assert.AreEqual(JobStepOutcome.Failure, step.Outcome);
 
-			step = job.Batches[0].Steps[1];
+			step = job.Batches[1].Steps[0];
 			Assert.AreEqual(0, step.NodeIdx);
 			Assert.AreEqual(JobStepState.Ready, step.State);
 			Assert.AreEqual(JobStepOutcome.Success, step.Outcome);
 
-			step = job.Batches[0].Steps[2];
+			step = job.Batches[1].Steps[1];
 			Assert.AreEqual(1, step.NodeIdx);
 			Assert.AreEqual(JobStepState.Waiting, step.State);
 			Assert.AreEqual(JobStepOutcome.Success, step.Outcome);
 
 			// Fail the retried step
-			job = await RunStep(job, graph, 0, 1, JobStepOutcome.Failure);
-			Assert.AreEqual(1, job.Batches.Count);
-			Assert.AreEqual(3, job.Batches[0].Steps.Count);
+			job = await RunStep(job, graph, 1, 0, JobStepOutcome.Failure);
+			Assert.AreEqual(2, job.Batches.Count);
+			Assert.AreEqual(1, job.Batches[0].Steps.Count);
+			Assert.AreEqual(2, job.Batches[1].Steps.Count);
 
 			step = job.Batches[0].Steps[0];
 			Assert.AreEqual(0, step.NodeIdx);
 			Assert.AreEqual(JobStepState.Completed, step.State);
 			Assert.AreEqual(JobStepOutcome.Failure, step.Outcome);
 
-			step = job.Batches[0].Steps[1];
+			step = job.Batches[1].Steps[0];
 			Assert.AreEqual(0, step.NodeIdx);
 			Assert.AreEqual(JobStepState.Completed, step.State);
 			Assert.AreEqual(JobStepOutcome.Failure, step.Outcome);
 
-			step = job.Batches[0].Steps[2];
+			step = job.Batches[1].Steps[1];
 			Assert.AreEqual(1, step.NodeIdx);
 			Assert.AreEqual(JobStepState.Skipped, step.State);
 			Assert.AreEqual(JobStepOutcome.Failure, step.Outcome);
 
 			// Retry the failed step for a second time
-			job = Deref(await JobCollection.TryUpdateStepAsync(job, graph, job.Batches[0].Id, job.Batches[0].Steps[1].Id, newRetryByUserId: UserId.Anonymous));
-			Assert.AreEqual(1, job.Batches.Count);
-			Assert.AreEqual(4, job.Batches[0].Steps.Count);
+			job = Deref(await JobCollection.TryUpdateStepAsync(job, graph, job.Batches[1].Id, job.Batches[1].Steps[0].Id, newRetryByUserId: UserId.Anonymous));
+			Assert.AreEqual(3, job.Batches.Count);
+			Assert.AreEqual(1, job.Batches[0].Steps.Count);
+			Assert.AreEqual(1, job.Batches[1].Steps.Count);
+			Assert.AreEqual(2, job.Batches[2].Steps.Count);
 
 			step = job.Batches[0].Steps[0];
 			Assert.AreEqual(0, step.NodeIdx);
 			Assert.AreEqual(JobStepState.Completed, step.State);
 			Assert.AreEqual(JobStepOutcome.Failure, step.Outcome);
 
-			step = job.Batches[0].Steps[1];
+			step = job.Batches[1].Steps[0];
 			Assert.AreEqual(0, step.NodeIdx);
 			Assert.AreEqual(JobStepState.Completed, step.State);
 			Assert.AreEqual(JobStepOutcome.Failure, step.Outcome);
 
-			step = job.Batches[0].Steps[2];
+			step = job.Batches[2].Steps[0];
 			Assert.AreEqual(0, step.NodeIdx);
 			Assert.AreEqual(JobStepState.Ready, step.State);
 			Assert.AreEqual(JobStepOutcome.Success, step.Outcome);
 
-			step = job.Batches[0].Steps[3];
+			step = job.Batches[2].Steps[1];
 			Assert.AreEqual(1, step.NodeIdx);
 			Assert.AreEqual(JobStepState.Waiting, step.State);
 			Assert.AreEqual(JobStepOutcome.Success, step.Outcome);
@@ -526,30 +530,89 @@ namespace Horde.Server.Tests
 
 			// Retry the last failed step. All the failed upstream steps should run again.
 			job = Deref(await JobCollection.TryUpdateStepAsync(job, graph, job.Batches[0].Id, job.Batches[0].Steps[2].Id, newRetryByUserId: UserId.Anonymous));
-			Assert.AreEqual(1, job.Batches.Count);
-			Assert.AreEqual(4, job.Batches[0].Steps.Count);
+			Assert.AreEqual(2, job.Batches.Count);
+			Assert.AreEqual(1, job.Batches[0].Steps.Count);
+			Assert.AreEqual(3, job.Batches[1].Steps.Count);
 
 			step = job.Batches[0].Steps[0];
 			Assert.AreEqual(0, step.NodeIdx);
 			Assert.AreEqual(JobStepState.Completed, step.State);
 			Assert.AreEqual(JobStepOutcome.Failure, step.Outcome);
 
-			step = job.Batches[0].Steps[1];
+			step = job.Batches[1].Steps[0];
 			Assert.AreEqual(0, step.NodeIdx);
 			Assert.AreEqual(JobStepState.Ready, step.State);
 			Assert.AreEqual(JobStepOutcome.Success, step.Outcome);
 
-			step = job.Batches[0].Steps[2];
+			step = job.Batches[1].Steps[1];
 			Assert.AreEqual(1, step.NodeIdx);
 			Assert.AreEqual(JobStepState.Waiting, step.State);
 			Assert.AreEqual(JobStepOutcome.Success, step.Outcome);
 
-			step = job.Batches[0].Steps[3];
+			step = job.Batches[1].Steps[2];
 			Assert.AreEqual(2, step.NodeIdx);
 			Assert.AreEqual(JobStepState.Waiting, step.State);
 			Assert.AreEqual(JobStepOutcome.Success, step.Outcome);
 		}
 
+		[TestMethod]
+		public async Task UpdateBatchesDuringJobAsync()
+		{
+			Mock<ITemplate> templateMock = new Mock<ITemplate>(MockBehavior.Strict);
+			templateMock.SetupGet(x => x.InitialAgentType).Returns((string?)null);
+
+			List<NewGroup> newGroups = new List<NewGroup>();
+
+			NewGroup group = AddGroup(newGroups);
+			AddNode(group, "Step 1", null);
+			AddNode(group, "Step 2", new[] { "Step 1" });
+			AddNode(group, "Step 3", new[] { "Step 2" });
+
+			IGraph graph = await GraphCollection.AppendAsync(null, newGroups, null, null);
+
+			CreateJobOptions options = new CreateJobOptions();
+			options.Arguments.Add("-Target=Step 3");
+
+			IJob job = await JobCollection.AddAsync(JobId.GenerateNewId(), new StreamId("ue4-main"), new TemplateId("test-build"), ContentHash.SHA1("hello"), graph, "Test job", 123, 123, options);
+
+			// Pass the first step
+			job = await StartBatch(job, graph, 0);
+			job = await RunStep(job, graph, 0, 0, JobStepOutcome.Success);
+
+			Assert.AreEqual(1, job.Batches.Count);
+			Assert.AreEqual(3, job.Batches[0].Steps.Count);
+
+			IJobStep step = job.Batches[0].Steps[0];
+			Assert.AreEqual(0, step.NodeIdx);
+			Assert.AreEqual(JobStepState.Completed, step.State);
+			Assert.AreEqual(JobStepOutcome.Success, step.Outcome);
+
+			step = job.Batches[0].Steps[1];
+			Assert.AreEqual(1, step.NodeIdx);
+			Assert.AreEqual(JobStepState.Ready, step.State);
+
+			step = job.Batches[0].Steps[2];
+			Assert.AreEqual(2, step.NodeIdx);
+			Assert.AreEqual(JobStepState.Waiting, step.State);
+
+			// Update the priority of the second step. This should trigger a refresh of the graph structure, but NOT cause any new steps to be created.
+			job = Deref(await JobCollection.TryUpdateStepAsync(job, graph, job.Batches[0].Id, job.Batches[0].Steps[1].Id, newPriority: Priority.High));
+			Assert.AreEqual(1, job.Batches.Count);
+			Assert.AreEqual(3, job.Batches[0].Steps.Count);
+
+			step = job.Batches[0].Steps[0];
+			Assert.AreEqual(0, step.NodeIdx);
+			Assert.AreEqual(JobStepState.Completed, step.State);
+
+			step = job.Batches[0].Steps[1];
+			Assert.AreEqual(1, step.NodeIdx);
+			Assert.AreEqual(JobStepState.Ready, step.State);
+
+			step = job.Batches[0].Steps[2];
+			Assert.AreEqual(2, step.NodeIdx);
+			Assert.AreEqual(JobStepState.Waiting, step.State);
+		}
+
 		[TestMethod]
 		public async Task UnknownShelfAsync()
 		{
diff --git a/Engine/Source/Programs/Horde/Horde.Server/Jobs/JobCollection.cs b/Engine/Source/Programs/Horde/Horde.Server/Jobs/JobCollection.cs
--- a/Engine/Source/Programs/Horde/Horde.Server/Jobs/JobCollection.cs
+++ b/Engine/Source/Programs/Horde/Horde.Server/Jobs/JobCollection.cs
@@ -1552,10 +1552,13 @@ namespace Horde.Server.Jobs
 			foreach (JobStepBatchDocument batch in job.Batches)
 			{
 				INodeGroup group = graph.Groups[batch.GroupIdx];
-				foreach (IJobStep step in batch.Steps)
+				foreach (JobStepDocument step in batch.Steps)
 				{
-					INode node = group.Nodes[step.NodeIdx];
-					existingNodesToExecute.Add(node);
+					if (!step.Retry)
+					{
+						INode node = group.Nodes[step.NodeIdx];
+						existingNodesToExecute.Add(node);
+					}
 				}
 			}
 
@@ -1584,7 +1587,7 @@ namespace Horde.Server.Jobs
 						if (batch != null)
 						{
 							IJobStep lastStep = batch.Steps[batch.Steps.Count - 1];
-							if (nodeIdx < lastStep.NodeIdx)
+							if (nodeIdx <= lastStep.NodeIdx)
 							{
 								appendToBatches[groupIdx] = null;
 							}
@@ -1618,15 +1621,20 @@ namespace Horde.Server.Jobs
 							appendToBatches[groupIdx] = batch;
 						}
 
-						SubResourceId stepId;
-						if (!recycleStepIds.Remove(new NodeRef(groupIdx, nodeIdx), out stepId))
+						// Don't re-add nodes that have already executed in this batch. If we were missing a dependency, we would have already nulled out the entry in appendToBatches above; anything else
+						// is already valid.
+						if (batch.Steps.Count == 0 || nodeIdx > batch.Steps[^1].NodeIdx)
 						{
-							job.NextSubResourceId = job.NextSubResourceId.Next();
-							stepId = job.NextSubResourceId;
-						}
+							SubResourceId stepId;
+							if (!recycleStepIds.Remove(new NodeRef(groupIdx, nodeIdx), out stepId))
+							{
+								job.NextSubResourceId = job.NextSubResourceId.Next();
+								stepId = job.NextSubResourceId;
+							}
 
-						JobStepDocument step = new JobStepDocument(stepId, nodeIdx);
-						batch.Steps.Add(step);
+							JobStepDocument step = new JobStepDocument(stepId, nodeIdx);
+							batch.Steps.Add(step);
+						}
 					}
 				}
 			}
