Horde: Fix issue with long issue summaries
#rnx
#skipci

diff --git a/Engine/Source/Programs/Horde/HordeDashboard/src/components/IssueViewV2.tsx b/Engine/Source/Programs/Horde/HordeDashboard/src/components/IssueViewV2.tsx
--- a/Engine/Source/Programs/Horde/HordeDashboard/src/components/IssueViewV2.tsx
+++ b/Engine/Source/Programs/Horde/HordeDashboard/src/components/IssueViewV2.tsx
@@ -1877,8 +1877,8 @@ export const IssueModalV2: React.FC<{ popHistoryOnClose: boolean, issueId?: stri
             <Stack style={{ flexBasis: "70px", flexShrink: 0 }}>
                <Stack horizontal styles={{ root: { padding: 8 } }} style={{ padding: 20, paddingBottom: 8 }}>
                   <Stack horizontal style={{ width: 1024 }} tokens={{ childrenGap: 24 }} verticalAlign="center" verticalFill={true}>
-                     <Stack style={{ maxWidth: 880, lineBreak: "normal" }} >
-                        <Text styles={{ root: { fontWeight: "unset", fontFamily: "Horde Open Sans SemiBold", fontSize: "14px", color: "#087BC4", textDecoration: details.issue?.resolvedAt ? "line-through" : undefined } }}>{title}</Text>
+                     <Stack >
+                        <Text styles={{ root: { fontWeight: "unset",  width: 880, wordBreak:"break-word", fontFamily: "Horde Open Sans SemiBold", fontSize: "14px", color: "#087BC4", textDecoration: details.issue?.resolvedAt ? "line-through" : undefined } }}>{title}</Text>
                      </Stack>
 
                      <Stack onClick={() => { setEditShown(true) }} style={{ cursor: "pointer" }}>
