Horde: Specify port numbers explicitly in the config file.

diff --git a/Engine/Source/Programs/Horde/Horde.Server/appsettings.json b/Engine/Source/Programs/Horde/Horde.Server/appsettings.json
--- a/Engine/Source/Programs/Horde/Horde.Server/appsettings.json
+++ b/Engine/Source/Programs/Horde/Horde.Server/appsettings.json
@@ -1,5 +1,8 @@
 {
 	"Horde": {
+		"HttpPort": 5000,
+		"HttpsPort": 5001,
+		"Http2Port": 5002,
 		"ConsoleLogLevel": "Debug",
 		"DashboardUrl": "https://localhost:3000",
 		"DisableAuth": true,
