Horde: Fix initial update issue with doc view state hook and async promise
#rnx
#skipci

diff --git a/Engine/Source/Programs/Horde/HordeDashboard/src/components/docs/DocView.tsx b/Engine/Source/Programs/Horde/HordeDashboard/src/components/docs/DocView.tsx
--- a/Engine/Source/Programs/Horde/HordeDashboard/src/components/docs/DocView.tsx
+++ b/Engine/Source/Programs/Horde/HordeDashboard/src/components/docs/DocView.tsx
@@ -16,11 +16,6 @@ type Anchor = {
    anchor: string;
 }
 
-const documentCache = new Map<string, string>();
-const documentRequests = new Set<string>();
-const crumbCache = new Map<string, BreadcrumbItem[]>();
-const anchorCache = new Map<string, Anchor[]>();
-
 type State = {
    crumbs: BreadcrumbItem[];
    jumpLinks: ISideRailLink[];
@@ -51,18 +46,36 @@ class LinkState {
    }
 }
 
-const linkState = new LinkState();
+class DocumentCache {
 
-const DocPanel: React.FC<{ docName: string }> = ({ docName }) => {
+   constructor() {
+      makeObservable(this);
+   }
+
+   get(docName: string): { markdown: string, crumbs: BreadcrumbItem[], anchors: Anchor[] } | undefined {
 
-   const [state, setState] = useState(0);
+      // implicit subscription
+      if (this.available) { }
 
-   const docRef = useRef<HTMLDivElement>(null);
+      const text = this.documentCache.get(docName);
+      const crumbs: BreadcrumbItem[] = this.crumbCache.get(docName) ?? [];
+      const anchors: Anchor[] = this.anchorCache.get(docName) ?? [];
 
-   const text = documentCache.get(docName);
+      if (text) {
+         return {
+            markdown: text,
+            crumbs: crumbs,
+            anchors: anchors
+         }
+      }
+
+      // already in flight
+      if (this.documentRequests.has(docName)) {
+         return undefined;
+      }
+
+      this.documentRequests.add(docName);
 
-   if (!documentRequests.has(docName)) {
-      documentRequests.add(docName);
       fetch(`/${docName}`, { cache: "no-cache" })
          .then((response) => response.text())
          .then((textContent) => {
@@ -131,27 +144,51 @@ const DocPanel: React.FC<{ docName: string }> = ({ docName }) => {
                   }
                }
 
-               documentCache.set(docName, textContent);
-               crumbCache.set(docName, crumbs);
-               anchorCache.set(docName, anchors);
+               this.documentCache.set(docName, textContent);
+               this.crumbCache.set(docName, crumbs);
+               this.anchorCache.set(docName, anchors);
 
             } else {
-               documentCache.set(docName, `###Missing document ${docName}`);
-               crumbCache.set(docName, []);
-               anchorCache.set(docName, []);
+               this.documentCache.set(docName, `###Missing document ${docName}`);
+               this.crumbCache.set(docName, []);
+               this.anchorCache.set(docName, []);
             }
 
-            setState(state + 1);
+            this.setAvailable();
+
+         });   
+   }
+
+   @observable
+   private available: number = 0;
 
-         });
+   @action
+   private setAvailable() {
+      this.available++;
    }
 
-   if (!text) {
+   private documentCache = new Map<string, string>();
+   private documentRequests = new Set<string>();
+   private crumbCache = new Map<string, BreadcrumbItem[]>();
+   private anchorCache = new Map<string, Anchor[]>();
+
+}
+
+
+const linkState = new LinkState();
+const documentCache = new DocumentCache();
+
+const DocPanel: React.FC<{ docName: string }> = observer(({ docName }) => {
+
+   const cache = documentCache.get(docName);
+
+   if (!cache) {
       return null;
    }
 
-   let crumbs = crumbCache.get(docName) ?? [];
-   let anchors = anchorCache.get(docName) ?? [];
+   const text = cache.markdown;
+   let crumbs = cache.crumbs;
+   let anchors = cache.anchors;
 
    linkState.setState(crumbs, anchors.map(a => {
       return { text: a.text, url: a.anchor }
@@ -159,11 +196,11 @@ const DocPanel: React.FC<{ docName: string }> = ({ docName }) => {
 
 
    return <Stack styles={{ root: { width: "100%" } }} >
-      <div ref={docRef} style={{margin: "16px 32px"}}>
+      <div style={{ margin: "16px 32px" }}>
          <Markdown>{text}</Markdown>
       </div>
    </Stack>;
-}
+})
 
 const DocRail = observer(() => {
 
@@ -216,7 +253,7 @@ export const DocView = () => {
       } else {
          docName = "documentation/Docs/Home.md";
       }
-      
+
    } else {
 
       if (!docName.startsWith("Docs/")) {
