Horde: Fix stream documents always being updated on read.

diff --git a/Engine/Source/Programs/Horde/Horde.Server/Streams/StreamCollection.cs b/Engine/Source/Programs/Horde/Horde.Server/Streams/StreamCollection.cs
--- a/Engine/Source/Programs/Horde/Horde.Server/Streams/StreamCollection.cs
+++ b/Engine/Source/Programs/Horde/Horde.Server/Streams/StreamCollection.cs
@@ -120,14 +120,20 @@ namespace Horde.Server.Streams
 
 				if (Config.Schedule == null)
 				{
-					Schedule = null;
-					replaceDocument = true;
+					if (Schedule != null)
+					{
+						Schedule = null;
+						replaceDocument = true;
+					}
 				}
-				else if (Schedule == null)
+				else
 				{
-					Schedule = new TemplateScheduleDoc();
-					Schedule.LastTriggerTimeUtc = utcNow;
-					replaceDocument = true;
+					if (Schedule == null)
+					{
+						Schedule = new TemplateScheduleDoc();
+						Schedule.LastTriggerTimeUtc = utcNow;
+						replaceDocument = true;
+					}
 				}
 
 				Schedule?.PostLoad(this);
