Horde: Optimize query for all streams; perform a single query for all matching streams, rather than individual queries for each one the user is permitted to see.

diff --git a/Engine/Source/Programs/Horde/Horde.Server/Streams/StreamsController.cs b/Engine/Source/Programs/Horde/Horde.Server/Streams/StreamsController.cs
--- a/Engine/Source/Programs/Horde/Horde.Server/Streams/StreamsController.cs
+++ b/Engine/Source/Programs/Horde/Horde.Server/Streams/StreamsController.cs
@@ -68,7 +68,7 @@ namespace Horde.Server.Streams
 		{
 			ProjectId[] projectIdValues = Array.ConvertAll(projectIds, x => new ProjectId(x));
 
-			List<GetStreamResponse> responses = new List<GetStreamResponse>();
+			List<StreamConfig> streamConfigs = new List<StreamConfig>();
 			foreach (ProjectConfig projectConfig in _globalConfig.Value.Projects)
 			{
 				if (projectIdValues.Length == 0 || projectIdValues.Contains(projectConfig.Id))
@@ -77,14 +77,21 @@ namespace Horde.Server.Streams
 					{
 						if (streamConfig.Authorize(StreamAclAction.ViewStream, User))
 						{
-							IStream stream = await _streamCollection.GetAsync(streamConfig);
-							GetStreamResponse response = await CreateGetStreamResponseAsync(stream);
-							responses.Add(response);
+							streamConfigs.Add(streamConfig);
 						}
 					}
 				}
 			}
 
+			List<IStream> streams = await _streamCollection.GetAsync(streamConfigs);
+
+			List<GetStreamResponse> responses = new List<GetStreamResponse>();
+			foreach(IStream stream in streams)
+			{
+				GetStreamResponse response = await CreateGetStreamResponseAsync(stream);
+				responses.Add(response);
+			}
+
 			return responses.OrderBy(x => x.Id).Select(x => PropertyFilter.Apply(x, filter)).ToList();
 		}
 
