Horde: Add gate information to stream schedule
#skipci
#rnx

diff --git a/Engine/Source/Programs/Horde/HordeDashboard/src/backend/Api.ts b/Engine/Source/Programs/Horde/HordeDashboard/src/backend/Api.ts
--- a/Engine/Source/Programs/Horde/HordeDashboard/src/backend/Api.ts
+++ b/Engine/Source/Programs/Horde/HordeDashboard/src/backend/Api.ts
@@ -1315,6 +1315,8 @@ export interface GetJobStepRefResponse {
 	/**Time at which the step finished. */
 	finishTime?: Date | string;
 
+	/** Issue ids affecting this step ref */
+	issueIds?: number[];
 }
 
 export interface GetJobStepTraceResponse {
@@ -1758,6 +1760,15 @@ export type GetSchedulePatternResponse = {
 
 }
 
+/**Gate allowing a schedule to trigger.*/
+export type ScheduleGateConfig = {
+	/**The template containing the dependency*/
+	templateId: string;
+
+	/**Target to wait for*/
+	target: string;
+}
+
 /**Response describing a schedule */
 export type GetScheduleResponse = {
 
@@ -1770,6 +1781,9 @@ export type GetScheduleResponse = {
 	/**Maximum number of changes the schedule can fall behind head revision. If greater than zero, builds will be triggered for every submitted changelist until the backlog is this size. */
 	maxChanges: number;
 
+	/** Gate for this schedule to trigger */
+	gate?: ScheduleGateConfig;
+
 	/**The template job to execute */
 	templateId: string;
 
@@ -1791,6 +1805,8 @@ export type GetScheduleResponse = {
 	/** List of active jobs */
 	activeJobs: string[];
 
+
+
 }
 
 /**Response describing when a schedule is expected to trigger */
@@ -2913,6 +2929,9 @@ export type GetIssueResponse = {
 	/**Time at which the issue was retrieved */
 	retrievedAt: Date | string;
 
+	/**Time at which the issue was retrieved */
+	lastSeenAt: Date | string;
+
 	/**The associated project for the issue */
 	project?: string;
 
diff --git a/Engine/Source/Programs/Horde/HordeDashboard/src/components/SchedulePane.tsx b/Engine/Source/Programs/Horde/HordeDashboard/src/components/SchedulePane.tsx
--- a/Engine/Source/Programs/Horde/HordeDashboard/src/components/SchedulePane.tsx
+++ b/Engine/Source/Programs/Horde/HordeDashboard/src/components/SchedulePane.tsx
@@ -8,6 +8,7 @@ import backend from '../backend';
 import { GetJobResponse, GetTemplateRefResponse, JobState } from '../backend/Api';
 import dashboard from '../backend/Dashboard';
 import { displayTimeZone } from "../base/utilities/timeUtils";
+import { projectStore } from '../backend/ProjectStore';
 
 const classes = mergeStyleSets({
    detailsRow: {
@@ -95,7 +96,26 @@ export const SchedulePane: React.FC<{ templates: GetTemplateRefResponse[] }> = (
 
       if (name === "job_name") {
 
-         return <Stack horizontalAlign="center" verticalFill verticalAlign="center">{item.template.name}</Stack>;
+         let name = item.template.name;
+         let target = schedule.gate?.target;
+         let templateId = schedule.gate?.templateId;
+         let gateTemplate = templates.find(t => t.id === templateId);
+
+
+         return <Stack horizontalAlign="start" verticalFill verticalAlign="center" tokens={{ childrenGap: 4 }}>
+            <Stack style={{ fontWeight: 600 }}>{name}</Stack>
+            {!!gateTemplate && <Stack tokens={{ childrenGap: 4 }}>
+               <Stack horizontal tokens={{ childrenGap: 8 }}>
+                  <Stack style={{ fontWeight: 600, fontSize: 10 }}>Gate:</Stack>
+                  <Stack style={{ fontSize: 10, paddingLeft: 8 }}>{gateTemplate?.name}</Stack>
+               </Stack>
+               <Stack horizontal tokens={{ childrenGap: 8 }}>
+                  <Stack style={{ fontWeight: 600, fontSize: 10 }}>Target:</Stack>
+                  <Stack style={{ fontSize: 10 }}>{target}</Stack>
+               </Stack>
+            </Stack>
+            }
+         </Stack>;
       }
 
       if (name.startsWith("time_")) {
