Horde: Shut down agents that have been disabled for longer than specified grace period

diff --git a/Engine/Source/Programs/Horde/Horde.Server.Tests/Agents/Pools/PoolUpdateServiceTest.cs b/Engine/Source/Programs/Horde/Horde.Server.Tests/Agents/Pools/PoolUpdateServiceTest.cs
new file mode 100644
--- /dev/null
+++ b/Engine/Source/Programs/Horde/Horde.Server.Tests/Agents/Pools/PoolUpdateServiceTest.cs
@@ -0,0 +1,99 @@
+// Copyright Epic Games, Inc. All Rights Reserved.
+
+using System;
+using System.Threading.Tasks;
+using Microsoft.VisualStudio.TestTools.UnitTesting;
+using System.Threading;
+using Horde.Server.Agents;
+using Horde.Server.Agents.Pools;
+using Microsoft.Extensions.Logging.Abstractions;
+
+namespace Horde.Server.Tests.Agents.Pools
+{
+	[TestClass]
+	public class PoolUpdateServiceTest : TestSetup
+	{
+		private readonly PoolUpdateService _pus;
+		private IPool _pool = default!;
+		private IAgent _enabledAgent = default!;
+		private IAgent _disabledAgent = default!;
+		private IAgent _disabledAgentBeyondGracePeriod = default!;
+		
+		public PoolUpdateServiceTest()
+		{
+			_pus = new (AgentCollection, PoolCollection, Clock, GlobalConfig, Tracer, new NullLogger<PoolUpdateService>());
+		}
+
+		[TestInitialize]
+		public async Task Setup()
+		{
+			_pool = await PoolService.CreatePoolAsync("testPool", null, true);
+			_enabledAgent = await CreateAgentAsync(_pool, true);
+			_disabledAgent = await CreateAgentAsync(_pool, false);
+			DateTime now = Clock.UtcNow;
+			Clock.UtcNow = now - TimeSpan.FromHours(9);
+			_disabledAgentBeyondGracePeriod = await CreateAgentAsync(_pool, false);
+			Clock.UtcNow = now;
+		}
+
+		private async Task RefreshAgents()
+		{
+			_enabledAgent = (await AgentService.GetAgentAsync(_enabledAgent.Id))!;
+			_disabledAgent = (await AgentService.GetAgentAsync(_disabledAgent.Id))!;
+			_disabledAgentBeyondGracePeriod = (await AgentService.GetAgentAsync(_disabledAgentBeyondGracePeriod.Id))!;
+		}
+
+		protected override void Dispose(bool disposing)
+		{
+			if (disposing)
+			{
+				_pus.Dispose();
+			}
+
+			base.Dispose(disposing);
+		}
+
+		[TestMethod]
+		public async Task ShutdownDisabledAgentBasedOnGlobalConfig()
+		{
+			Assert.IsFalse(_enabledAgent.RequestShutdown);
+			Assert.IsFalse(_disabledAgent.RequestShutdown);
+			Assert.IsFalse(_disabledAgentBeyondGracePeriod.RequestShutdown);
+			
+			await _pus.ShutdownDisabledAgentsAsync(CancellationToken.None);
+			await RefreshAgents();
+			
+			Assert.IsFalse(_enabledAgent.RequestShutdown);
+			Assert.IsFalse(_disabledAgent.RequestShutdown);
+			Assert.IsTrue(_disabledAgentBeyondGracePeriod.RequestShutdown);
+		}
+		
+		[TestMethod]
+		public async Task ShutdownDisabledAgentBasedOnPoolConfig()
+		{
+			// Explicitly set the grace period for the pool to be longer than the default of 8 hours
+			await PoolCollection.TryUpdateAsync(_pool, shutdownIfDisabledGracePeriod: TimeSpan.FromHours(24));
+			
+			await _pus.ShutdownDisabledAgentsAsync(CancellationToken.None);
+			await RefreshAgents();
+			
+			// Should result in no shutdown requests
+			Assert.IsFalse(_enabledAgent.RequestShutdown);
+			Assert.IsFalse(_disabledAgent.RequestShutdown);
+			Assert.IsFalse(_disabledAgentBeyondGracePeriod.RequestShutdown);
+		}
+		
+		[TestMethod]
+		public async Task OnlyShutdownDisabledAgentsThatAreAutoScaled()
+		{
+			await PoolCollection.TryUpdateAsync(_pool, newEnableAutoscaling: false);
+			
+			await _pus.ShutdownDisabledAgentsAsync(CancellationToken.None);
+			await RefreshAgents();
+			
+			Assert.IsFalse(_enabledAgent.RequestShutdown);
+			Assert.IsFalse(_disabledAgent.RequestShutdown);
+			Assert.IsFalse(_disabledAgentBeyondGracePeriod.RequestShutdown);
+		}
+	}
+}
diff --git a/Engine/Source/Programs/Horde/Horde.Server/Agents/Pools/IPoolCollection.cs b/Engine/Source/Programs/Horde/Horde.Server/Agents/Pools/IPoolCollection.cs
--- a/Engine/Source/Programs/Horde/Horde.Server/Agents/Pools/IPoolCollection.cs
+++ b/Engine/Source/Programs/Horde/Horde.Server/Agents/Pools/IPoolCollection.cs
@@ -128,6 +128,7 @@ namespace Horde.Server.Agents.Pools
 		/// <param name="lastScaleDownTime">New time for last (auto) scale down</param>
 		/// <param name="scaleOutCooldown">Cooldown time between scale-out events</param>
 		/// <param name="scaleInCooldown">Cooldown time between scale-in events</param>
+		/// <param name="shutdownIfDisabledGracePeriod">Time to wait before shutting down a disabled agent</param>
 		/// <param name="lastScaleResult">Result from last scaling in/out attempt</param>
 		/// <param name="lastAgentCount">Last calculated agent count</param>
 		/// <param name="lastDesiredAgentCount">Last calculated desired agent count</param>
@@ -154,6 +155,7 @@ namespace Horde.Server.Agents.Pools
 			DateTime? lastScaleDownTime = null,
 			TimeSpan? scaleOutCooldown = null,
 			TimeSpan? scaleInCooldown = null,
+			TimeSpan? shutdownIfDisabledGracePeriod = null,
 			ScaleResult? lastScaleResult = null,
 			int? lastAgentCount = null,
 			int? lastDesiredAgentCount = null,
diff --git a/Engine/Source/Programs/Horde/Horde.Server/Agents/Pools/PoolCollection.cs b/Engine/Source/Programs/Horde/Horde.Server/Agents/Pools/PoolCollection.cs
--- a/Engine/Source/Programs/Horde/Horde.Server/Agents/Pools/PoolCollection.cs
+++ b/Engine/Source/Programs/Horde/Horde.Server/Agents/Pools/PoolCollection.cs
@@ -61,6 +61,9 @@ namespace Horde.Server.Agents.Pools
 
 			[BsonIgnoreIfNull]
 			public TimeSpan? ScaleInCooldown { get; set; }
+			
+			[BsonIgnoreIfNull]
+			public TimeSpan? ShutdownIfDisabledGracePeriod { get; set;  }
 
 			[BsonIgnoreIfNull]
 			public ScaleResult? LastScaleResult { get; set; }
@@ -289,6 +292,7 @@ namespace Horde.Server.Agents.Pools
 			DateTime? lastScaleDownTime,
 			TimeSpan? scaleOutCooldown,
 			TimeSpan? scaleInCooldown, 
+			TimeSpan? shutdownIfDisabledGracePeriod, 
 			ScaleResult? lastScaleResult, 
 			int? lastAgentCount, 
 			int? lastDesiredAgentCount, 
@@ -388,6 +392,10 @@ namespace Horde.Server.Agents.Pools
 			{
 				transaction.Set(x => x.ScaleInCooldown, scaleInCooldown);
 			}
+			if (shutdownIfDisabledGracePeriod != null)
+			{
+				transaction.Set(x => x.ShutdownIfDisabledGracePeriod, shutdownIfDisabledGracePeriod);
+			}
 			if (lastScaleResult != null)
 			{
 				transaction.Set(x => x.LastScaleResult, lastScaleResult);
diff --git a/Engine/Source/Programs/Horde/Horde.Server/Agents/Pools/PoolConfig.cs b/Engine/Source/Programs/Horde/Horde.Server/Agents/Pools/PoolConfig.cs
--- a/Engine/Source/Programs/Horde/Horde.Server/Agents/Pools/PoolConfig.cs
+++ b/Engine/Source/Programs/Horde/Horde.Server/Agents/Pools/PoolConfig.cs
@@ -78,6 +78,11 @@ namespace Horde.Server.Agents.Pools
 		/// </summary>
 		public TimeSpan? ScaleInCooldown { get; }
 		
+		/// <summary>
+		/// Time to wait before shutting down an agent that has been disabled
+		/// </summary>
+		public TimeSpan? ShutdownIfDisabledGracePeriod { get; }
+		
 		/// <summary>
 		/// Last result from scaling the pool
 		/// </summary>
@@ -172,6 +177,9 @@ namespace Horde.Server.Agents.Pools
 		/// <inheritdoc/>
 		public TimeSpan? ScaleInCooldown { get; set; }
 		
+		/// <inheritdoc/>
+		public TimeSpan? ShutdownIfDisabledGracePeriod { get; set;  }
+		
 		/// <inheritdoc/>
 		public ScaleResult? LastScaleResult { get; set; }
 
diff --git a/Engine/Source/Programs/Horde/Horde.Server/Agents/Pools/PoolUpdateService.cs b/Engine/Source/Programs/Horde/Horde.Server/Agents/Pools/PoolUpdateService.cs
--- a/Engine/Source/Programs/Horde/Horde.Server/Agents/Pools/PoolUpdateService.cs
+++ b/Engine/Source/Programs/Horde/Horde.Server/Agents/Pools/PoolUpdateService.cs
@@ -5,14 +5,13 @@ using System.Collections.Generic;
 using System.Linq;
 using System.Threading;
 using System.Threading.Tasks;
-using Horde.Server.Projects;
 using Horde.Server.Server;
 using Horde.Server.Streams;
-using Horde.Server.Utilities;
 using HordeCommon;
 using Microsoft.Extensions.Hosting;
 using Microsoft.Extensions.Logging;
 using Microsoft.Extensions.Options;
+using OpenTelemetry.Trace;
 
 namespace Horde.Server.Agents.Pools
 {
@@ -21,33 +20,103 @@ namespace Horde.Server.Agents.Pools
 	/// </summary>
 	public sealed class PoolUpdateService : IHostedService, IDisposable
 	{
+		readonly IAgentCollection _agents;
 		readonly IPoolCollection _pools;
-		readonly IStreamCollection _streams;
+		readonly IClock _clock;
 		readonly IOptionsMonitor<GlobalConfig> _globalConfig;
-		readonly ILogger<PoolService> _logger;
-		readonly ITicker _ticker;
+		readonly ILogger<PoolUpdateService> _logger;
+		readonly Tracer _tracer;
+		readonly ITicker _updatePoolsTicker;
+		readonly ITicker _shutdownDisabledAgentsTicker;
 
 		/// <summary>
 		/// Constructor
 		/// </summary>
-		public PoolUpdateService(IPoolCollection pools, IStreamCollection streams, IClock clock, IOptionsMonitor<GlobalConfig> globalConfig, ILogger<PoolService> logger)
+		public PoolUpdateService(IAgentCollection agents, IPoolCollection pools, IClock clock, IOptionsMonitor<GlobalConfig> globalConfig, Tracer tracer, ILogger<PoolUpdateService> logger)
 		{
+			_agents = agents;
 			_pools = pools;
-			_streams = streams;
+			_clock = clock;
 			_globalConfig = globalConfig;
+			_tracer = tracer;
 			_logger = logger;
-			_ticker = clock.AddTicker($"{nameof(PoolUpdateService)}.{nameof(UpdatePoolsAsync)}", TimeSpan.FromSeconds(30.0), UpdatePoolsAsync, logger);
+			_updatePoolsTicker = clock.AddTicker($"{nameof(PoolUpdateService)}.{nameof(UpdatePoolsAsync)}", TimeSpan.FromSeconds(30.0), UpdatePoolsAsync, logger);
+			_shutdownDisabledAgentsTicker = clock.AddSharedTicker($"{nameof(PoolUpdateService)}.{nameof(ShutdownDisabledAgentsAsync)}", TimeSpan.FromHours(1), ShutdownDisabledAgentsAsync, logger);
 		}
 
 		/// <inheritdoc/>
-		public Task StartAsync(CancellationToken cancellationToken) => _ticker.StartAsync();
+		public async Task StartAsync(CancellationToken cancellationToken)
+		{
+			await _updatePoolsTicker.StartAsync();
+			await _shutdownDisabledAgentsTicker.StartAsync();
+		}
 
 		/// <inheritdoc/>
-		public Task StopAsync(CancellationToken cancellationToken) => _ticker.StopAsync();
+		public async Task StopAsync(CancellationToken cancellationToken)
+		{
+			await _updatePoolsTicker.StopAsync();
+			await _shutdownDisabledAgentsTicker.StopAsync();
+		} 
 
 		/// <inheritdoc/>
-		public void Dispose() => _ticker.Dispose();
+		public void Dispose()
+		{
+			_updatePoolsTicker.Dispose();
+			_shutdownDisabledAgentsTicker.Dispose();
+		}
+
+		/// <summary>
+		/// Shutdown agents that have been disabled for longer than the configured grace period
+		/// </summary>
+		/// <param name="stoppingToken">Cancellation token for the async task</param>
+		/// <returns>Async task</returns>
+		internal async ValueTask ShutdownDisabledAgentsAsync(CancellationToken stoppingToken)
+		{
+			using TelemetrySpan span = _tracer.StartActiveSpan($"{nameof(PoolUpdateService)}.{nameof(ShutdownDisabledAgentsAsync)}");
+
+			List<IPool> pools = await _pools.GetAsync();
+			IEnumerable<IAgent> disabledAgents = await _agents.FindAsync(enabled: false);
+			disabledAgents = disabledAgents.Where(x => IsAgentAutoScaled(x, pools));
+
+			int c = 0;
+			foreach (IAgent agent in disabledAgents)
+			{
+				if (HasGracePeriodExpired(agent, pools, _globalConfig.CurrentValue.AgentShutdownIfDisabledGracePeriod))
+				{
+					await _agents.TryUpdateSettingsAsync(agent, requestShutdown: true);
+					_logger.LogInformation("Shutting down agent {AgentId} as it has been disabled for longer than grace period", agent.Id.ToString());
+					c++;
+				}
+			}
+
+			span.SetAttribute("numShutdown", c);
+		}
+		
+		private bool HasGracePeriodExpired(IAgent agent, List<IPool> pools, TimeSpan globalGracePeriod)
+		{
+			if (agent.LastStatusChange == null)
+			{
+				return false;
+			}
+
+			TimeSpan gracePeriod = GetGracePeriod(agent, pools) ?? globalGracePeriod;
+			DateTime expirationTime = agent.LastStatusChange.Value + gracePeriod;
+			return _clock.UtcNow > expirationTime;
+		}
 
+		private static TimeSpan? GetGracePeriod(IAgent agent, List<IPool> pools)
+		{
+			IEnumerable<PoolId> poolIds = agent.ExplicitPools.Concat(agent.DynamicPools);
+			IPool? pool = pools.FirstOrDefault(x => poolIds.Contains(x.Id) && x.ShutdownIfDisabledGracePeriod != null);
+			return pool?.ShutdownIfDisabledGracePeriod;
+		}
+		
+		private static bool IsAgentAutoScaled(IAgent agent, List<IPool> pools)
+		{
+			IEnumerable<PoolId> poolIds = agent.ExplicitPools.Concat(agent.DynamicPools);
+			return pools.Find(x => poolIds.Contains(x.Id) && x.EnableAutoscaling) != null;
+		}
+		
 		/// <summary>
 		/// Execute the background task
 		/// </summary>
@@ -55,6 +124,8 @@ namespace Horde.Server.Agents.Pools
 		/// <returns>Async task</returns>
 		async ValueTask UpdatePoolsAsync(CancellationToken stoppingToken)
 		{
+			using TelemetrySpan span = _tracer.StartActiveSpan($"{nameof(PoolUpdateService)}.{nameof(UpdatePoolsAsync)}");
+			
 			// Capture the start time for this operation. We use this to attempt to sequence updates to agents, and prevent overriding another server's updates.
 			DateTime startTime = DateTime.UtcNow;
 
diff --git a/Engine/Source/Programs/Horde/Horde.Server/Server/GlobalConfig.cs b/Engine/Source/Programs/Horde/Horde.Server/Server/GlobalConfig.cs
--- a/Engine/Source/Programs/Horde/Horde.Server/Server/GlobalConfig.cs
+++ b/Engine/Source/Programs/Horde/Horde.Server/Server/GlobalConfig.cs
@@ -162,6 +162,12 @@ namespace Horde.Server.Server
 		/// Maximum number of conforms to run at once
 		/// </summary>
 		public int MaxConformCount { get; set; }
+		
+		/// <summary>
+		/// Time to wait before shutting down an agent that has been disabled
+		/// Used if no value is set on the actual pool.
+		/// </summary>
+		public TimeSpan AgentShutdownIfDisabledGracePeriod { get; set; } = TimeSpan.FromHours(8);
 
 		/// <summary>
 		/// Storage configuration
