Horde: Add logo variation and development watermark in dev environment
#skipci
#rnx

diff --git a/Engine/Source/Programs/Horde/HordeDashboard/src/backend/Dashboard.ts b/Engine/Source/Programs/Horde/HordeDashboard/src/backend/Dashboard.ts
--- a/Engine/Source/Programs/Horde/HordeDashboard/src/backend/Dashboard.ts
+++ b/Engine/Source/Programs/Horde/HordeDashboard/src/backend/Dashboard.ts
@@ -123,6 +123,16 @@ export class Dashboard {
         return false;        
     }
 
+    get development(): boolean {
+
+        try {            
+            return window?.location?.hostname?.indexOf("devtools-dev") !== -1;
+        } catch (reason) {
+            console.error(reason);
+        }
+        return false;        
+    }
+
     get browser(): WebBrowser {
 
         const agent = window.navigator.userAgent.toLowerCase();
diff --git a/Engine/Source/Programs/Horde/HordeDashboard/src/components/Breadcrumbs.tsx b/Engine/Source/Programs/Horde/HordeDashboard/src/components/Breadcrumbs.tsx
--- a/Engine/Source/Programs/Horde/HordeDashboard/src/components/Breadcrumbs.tsx
+++ b/Engine/Source/Programs/Horde/HordeDashboard/src/components/Breadcrumbs.tsx
@@ -111,6 +111,11 @@ export const Breadcrumbs: React.FC<{ items: BreadcrumbItem[], title?: string, su
                <Text variant="large" style={{ fontFamily: "Horde Open Sans Light", opacity: 0.33 }}>Horde Preview</Text>
             </Stack>
          </Stack>}
+         {!alert && !!dashboard.development && <Stack horizontalAlign="center" disableShrink={true} style={{ position: "absolute", width: "100%", pointerEvents: "none" }} >
+            <Stack style={{ paddingTop: 22 }}>
+               <Text variant="large" style={{ fontFamily: "Horde Open Sans Light", opacity: 0.33 }}>Horde Development</Text>
+            </Stack>
+         </Stack>}
 
          {!!alert && dashboard.alertSquelch !== alert && <Stack horizontalAlign="center" disableShrink={true} style={{ position: "absolute", width: "100%" }} >
             <Stack horizontal>
diff --git a/Engine/Source/Programs/Horde/HordeDashboard/src/components/TopNav.tsx b/Engine/Source/Programs/Horde/HordeDashboard/src/components/TopNav.tsx
--- a/Engine/Source/Programs/Horde/HordeDashboard/src/components/TopNav.tsx
+++ b/Engine/Source/Programs/Horde/HordeDashboard/src/components/TopNav.tsx
@@ -596,6 +596,8 @@ export const TopNav: React.FC<{ suppressServer?: boolean }> = observer(({ suppre
       showServer = false;
    }
 
+   const logoSrc = dashboard.development ? "/images/horde_black.svg" : "/images/horde.svg";
+
    return (
       <div style={{ backgroundColor: modeColors.header }}>
          {showVersion && <VersionModal show={true} onClose={() => { setShowVersion(false) }} />}
@@ -605,7 +607,7 @@ export const TopNav: React.FC<{ suppressServer?: boolean }> = observer(({ suppre
 
                <Link to="/index"><Stack horizontal styles={{ root: { paddingLeft: 8, cursor: 'pointer' } }}>
                   <Stack styles={{ root: { paddingTop: 2, paddingRight: 6 } }}>
-                     <img style={{ width: 32 }} src="/images/horde.svg" alt="" />
+                     <img style={{ width: 32 }} src={logoSrc} alt="" />
                   </Stack>
                   <Stack styles={{ root: { paddingTop: 11 } }}>
                      <Text styles={{ root: { fontFamily: "Horde Raleway Bold", color: dashboard.darktheme ? "#FFFFFFFF" : modeColors.text } }}>HORDE{dashboard.preview ? " PREVIEW" : ""}</Text>
