Horde: Remove the default 'horde-agent' tool definition. For installed builds we seed this from the appsettings.json file, and adding a separate entry here overwrites it.

diff --git a/Engine/Source/Programs/Horde/Horde.Server/Defaults/globals.json b/Engine/Source/Programs/Horde/Horde.Server/Defaults/globals.json
--- a/Engine/Source/Programs/Horde/Horde.Server/Defaults/globals.json
+++ b/Engine/Source/Programs/Horde/Horde.Server/Defaults/globals.json
@@ -78,25 +78,5 @@
 			"name": "Linux",
 			"condition": "Platform == 'Linux'"
 		}
-	],
-	"tools": [
-		{
-			"id": "horde-agent",
-			"name": "Horde Agent",
-			"description": "Horde agent software",
-			"acl": {
-				"entries": [
-					{
-						"claim": {
-							"type": "http://epicgames.com/ue/horde/role",
-							"value": "agent"
-						},
-						"actions": [
-							"DownloadTool"
-						]
-					}
-				]
-			}
-		}
 	]
 }
