Horde: Add white unreal/horde logo for dark theme
#rnx
#skipci

diff --git a/Engine/Source/Programs/Horde/HordeDashboard/src/components/TopNav.tsx b/Engine/Source/Programs/Horde/HordeDashboard/src/components/TopNav.tsx
--- a/Engine/Source/Programs/Horde/HordeDashboard/src/components/TopNav.tsx
+++ b/Engine/Source/Programs/Horde/HordeDashboard/src/components/TopNav.tsx
@@ -596,7 +596,11 @@ export const TopNav: React.FC<{ suppressServer?: boolean }> = observer(({ suppre
       showServer = false;
    }
 
-   const logoSrc = dashboard.development ? "/images/horde_black.svg" : "/images/horde.svg";
+   let logoSrc = "/images/horde.svg";
+
+   if (dashboard.development) {
+      logoSrc = dashboard.darktheme ? "/images/horde_white.svg" : "/images/horde_black.svg";
+   }
 
    return (
       <div style={{ backgroundColor: modeColors.header }}>
