Horde: Use condensed stream view when show names is false or there is only one template in the tab
#rnx
#skipci

diff --git a/Engine/Source/Programs/Horde/HordeDashboard/src/components/StreamView.tsx b/Engine/Source/Programs/Horde/HordeDashboard/src/components/StreamView.tsx
--- a/Engine/Source/Programs/Horde/HordeDashboard/src/components/StreamView.tsx
+++ b/Engine/Source/Programs/Horde/HordeDashboard/src/components/StreamView.tsx
@@ -18,6 +18,7 @@ import { JobViewIncremental } from './JobViewIncremental';
 import { NewBuild } from './NewBuild';
 import { StreamSummary } from './StreamSummary';
 import { TopNav } from './TopNav';
+import { JobsTabData } from '../backend/Api';
 
 export const SummaryPage: React.FC = () => {
 
@@ -89,8 +90,9 @@ const StreamViewInner: React.FC = observer(() => {
    }
 
    let queryTab = query.get("tab") ?? undefined;
+   const currentTab = stream.tabs.find(t => t.title === queryTab) as JobsTabData | undefined;
 
-   if (!queryTab || (queryTab.toLowerCase() !== "summary" && queryTab.toLowerCase() !== "all" && !stream.tabs.find(t => t.title === queryTab))) {
+   if (!queryTab || (queryTab.toLowerCase() !== "summary" && queryTab.toLowerCase() !== "all" && !currentTab)) {
 
       if (!!stream.tabs.find(t => t.title.toLowerCase() === "incremental")) {
          return <Navigate to={`/stream/${streamId}?tab=Incremental`} replace={true} />
@@ -189,9 +191,12 @@ const StreamViewInner: React.FC = observer(() => {
 
    const newBuildTab = queryTab?.toLowerCase() === "summary" ? "all" : queryTab;
    const isSwarmTab = queryTab?.toLowerCase() === "swarm" || queryTab?.toLowerCase() === "presubmit";
-   let isIncrementalTab = queryTab?.toLowerCase() === "incremental";
+   let isIncrementalTab = queryTab?.toLowerCase() === "incremental";   
 
-   // @todo: We should have a way of defining a tab's views, these are transitioning anyway, so, ok
+   if (!isIncrementalTab && currentTab) {
+      isIncrementalTab = currentTab.showNames === false || currentTab.templates?.length === 1;
+   }
+   
    if (!isIncrementalTab && queryTab) {
       const tabName = queryTab.toLowerCase();
       isIncrementalTab = tabName === "primary inc" || tabName === "secondary" || tabName.indexOf("incremental") !== -1;
