Horde: Use explicit showNames property for condensed view
#skipci
#rnx

diff --git a/Engine/Source/Programs/Horde/HordeDashboard/src/components/StreamView.tsx b/Engine/Source/Programs/Horde/HordeDashboard/src/components/StreamView.tsx
--- a/Engine/Source/Programs/Horde/HordeDashboard/src/components/StreamView.tsx
+++ b/Engine/Source/Programs/Horde/HordeDashboard/src/components/StreamView.tsx
@@ -194,7 +194,7 @@ const StreamViewInner: React.FC = observer(() => {
    let isIncrementalTab = queryTab?.toLowerCase() === "incremental";   
 
    if (!isIncrementalTab && currentTab) {
-      isIncrementalTab = currentTab.showNames === false || currentTab.templates?.length === 1;
+      isIncrementalTab = currentTab.showNames === false; /*|| currentTab.templates?.length === 1;*/
    }
    
    if (!isIncrementalTab && queryTab) {
