Horde: Reduce log verbosity for status message updates.

diff --git a/Engine/Source/Programs/Horde/Horde.Agent/Services/StatusService.cs b/Engine/Source/Programs/Horde/Horde.Agent/Services/StatusService.cs
--- a/Engine/Source/Programs/Horde/Horde.Agent/Services/StatusService.cs
+++ b/Engine/Source/Programs/Horde/Horde.Agent/Services/StatusService.cs
@@ -67,7 +67,7 @@ namespace Horde.Agent.Services
 		public void Set(AgentStatusMessage status)
 		{
 			Interlocked.Exchange(ref _current, status);
-			_logger.LogInformation("Updating status: {@Status}", status);
+			_logger.LogDebug("Updating status: {@Status}", status);
 		}
 
 		/// <summary>
