Horde: Add utilization link to the pool view
#skipci
#rnx

diff --git a/Engine/Source/Programs/Horde/HordeDashboard/src/components/PoolView.tsx b/Engine/Source/Programs/Horde/HordeDashboard/src/components/PoolView.tsx
--- a/Engine/Source/Programs/Horde/HordeDashboard/src/components/PoolView.tsx
+++ b/Engine/Source/Programs/Horde/HordeDashboard/src/components/PoolView.tsx
@@ -1068,6 +1068,15 @@ const PoolPanel: React.FC = () => {
       }
 
       if (column.name === "Value") {
+         if (item.name === "Agents") {
+            return <Stack horizontal tokens={{childrenGap: 4}}>
+               <Stack>{item.value}</Stack>
+               <Stack horizontal tokens={{childrenGap: 4}}>
+                  <Stack> - </Stack>
+                  <Stack><Link to={`/reports/utilization?pools=${pool.id}`}>Utilization</Link></Stack>
+               </Stack>
+            </Stack>
+         }
          return item.value;
       }
 
