Horde: Various installer fixes.

* Include a Windows-specific, self-contained build of the Horde Agent in the installer.
* Include a separate portable build of the Horde Agent for download from the server, for use on other platforms.
* Disable HTTPS by default, and do not create self-signed certs for GRPC comms.
* Minimize the console created by the server when launching the browser shortcut from the start menu.

diff --git a/Engine/Source/Programs/AutomationTool/BuildGraph/Tasks/WriteJsonValueTask.cs b/Engine/Source/Programs/AutomationTool/BuildGraph/Tasks/WriteJsonValueTask.cs
--- a/Engine/Source/Programs/AutomationTool/BuildGraph/Tasks/WriteJsonValueTask.cs
+++ b/Engine/Source/Programs/AutomationTool/BuildGraph/Tasks/WriteJsonValueTask.cs
@@ -69,7 +69,16 @@ namespace AutomationTool.Tasks
 		{
 			HashSet<FileReference> Files = ResolveFilespec(Unreal.RootDirectory, Parameters.File, TagNameToFileSet);
 
-			JsonNode? ValueNode = String.IsNullOrEmpty(Parameters.Value) ? null : JsonNode.Parse(Parameters.Value);
+			JsonNode? ValueNode;
+			try
+			{
+				ValueNode = String.IsNullOrEmpty(Parameters.Value) ? null : JsonNode.Parse(Parameters.Value);
+			}
+			catch (Exception ex)
+			{
+				throw new AutomationException(ex, $"Unable to parse '{Parameters.Value}': {ex.Message}");
+			}
+
 			foreach (FileReference JsonFile in Files)
 			{
 				string JsonText = FileReference.Exists(JsonFile) ? await FileReference.ReadAllTextAsync(JsonFile) : "{}";
@@ -79,7 +88,15 @@ namespace AutomationTool.Tasks
 					throw new AutomationException("Key must be in JsonPath format (eg. $.Foo.Bar[123])");
 				}
 
-				JsonNode? RootNode = JsonNode.Parse(JsonText);
+				JsonNode? RootNode;
+				try
+				{
+					RootNode = JsonNode.Parse(JsonText, documentOptions: new JsonDocumentOptions { CommentHandling = JsonCommentHandling.Skip });
+				}
+				catch (Exception ex)
+				{
+					throw new AutomationException($"Error parsing {JsonFile}: {ex.Message}");
+				}
 				RootNode = MergeValue(Parameters.Key, 1, RootNode, ValueNode);
 
 				string NewJsonText = RootNode?.ToJsonString(new JsonSerializerOptions { WriteIndented = true }) ?? String.Empty;
diff --git a/Engine/Source/Programs/Horde/BuildHorde.xml b/Engine/Source/Programs/Horde/BuildHorde.xml
--- a/Engine/Source/Programs/Horde/BuildHorde.xml
+++ b/Engine/Source/Programs/Horde/BuildHorde.xml
@@ -88,8 +88,21 @@
 			<Delete Files="$(StagingDir)/Agent/..."/>
 			<DotNet Arguments="publish &quot;$(RootDir)/Engine/Source/Programs/Horde/Horde/Horde.csproj&quot; --output &quot;$(StagingDir)/Tools&quot; --runtime win-x64 --self-contained $(VersionArguments)"/>
 			<DotNet Arguments="publish &quot;$(RootDir)/Engine/Source/Programs/Horde/Horde.Server/Horde.Server.csproj&quot; --output &quot;$(StagingDir)/Server&quot; --runtime win-x64 --self-contained $(VersionArguments)"/>
-			<DotNet Arguments="publish &quot;$(RootDir)/Engine/Source/Programs/Horde/Horde.Agent/Horde.Agent.csproj&quot; --output &quot;$(StagingDir)/Agent&quot; --runtime win-x64 --self-contained $(VersionArguments)"/>
-			<DotNet Arguments="publish &quot;$(RootDir)/Engine/Source/Programs/Horde/Horde.Agent.TrayApp/Horde.Agent.TrayApp.csproj&quot; --output &quot;$(StagingDir)/Agent&quot; --runtime win-x64 --self-contained $(VersionArguments)"/>
+			<DotNet Arguments="publish &quot;$(RootDir)/Engine/Source/Programs/Horde/Horde.Agent/Horde.Agent.csproj&quot; --output &quot;$(StagingDir)/Agent&quot; --runtime win7-x64 --self-contained $(VersionArguments)"/>
+			<DotNet Arguments="publish &quot;$(RootDir)/Engine/Source/Programs/Horde/Horde.Agent.TrayApp/Horde.Agent.TrayApp.csproj&quot; --output &quot;$(StagingDir)/Agent&quot; --runtime win7-x64 --self-contained $(VersionArguments)"/>
+
+			<!-- Create another build of the agent, this time for all platforms, and chunk it for download. -->
+			<DotNet Arguments="publish &quot;$(RootDir)/Engine/Source/Programs/Horde/Horde.Agent/Horde.Agent.csproj&quot; --output &quot;$(StagingDir)/Agent-CrossPlatform&quot; $(VersionArguments)"/>
+			<Spawn Exe="$(StagingDir)/Tools/Horde.exe" Arguments="bundle create -input=&quot;$(StagingDir)/Agent-CrossPlatform&quot; -file=&quot;$(StagingDir)/Server/Tools/horde-agent.ref&quot;"/>
+			
+			<Property Name="NewTool" Separator="">
+			  {
+				"Id": "horde-agent",
+				"RefName": "default-ref",
+				"Version": "$(InformationalVersion)"
+			  }
+			</Property>
+			<WriteJsonValue File="$(StagingDir)/Server/appsettings.json" Key="$.Horde.BundledTools[]" Value="$(NewTool)"/>
 
 			<!-- Copy the Dashboard -->
 			<Copy Files="#DashboardOutputFiles" From="$(StagingDir)/Dashboard/..." To="$(StagingDir)/Server/DashboardApp/..." />
diff --git a/Engine/Source/Programs/Horde/Horde.Server/Commands/Install/SetupCommand.cs b/Engine/Source/Programs/Horde/Horde.Server/Commands/Install/SetupCommand.cs
deleted file mode 100644
diff --git a/Engine/Source/Programs/Horde/Horde.Server/ServerSettings.cs b/Engine/Source/Programs/Horde/Horde.Server/ServerSettings.cs
--- a/Engine/Source/Programs/Horde/Horde.Server/ServerSettings.cs
+++ b/Engine/Source/Programs/Horde/Horde.Server/ServerSettings.cs
@@ -412,9 +412,9 @@ namespace Horde.Server
 		public int HttpPort { get; set; } = 5000;
 
 		/// <summary>
-		/// Port for serving HTTP with TLS enabled. Uses the default Kestrel port (5001) if not specified.
+		/// Port for serving HTTP with TLS enabled.
 		/// </summary>
-		public int HttpsPort { get; set; } = 5001;
+		public int HttpsPort { get; set; } = 0;
 
 		/// <summary>
 		/// Dedicated port for serving only HTTP/2.
