Horde: Improve reporting of server connection status in tray app.

diff --git a/Engine/Source/Programs/Horde/Horde.Agent/Leases/LeaseManager.cs b/Engine/Source/Programs/Horde/Horde.Agent/Leases/LeaseManager.cs
--- a/Engine/Source/Programs/Horde/Horde.Agent/Leases/LeaseManager.cs
+++ b/Engine/Source/Programs/Horde/Horde.Agent/Leases/LeaseManager.cs
@@ -369,7 +369,18 @@ namespace Horde.Agent.Leases
 					}
 
 					// Update the current status
-					_statusService.Set(true, _activeLeases.Count, (_activeLeases.Count == 0) ? "Waiting for work" : $"Executing {_activeLeases.Count} lease(s)");
+					if (!rpcCon.Healthy)
+					{
+						_statusService.Set(false, _activeLeases.Count, "Attempting to connect to server...");
+					}
+					else if (_activeLeases.Count == 0)
+					{
+						_statusService.Set(true, 0, "Waiting for work");
+					}
+					else
+					{
+						_statusService.Set(true, _activeLeases.Count, $"Executing {_activeLeases.Count} lease(s)");
+					}
 				}
 
 				// Update the historical update times
diff --git a/Engine/Source/Programs/Horde/Horde.Agent/Utility/RpcConnection.cs b/Engine/Source/Programs/Horde/Horde.Agent/Utility/RpcConnection.cs
--- a/Engine/Source/Programs/Horde/Horde.Agent/Utility/RpcConnection.cs
+++ b/Engine/Source/Programs/Horde/Horde.Agent/Utility/RpcConnection.cs
@@ -30,6 +30,11 @@ namespace Horde.Agent.Utility
 
 	interface IRpcConnection : IAsyncDisposable
 	{
+		/// <summary>
+		/// Reports whether the server connection is healthy
+		/// </summary>
+		bool Healthy { get; }
+
 		/// <summary>
 		/// Logger for this connection
 		/// </summary>
@@ -369,8 +374,10 @@ namespace Horde.Agent.Utility
 		private readonly TaskCompletionSource<bool> _stoppingTaskSource = new TaskCompletionSource<bool>();
 		private TaskCompletionSource<RpcSubConnection> _subConnectionTaskSource = new TaskCompletionSource<RpcSubConnection>();
 		private Task? _backgroundTask;
+		private bool _healthy;
 		private readonly ILogger _logger;
 
+		public bool Healthy => _healthy;
 		public ILogger Logger => _logger;
 
 		/// <summary>
@@ -568,6 +575,9 @@ namespace Horde.Agent.Utility
 									_subConnectionTaskSource.SetResult(subConnection);
 								}
 
+								// Set the healthy flag for reporting status to the tray app
+								_healthy = true;
+
 								// Wait for the StoppingTask token to be set, or the server to inform that it's shutting down
 								moveNextTask = call.ResponseStream.MoveNext();
 								await Task.WhenAny(_stoppingTaskSource.Task, Task.Delay(TimeSpan.FromSeconds(45.0)), moveNextTask);
@@ -601,6 +611,8 @@ namespace Horde.Agent.Utility
 				}
 				finally
 				{
+					_healthy = false;
+
 					if (subConnection != null)
 					{
 						await subConnection.DisposeAsync();
