Horde: Fix download of cross-platform agent zip from installed server.

diff --git a/Engine/Source/Programs/Horde/BuildHorde.xml b/Engine/Source/Programs/Horde/BuildHorde.xml
--- a/Engine/Source/Programs/Horde/BuildHorde.xml
+++ b/Engine/Source/Programs/Horde/BuildHorde.xml
@@ -93,12 +93,14 @@
 
 			<!-- Create another build of the agent, this time for all platforms, and chunk it for download. -->
 			<DotNet Arguments="publish &quot;$(RootDir)/Engine/Source/Programs/Horde/Horde.Agent/Horde.Agent.csproj&quot; --output &quot;$(StagingDir)/Agent-CrossPlatform&quot; $(VersionArguments)"/>
-			<Spawn Exe="$(StagingDir)/Tools/Horde.exe" Arguments="bundle create -input=&quot;$(StagingDir)/Agent-CrossPlatform&quot; -file=&quot;$(StagingDir)/Server/Tools/horde-agent.ref&quot;"/>
+			<Spawn Exe="$(StagingDir)/Tools/Horde.exe" Arguments="bundle create -input=&quot;$(StagingDir)/Agent-CrossPlatform&quot; -file=&quot;$(StagingDir)/Server/Tools/horde-agent/default.ref&quot;"/>
 			
 			<Property Name="NewTool" Separator="">
 			  {
 				"Id": "horde-agent",
-				"RefName": "default-ref",
+				"Name": "Horde Agent",
+				"Description": "Cross-platform build of the Horde Agent.",
+				"RefName": "default",
 				"Version": "$(InformationalVersion)"
 			  }
 			</Property>
