Horde: Support overflow on job tabs
#skipci
#rnx

diff --git a/Engine/Source/Programs/Horde/HordeDashboard/src/components/StreamView.tsx b/Engine/Source/Programs/Horde/HordeDashboard/src/components/StreamView.tsx
--- a/Engine/Source/Programs/Horde/HordeDashboard/src/components/StreamView.tsx
+++ b/Engine/Source/Programs/Horde/HordeDashboard/src/components/StreamView.tsx
@@ -131,13 +131,12 @@ const StreamViewInner: React.FC = observer(() => {
    const crumbTitle = `Horde: //${stream.project?.name}/${stream.name}`;
 
    const pivotItems = stream.tabs.map(tab => {
-      return <PivotItem headerText={tab.title} itemKey={tab.title} key={tab.title} headerButtonProps={{ href: `/stream/${streamId}?tab=${encodeURIComponent(tab.title)}` }} />;
+      return <PivotItem headerText={tab.title} itemKey={tab.title} key={tab.title}  />;
    });
 
-   pivotItems.unshift(<PivotItem headerText="All" itemKey="all" key="pivot_item_all" headerButtonProps={{ href: `/stream/${streamId}?tab=all` }} />)
-
-   pivotItems.unshift(<PivotItem headerText="Summary" itemKey="summary" key="pivot_item_summary" headerButtonProps={{ href: `/stream/${streamId}?tab=summary` }} />)
+   pivotItems.unshift(<PivotItem headerText="All" itemKey="all" key="pivot_item_all"  />)
 
+   pivotItems.unshift(<PivotItem headerText="Summary" itemKey="summary" key="pivot_item_summary" />)      
 
    let findJobsItems: IContextualMenuProps = { items: [] };
 
@@ -216,6 +215,7 @@ const StreamViewInner: React.FC = observer(() => {
    */
 
    const windowWidth = windowSize.width;
+   
 
    return (
       <Stack className={hordeClasses.horde}>
@@ -237,8 +237,9 @@ const StreamViewInner: React.FC = observer(() => {
             <Stack tokens={{ childrenGap: 0 }} styles={{ root: { backgroundColor: modeColors.background, width: "100%" } }}>
                <Stack style={{ width: 1440, paddingTop: 12, marginLeft: 4 }}>
                   <Stack style={{ maxWidth: windowWidth - 12 }} horizontal verticalAlign='center' verticalFill={true}>
-                     <Stack style={{ paddingLeft: 4, paddingTop: 2 }}>
+                     <Stack style={{ paddingLeft: 4, paddingTop: 2, width: 1180 }}>
                         <Pivot className={hordeClasses.pivot}
+                           overflowBehavior='menu'
                            selectedKey={queryTab}
                            linkSize="normal"
                            linkFormat="links"
