Horde: Configure the agent with the correct server URL at the end of the installation process.

diff --git a/Engine/Source/Programs/Horde/Horde.Agent/Commands/Utilities/SetServerCommand.cs b/Engine/Source/Programs/Horde/Horde.Agent/Commands/Utilities/SetServerCommand.cs
new file mode 100644
--- /dev/null
+++ b/Engine/Source/Programs/Horde/Horde.Agent/Commands/Utilities/SetServerCommand.cs
@@ -0,0 +1,132 @@
+// Copyright Epic Games, Inc. All Rights Reserved.
+
+using System;
+using System.Buffers;
+using System.Text.Json;
+using System.Text.Json.Nodes;
+using System.Threading.Tasks;
+using EpicGames.Core;
+using Microsoft.Extensions.Logging;
+
+namespace Horde.Agent.Commands.Utilities
+{
+	using JsonObject = System.Text.Json.Nodes.JsonObject;
+
+	/// <summary>
+	/// Shows capabilities of this agent
+	/// </summary>
+	[Command("setserver", "Configures a server for this agent")]
+	class SetServerCommand : Command
+	{
+		[CommandLine("-Name=", Description = "Name of the server to use")]
+		public string Name { get; set; } = "Default";
+
+		[CommandLine("-Url=", Description = "URL of the server", Required = true)]
+		public string? Url { get; set; }
+
+		[CommandLine("-Env=", Description = "Environment to configure for (Prod/Dev).")]
+		public string? Environment { get; set; }
+
+		[CommandLine("-Token=", Description = "Token to use for initial connection")]
+		public string? Token { get; set; }
+
+		[CommandLine("-Thumbprint=", Description = "Optional thumbprint of the server's SSL certificate. Will bypass the CA if matched. Useful for deploying self-signed certs.")]
+		public string? Thumbprint { get; set; }
+
+		[CommandLine("-Default", Description = "Makes this server the default")]
+		public bool Default { get; set; }
+
+		/// <inheritdoc/>
+		public override async Task<int> ExecuteAsync(ILogger logger)
+		{
+			// Update the agent to recognize the server certificate, and give it a custom token for being able to connect
+			FileReference agentConfigFile = FileReference.Combine(Program.AppDir, "appsettings.json");
+			JsonObject agentConfig = await ReadConfigAsync(agentConfigFile);
+
+			JsonObject hordeConfig = FindOrAddNode(agentConfig, "Horde", () => new JsonObject());
+			if (Default)
+			{
+				hordeConfig[nameof(AgentSettings.Server)] = Name;
+			}
+
+			JsonArray serverProfiles = FindOrAddNode(hordeConfig, nameof(AgentSettings.ServerProfiles), () => new JsonArray());
+
+			JsonObject serverProfile = FindOrAddElementByKey(serverProfiles, nameof(ServerProfile.Name), Name);
+			serverProfile[nameof(ServerProfile.Name)] = Name;
+			serverProfile[nameof(ServerProfile.Url)] = Url;
+
+			if (!String.IsNullOrEmpty(Environment))
+			{
+				serverProfile[nameof(ServerProfile.Environment)] = Environment;
+			}
+			if (!String.IsNullOrEmpty(Token))
+			{
+				serverProfile[nameof(ServerProfile.Token)] = Token;
+			}
+			if (!String.IsNullOrEmpty(Thumbprint))
+			{
+				serverProfile[nameof(ServerProfile.Thumbprint)] = new JsonArray(JsonValue.Create(Thumbprint));
+			}
+
+			await SaveConfigAsync(agentConfigFile, agentConfig);
+			return 0;
+		}
+
+		static async Task<JsonObject> ReadConfigAsync(FileReference file)
+		{
+			byte[] data = await FileReference.ReadAllBytesAsync(file);
+			JsonObject? obj = JsonNode.Parse(data, new JsonNodeOptions { PropertyNameCaseInsensitive = true }, new JsonDocumentOptions { AllowTrailingCommas = true, CommentHandling = JsonCommentHandling.Skip }) as JsonObject;
+			return obj ?? new JsonObject();
+		}
+
+		static async Task SaveConfigAsync(FileReference file, JsonNode node)
+		{
+			ArrayBufferWriter<byte> buffer = new ArrayBufferWriter<byte>();
+			using (Utf8JsonWriter writer = new Utf8JsonWriter(buffer, new JsonWriterOptions { Indented = true }))
+			{
+				node.WriteTo(writer);
+			}
+			await FileReference.WriteAllBytesAsync(file, buffer.WrittenMemory.ToArray());
+		}
+
+		static T FindOrAddNode<T>(JsonObject obj, string name, Func<T> factory) where T : JsonNode
+		{
+			JsonNode? node = obj[name];
+			if (node != null)
+			{
+				if (node is T existingTypedNode)
+				{
+					return existingTypedNode;
+				}
+				else
+				{
+					obj.Remove(name);
+				}
+			}
+
+			T newTypedNode = factory();
+			obj.Add(name, newTypedNode);
+			return newTypedNode;
+		}
+
+		static JsonObject FindOrAddElementByKey(JsonArray array, string key, string name)
+		{
+			foreach (JsonNode? element in array)
+			{
+				if (element is JsonObject obj)
+				{
+					JsonNode? node = obj[key];
+					if (node != null && (string?)node.AsValue() == name)
+					{
+						return obj;
+					}
+				}
+			}
+
+			JsonObject newObj = new JsonObject();
+			newObj[key] = name;
+			array.Add(newObj);
+			return newObj;
+		}
+	}
+}
