Horde: Fix test compile error.

diff --git a/Engine/Source/Programs/Horde/Horde.Agent.Tests/Helper.cs b/Engine/Source/Programs/Horde/Horde.Agent.Tests/Helper.cs
--- a/Engine/Source/Programs/Horde/Horde.Agent.Tests/Helper.cs
+++ b/Engine/Source/Programs/Horde/Horde.Agent.Tests/Helper.cs
@@ -58,6 +58,7 @@ namespace Horde.Agent.Tests
 		private readonly HordeRpc.HordeRpcClient _hordeRpcClient;
 		private readonly JobRpc.JobRpcClient _jobRpcClient;
 
+		public bool Healthy => true;
 		public ILogger Logger => NullLogger.Instance;
 
 		public RpcConnectionStub(GrpcChannel grpcChannel, HordeRpc.HordeRpcClient hordeRpcClient, JobRpc.JobRpcClient jobRpcClient)
