Horde: Increase the max request body size.

diff --git a/Engine/Source/Programs/Horde/Horde.Server/Commands/ServerCommand.cs b/Engine/Source/Programs/Horde/Horde.Server/Commands/ServerCommand.cs
--- a/Engine/Source/Programs/Horde/Horde.Server/Commands/ServerCommand.cs
+++ b/Engine/Source/Programs/Horde/Horde.Server/Commands/ServerCommand.cs
@@ -60,7 +60,7 @@ namespace Horde.Server.Commands
 
 					webBuilder.ConfigureKestrel(options =>
 					{
-						options.Limits.MaxRequestBodySize = 100 * 1024 * 1024;
+						options.Limits.MaxRequestBodySize = 256 * 1024 * 1024;
 						
 						// When agents are saturated with work (CPU or I/O), slow sending of gRPC data can happen.
 						// Kestrel protects against this behavior by default as it's commonly used for malicious attacks.
