Horde: Add a setting for whether to enable new agents in the farm by default.

We can either gate the addition of new agents to the farm by securing the ability to create new agents altogether, or let any machine join but disable them by default.

diff --git a/Engine/Source/Programs/Horde/Horde.Server/Server/RpcService.cs b/Engine/Source/Programs/Horde/Horde.Server/Server/RpcService.cs
--- a/Engine/Source/Programs/Horde/Horde.Server/Server/RpcService.cs
+++ b/Engine/Source/Programs/Horde/Horde.Server/Server/RpcService.cs
@@ -267,8 +267,7 @@ namespace Horde.Server.Server
 			IAgent? agent = await _agentService.GetAgentAsync(new AgentId(request.Name));
 			if (agent == null)
 			{
-				// TODO: Once we have this flow nailed down, do not enable agents until approved for joining the farm.
-				agent = await _agentService.CreateAgentAsync(request.Name, true, null);
+				agent = await _agentService.CreateAgentAsync(request.Name, _globalConfig.Value.ServerSettings.EnableNewAgentsByDefault, null);
 			}
 
 			List<AclClaimConfig> claims = new List<AclClaimConfig>();
diff --git a/Engine/Source/Programs/Horde/Horde.Server/ServerSettings.cs b/Engine/Source/Programs/Horde/Horde.Server/ServerSettings.cs
--- a/Engine/Source/Programs/Horde/Horde.Server/ServerSettings.cs
+++ b/Engine/Source/Programs/Horde/Horde.Server/ServerSettings.cs
@@ -553,9 +553,9 @@ namespace Horde.Server
 		public string CorsOrigin { get; set; } = null!;
 
 		/// <summary>
-		/// Whether to enable a schedule in test data (false by default for development builds)
+		/// Whether to automatically enable new agents by default. If false, new agents must manually be enabled before they can take on work.
 		/// </summary>
-		public bool EnableScheduleInTestData { get; set; }
+		public bool EnableNewAgentsByDefault { get; set; } = true;
 
 		/// <summary>
 		/// The number of months to retain test data
