Horde: Refactor compute API to support attachment of arbitrary buffers from workers.

diff --git a/Engine/Source/Programs/Horde/Horde.Agent/Commands/Compute/ComputeCommand.cs b/Engine/Source/Programs/Horde/Horde.Agent/Commands/Compute/ComputeCommand.cs
--- a/Engine/Source/Programs/Horde/Horde.Agent/Commands/Compute/ComputeCommand.cs
+++ b/Engine/Source/Programs/Horde/Horde.Agent/Commands/Compute/ComputeCommand.cs
@@ -127,7 +127,7 @@ namespace Horde.Agent.Commands.Compute
 			using Socket tcpSocket = new Socket(SocketType.Stream, ProtocolType.IP);
 			await tcpSocket.ConnectAsync(IPAddress.Loopback, Port);
 
-			await using (IComputeSocket socket = new RemoteComputeSocket(new TcpTransport(tcpSocket), logger))
+			await using (IComputeSocket socket = new RemoteComputeSocket(new TcpTransport(tcpSocket), ComputeSocketEndpoint.Remote, logger))
 			{
 				logger.LogInformation("Running worker...");
 				await RunWorkerAsync(socket, _memoryCache, logger, CancellationToken.None);
diff --git a/Engine/Source/Programs/Horde/Horde.Agent/Commands/Compute/CppComputeCommand.cs b/Engine/Source/Programs/Horde/Horde.Agent/Commands/Compute/CppComputeCommand.cs
--- a/Engine/Source/Programs/Horde/Horde.Agent/Commands/Compute/CppComputeCommand.cs
+++ b/Engine/Source/Programs/Horde/Horde.Agent/Commands/Compute/CppComputeCommand.cs
@@ -48,7 +48,6 @@ namespace Horde.Agent.Commands.Compute
 		protected override async Task<bool> HandleRequestAsync(IComputeLease lease, CancellationToken cancellationToken)
 		{
 			const int ControlChannelId = 0;
-			const int ProcessChannelId = 1;
 
 			// Read the task definition
 			byte[] data = await FileReference.ReadAllBytesAsync(TaskFile, cancellationToken);
@@ -63,7 +62,7 @@ namespace Horde.Agent.Commands.Compute
 			{
 				await channel.UploadFilesAsync("", sandbox, storage, cancellationToken);
 
-				await using (IComputeProcess process = await channel.ExecuteAsync(ProcessChannelId, jsonComputeTask.Executable, jsonComputeTask.Arguments, jsonComputeTask.WorkingDir, jsonComputeTask.EnvVars, cancellationToken))
+				await using (IComputeProcess process = await channel.ExecuteAsync(jsonComputeTask.Executable, jsonComputeTask.Arguments, jsonComputeTask.WorkingDir, jsonComputeTask.EnvVars, cancellationToken))
 				{
 					string? line;
 					while ((line = await process.ReadLineAsync(cancellationToken)) != null)
diff --git a/Engine/Source/Programs/Horde/Horde.Agent/Leases/Handlers/ComputeHandler.cs b/Engine/Source/Programs/Horde/Horde.Agent/Leases/Handlers/ComputeHandler.cs
--- a/Engine/Source/Programs/Horde/Horde.Agent/Leases/Handlers/ComputeHandler.cs
+++ b/Engine/Source/Programs/Horde/Horde.Agent/Leases/Handlers/ComputeHandler.cs
@@ -92,7 +92,7 @@ namespace Horde.Agent.Leases.Handlers
 				using (CancellationTokenSource cts = CancellationTokenSource.CreateLinkedTokenSource(cancellationToken))
 				{
 					await using BackgroundTask timeoutTask = BackgroundTask.StartNew(ctx => TickTimeoutAsync(transport, cts, ctx));
-					await using (IComputeSocket socket = new RemoteComputeSocket(transport, _logger))
+					await using (IComputeSocket socket = new RemoteComputeSocket(transport, ComputeSocketEndpoint.Local, _logger))
 					{
 						DirectoryReference sandboxDir = DirectoryReference.Combine(session.WorkingDir, "Sandbox", leaseId);
 						try
diff --git a/Engine/Source/Programs/Horde/Samples/RemoteClient/RemoteClientApp.cs b/Engine/Source/Programs/Horde/Samples/RemoteClient/RemoteClientApp.cs
--- a/Engine/Source/Programs/Horde/Samples/RemoteClient/RemoteClientApp.cs
+++ b/Engine/Source/Programs/Horde/Samples/RemoteClient/RemoteClientApp.cs
@@ -78,7 +78,7 @@ namespace RemoteClient
 
 		const int PrimaryChannelId = 0;
 		const int BackgroundChannelId = 1;
-		const int ChildProcessChannelId = ComputeChannel.DefaultWorkerChannelId;
+		const int ChildProcessChannelId = 100;
 
 		static async Task RunRemoteAsync(IComputeLease lease, DirectoryReference uploadDir, string executable, List<string> arguments, ILogger logger)
 		{
@@ -87,7 +87,7 @@ namespace RemoteClient
 			{
 				await channel.WaitForAttachAsync();
 
-				// Fork another message loop on channel id 2. We'll use this to run an XOR task in the background.
+				// Fork another message loop. We'll use this to run an XOR task in the background.
 				using IComputeMessageChannel backgroundChannel = lease.Socket.CreateMessageChannel(BackgroundChannelId, 4 * 1024 * 1024, logger);
 				await using BackgroundTask otherChannelTask = BackgroundTask.StartNew(ctx => RunBackgroundXorAsync(backgroundChannel));
 				await channel.ForkAsync(BackgroundChannelId, 4 * 1024 * 1024, default);
@@ -103,7 +103,7 @@ namespace RemoteClient
 				}
 
 				// Run the task remotely in the background and echo the output to the console
-				await using (IComputeProcess process = await channel.ExecuteAsync(ChildProcessChannelId, executable, arguments, null, null))
+				await using (IComputeProcess process = await channel.ExecuteAsync(executable, arguments, null, null))
 				{
 					await using BackgroundTask tickTask = BackgroundTask.StartNew(ctx => WriteNumbersAsync(lease.Socket, logger, ctx));
 
diff --git a/Engine/Source/Programs/Horde/Samples/RemoteWorker/RemoteWorkerApp.cs b/Engine/Source/Programs/Horde/Samples/RemoteWorker/RemoteWorkerApp.cs
--- a/Engine/Source/Programs/Horde/Samples/RemoteWorker/RemoteWorkerApp.cs
+++ b/Engine/Source/Programs/Horde/Samples/RemoteWorker/RemoteWorkerApp.cs
@@ -9,8 +9,12 @@ namespace RemoteWorker
 	{
 		static async Task Main()
 		{
-			using IComputeChannel channel = ComputeChannel.ConnectAsWorker();
-			Console.WriteLine("Connected to client");
+			const int ChannelId = 100;
+
+			using WorkerComputeSocket socket = WorkerComputeSocket.Open();
+
+			using IComputeChannel channel = socket.CreateChannel(ChannelId);
+			Console.WriteLine("Connected to initiator");
 
 			byte[] data = new byte[4];
 			while(await channel.TryReceiveMessageAsync(data))
diff --git a/Engine/Source/Programs/Horde/Samples/RemoteWorkerCpp/ComputeBuffer.cpp b/Engine/Source/Programs/Horde/Samples/RemoteWorkerCpp/ComputeBuffer.cpp
new file mode 100644
--- /dev/null
+++ b/Engine/Source/Programs/Horde/Samples/RemoteWorkerCpp/ComputeBuffer.cpp
@@ -0,0 +1,624 @@
+// Copyright Epic Games, Inc. All Rights Reserved.
+
+#include <windows.h>
+#include <assert.h>
+#include <iostream>
+#include <rpcdcep.h>
+#include <assert.h>
+#include "ComputeBuffer.h"
+
+namespace
+{
+	enum class EWriteState : int
+	{
+		// Writer has moved to the next chunk
+		MovedToNext = 0b00,
+
+		// Chunk is still being appended to
+		Writing = 0b10,
+
+		// This chunk marks the end of the stream
+		Complete = 0b11,
+	};
+
+	struct FChunkState
+	{
+		const long long Value;
+
+		FChunkState(long long InValue)
+			: Value(InValue)
+		{
+		}
+
+		// Constructor
+		FChunkState(EWriteState WriteState, int ReaderFlags, int Length)
+			: Value((unsigned long long)Length | ((unsigned long long)ReaderFlags << 31) | ((unsigned long long)WriteState << 62))
+		{
+		}
+
+		// Written length of this chunk
+		int GetLength() const { return (int)(Value & 0x7fffffff); }
+
+		// Set of flags which are set for each reader that still has to read from a chunk
+		int GetReaderFlags() const { return (int)((Value >> 31) & 0x7fffffff); }
+
+		// State of the writer
+		EWriteState GetWriteState() const { return (EWriteState)((unsigned long long)Value >> 62); }
+
+		// Test whether a particular reader is still referencing the chunk
+		bool HasReaderFlag(int ReaderIdx) const { return (Value & (1ULL << (31 + ReaderIdx))) != 0; }
+	};
+
+	struct FChunkStatePtr
+	{
+	public:
+		// Read the state value from memory
+		FChunkState Get() const
+		{
+			return FChunkState(InterlockedCompareExchange64((volatile long long*)&Value, 0, 0));
+		}
+
+		// Append data to the chunk
+		void Append(long long length)
+		{
+			InterlockedAdd64(&Value, length);
+		}
+
+		// Mark the chunk as being written to
+		void StartWriting(int numReaders)
+		{
+			InterlockedExchange64(&Value, FChunkState(EWriteState::Writing, (1 << numReaders) - 1, 0).Value);
+		}
+
+		// Mark this chunk as the end of the stream
+		void MarkComplete()
+		{
+			InterlockedOr64(&Value, FChunkState(EWriteState::Complete, 0, 0).Value);
+		}
+
+		// Clear the reader flag
+		void FinishReading(int ReaderIdx)
+		{
+			InterlockedAnd64(&Value, ~FChunkState((EWriteState)0, 1 << ReaderIdx, 0).Value);
+		}
+
+		// Move to the next chunk
+		void FinishWriting()
+		{
+			InterlockedAnd64(&Value, ~FChunkState(EWriteState::Writing, 0, 0).Value);
+		}
+
+	private:
+		volatile long long Value;
+	};
+
+	struct FReaderState
+	{
+		const long long Value;
+
+		FReaderState(long long InValue) : Value(InValue) { }
+		FReaderState(int chunkIdx, int offset) : Value(((unsigned long long)chunkIdx << 32) | offset) { }
+
+		int GetChunkIdx() const { return (int)(Value >> 32); }
+		int GetOffset() const { return (int)(unsigned int)Value; }
+	};
+
+	struct FReaderStatePtr
+	{
+	public:
+		FReaderState Get() const
+		{
+			return FReaderState(InterlockedCompareExchange64((volatile long long*)&Value, 0, 0));
+		}
+
+		void Set(FReaderState State)
+		{
+			InterlockedExchange64(&Value, State.Value);
+		}
+
+		void Advance(size_t Length)
+		{
+			InterlockedAdd64(&Value, Length);
+		}
+
+	private:
+		volatile long long Value;
+	};
+}
+
+struct FComputeBufferHeader
+{
+	int NumReaders;
+	int NumChunks;
+	int ChunkLength;
+	int WriteChunkIdx;
+	FChunkStatePtr Chunks[FComputeBuffer::MaxChunks];
+	FReaderStatePtr Readers[FComputeBuffer::MaxReaders];
+};
+
+struct FComputeBufferResources
+{
+	static unsigned int Counter;
+
+	wchar_t Name[260];
+
+	FComputeBufferHeader* Header;
+	unsigned char* ChunkPtrs[FComputeBuffer::MaxChunks];
+
+	void* WriterEvent;
+	void* ReaderEvents[FComputeBuffer::MaxReaders];
+
+	FComputeBufferReader Reader;
+	FComputeBufferWriter Writer;
+
+	FComputeBufferResources(const wchar_t* Name)
+		: MemoryMappedFile(nullptr)
+		, Header(nullptr)
+		, ChunkPtrs{ nullptr, }
+		, RefCount(1)
+		, WriterEvent(nullptr)
+		, ReaderEvents{ nullptr, }
+	{
+		wcscpy_s(this->Name, Name);
+	}
+
+	~FComputeBufferResources()
+	{
+		if (Header != nullptr)
+		{
+			UnmapViewOfFile(Header);
+			Header = nullptr;
+		}
+
+		if (MemoryMappedFile != nullptr)
+		{
+			CloseHandle(MemoryMappedFile);
+			MemoryMappedFile = nullptr;
+		}
+
+		if (WriterEvent != nullptr)
+		{
+			CloseHandle(WriterEvent);
+			WriterEvent = nullptr;
+		}
+
+		for (int ReaderIdx = 0; ReaderIdx < FComputeBuffer::MaxReaders; ReaderIdx++)
+		{
+			if (ReaderEvents[ReaderIdx] != nullptr)
+			{
+				CloseHandle(ReaderEvents[ReaderIdx]);
+				ReaderEvents[ReaderIdx] = nullptr;
+			}
+		}
+	}
+
+	static std::shared_ptr<FComputeBufferResources> CreateNew(const FComputeBuffer::FParams& Params)
+	{
+		unsigned long long Capacity = sizeof(FComputeBufferHeader) + (Params.NumChunks * sizeof(unsigned int)) + (Params.NumChunks * Params.ChunkLength);
+
+		const wchar_t* Name = Params.Name;
+
+		wchar_t BaseNameBuffer[MAX_PATH];
+		if (Name == nullptr)
+		{
+			DWORD Pid = GetCurrentProcessId();
+			ULONGLONG TickCount = GetTickCount64();
+			swprintf_s(BaseNameBuffer, L"Local\\COMPUTE_%u_%llu_%u", Pid, TickCount, InterlockedIncrement(&Counter));
+			Name = BaseNameBuffer;
+		}
+
+		wchar_t NameBuffer[MAX_PATH];
+		swprintf_s(NameBuffer, L"%s_M", Name);
+
+		std::shared_ptr<FComputeBufferResources> Resources = std::make_shared<FComputeBufferResources>(Name);
+
+		LARGE_INTEGER LargeInteger;
+		LargeInteger.QuadPart = Capacity;
+
+		Resources->MemoryMappedFile = CreateFileMappingW(INVALID_HANDLE_VALUE, NULL, PAGE_READWRITE, LargeInteger.HighPart, LargeInteger.LowPart, NameBuffer);
+		if (Resources->MemoryMappedFile == nullptr)
+		{
+			return nullptr;
+		}
+
+		Resources->Header = (FComputeBufferHeader*)MapViewOfFile(Resources->MemoryMappedFile, FILE_MAP_ALL_ACCESS, 0, 0, 0);
+		if (Resources->Header == nullptr)
+		{
+			return nullptr;
+		}
+
+		memset(Resources->Header, 0, sizeof(*Header));
+		Resources->Header->NumReaders = Params.NumReaders;
+		Resources->Header->NumChunks = Params.NumChunks;
+		Resources->Header->ChunkLength = Params.ChunkLength;
+
+		swprintf_s(NameBuffer, L"%s_W", Name);
+		Resources->WriterEvent = CreateEventW(NULL, TRUE, FALSE, NameBuffer);
+		if (Resources->WriterEvent == nullptr)
+		{
+			return nullptr;
+		}
+
+		for (int ReaderIdx = 0; ReaderIdx < Resources->Header->NumReaders; ReaderIdx++)
+		{
+			swprintf_s(NameBuffer, L"%s_R%d", Name, ReaderIdx);
+			Resources->ReaderEvents[ReaderIdx] = CreateEventW(NULL, TRUE, FALSE, NameBuffer);
+			if (Resources->ReaderEvents[ReaderIdx] == nullptr)
+			{
+				return nullptr;
+			}
+		}
+
+		Resources->Header->Chunks[0].StartWriting(Resources->Header->NumReaders);
+
+		InitShared(Resources);
+		return Resources;
+	}
+
+	static std::shared_ptr<FComputeBufferResources> OpenExisting(const wchar_t* Name)
+	{
+		wchar_t NameBuffer[MAX_PATH];
+		swprintf_s(NameBuffer, L"%s_M", Name);
+
+		std::shared_ptr<FComputeBufferResources> Resources = std::make_shared<FComputeBufferResources>(Name);
+
+		Resources->MemoryMappedFile = OpenFileMappingW(FILE_MAP_ALL_ACCESS, TRUE, NameBuffer);
+		if (Resources->MemoryMappedFile == nullptr)
+		{
+			return nullptr;
+		}
+
+		Resources->Header = (FComputeBufferHeader*)MapViewOfFile(Resources->MemoryMappedFile, FILE_MAP_ALL_ACCESS, 0, 0, 0);
+		if (Resources->Header == nullptr)
+		{
+			return nullptr;
+		}
+
+		swprintf_s(NameBuffer, L"%s_W", Name);
+		Resources->WriterEvent = OpenEventW(SYNCHRONIZE | EVENT_MODIFY_STATE, TRUE, NameBuffer);
+		if (Resources->WriterEvent == nullptr)
+		{
+			return nullptr;
+		}
+
+		for (int ReaderIdx = 0; ReaderIdx < Resources->Header->NumReaders; ReaderIdx++)
+		{
+			swprintf_s(NameBuffer, L"%s_R%d", Name, ReaderIdx);
+			Resources->ReaderEvents[ReaderIdx] = OpenEventW(SYNCHRONIZE | EVENT_MODIFY_STATE, TRUE, NameBuffer);
+			if (Resources->ReaderEvents[ReaderIdx] == nullptr)
+			{
+				return nullptr;
+			}
+		}
+
+		InitShared(Resources);
+		return Resources;
+	}
+
+	void AddRef()
+	{
+		InterlockedIncrement(&RefCount);
+	}
+
+	void Release()
+	{
+		if (InterlockedDecrement(&RefCount) == 0)
+		{
+			delete this;
+		}
+	}
+
+private:
+	void* MemoryMappedFile;
+	volatile long RefCount;
+
+	static void InitShared(std::shared_ptr<FComputeBufferResources> Resources)
+	{
+		Resources->Reader = FComputeBufferReader(Resources, 0);
+		Resources->Writer = FComputeBufferWriter(Resources);
+
+		FComputeBufferHeader* Header = Resources->Header;
+
+		unsigned char* NextPtr = (unsigned char*)(Header + 1);
+		for (int ChunkIdx = 0; ChunkIdx < Header->NumChunks; ChunkIdx++)
+		{
+			Resources->ChunkPtrs[ChunkIdx] = NextPtr;
+			NextPtr += Header->ChunkLength;
+		}
+	}
+};
+
+unsigned int FComputeBufferResources::Counter = 0;
+
+
+
+//// FComputeBuffer /////
+
+FComputeBuffer::FComputeBuffer()
+	: Resources(nullptr)
+{
+}
+
+FComputeBuffer::FComputeBuffer(const FComputeBuffer& Buffer)
+	: FComputeBuffer()
+{
+	if (Buffer.Resources != nullptr)
+	{
+		Resources = Buffer.Resources;
+		Resources->AddRef();
+	}
+}
+
+FComputeBuffer::FComputeBuffer(FComputeBuffer&& Buffer) noexcept
+	: Resources(Buffer.Resources)
+{
+	Buffer.Resources = nullptr;
+}
+
+FComputeBuffer::~FComputeBuffer()
+{
+	Close();
+}
+
+bool FComputeBuffer::CreateNew(const FParams& Params)
+{
+	Resources = FComputeBufferResources::CreateNew(Params);
+	return Resources != nullptr;
+}
+
+bool FComputeBuffer::OpenExisting(const wchar_t* Name)
+{
+	Resources = FComputeBufferResources::OpenExisting(Name);
+	return Resources != nullptr;
+}
+
+void FComputeBuffer::Close()
+{
+	Resources.reset();
+}
+
+FComputeBufferReader& FComputeBuffer::GetReader()
+{
+	return Resources->Reader;
+}
+
+const FComputeBufferReader& FComputeBuffer::GetReader() const
+{
+	return Resources->Reader;
+}
+
+FComputeBufferWriter& FComputeBuffer::GetWriter()
+{
+	return Resources->Writer;
+}
+
+const FComputeBufferWriter& FComputeBuffer::GetWriter() const
+{
+	return Resources->Writer;
+}
+
+
+
+
+
+
+
+//// FComputeBufferReader /////
+
+FComputeBufferReader::FComputeBufferReader()
+	: Resources(nullptr)
+	, ReaderIdx(-1)
+{
+}
+
+FComputeBufferReader::FComputeBufferReader(std::shared_ptr<FComputeBufferResources> Resources, int ReaderIdx)
+	: Resources(std::move(Resources))
+	, ReaderIdx(ReaderIdx)
+{
+}
+
+FComputeBufferReader::~FComputeBufferReader()
+{
+}
+
+bool FComputeBufferReader::IsComplete() const
+{
+	FComputeBufferHeader* Header = Resources->Header;
+
+	FReaderState ReaderState = Header->Readers[ReaderIdx].Get();
+	FChunkState ChunkState = Header->Chunks[ReaderState.GetChunkIdx()].Get();
+
+	return ChunkState.GetWriteState() == EWriteState::Complete && ReaderState.GetOffset() == ChunkState.GetLength();
+}
+
+void FComputeBufferReader::AdvanceReadPosition(size_t Size)
+{
+	FComputeBufferHeader* Header = Resources->Header;
+	Header->Readers[ReaderIdx].Advance(Size);
+}
+
+size_t FComputeBufferReader::GetMaxReadSize() const
+{
+	FComputeBufferHeader* Header = Resources->Header;
+
+	const FReaderStatePtr& readerStatePtr = Header->Readers[ReaderIdx];
+	FReaderState readerState = readerStatePtr.Get();
+
+	const FChunkStatePtr& chunkStatePtr = Header->Chunks[readerState.GetChunkIdx()];
+	FChunkState chunkState = chunkStatePtr.Get();
+
+	if (chunkState.HasReaderFlag(ReaderIdx))
+	{
+		return chunkState.GetLength() - readerState.GetOffset();
+	}
+	else
+	{
+		return 0;
+	}
+}
+
+const unsigned char* FComputeBufferReader::WaitToRead(size_t MinSize, int TimeoutMs)
+{
+	FComputeBufferHeader* Header = Resources->Header;
+	for (; ; )
+	{
+		FReaderStatePtr& readerStatePtr = Header->Readers[ReaderIdx];
+		FReaderState readerState = readerStatePtr.Get();
+
+		FChunkStatePtr& chunkStatePtr = Header->Chunks[readerState.GetChunkIdx()];
+		FChunkState chunkState = chunkStatePtr.Get();
+
+		if (!chunkState.HasReaderFlag(ReaderIdx))
+		{
+			// Wait until the current chunk is readable
+			ResetEvent(Resources->ReaderEvents[ReaderIdx]);
+			if (!chunkState.HasReaderFlag(ReaderIdx) && WaitForSingleObject(Resources->ReaderEvents[ReaderIdx], TimeoutMs) == WAIT_TIMEOUT)
+			{
+				return nullptr;
+			}
+		}
+		else if (readerState.GetOffset() + MinSize <= chunkState.GetLength())
+		{
+			// We have enough data in the chunk to be able to read a message
+			return Resources->ChunkPtrs[readerState.GetChunkIdx()] + readerState.GetOffset();
+		}
+		else if (chunkState.GetWriteState() == EWriteState::Writing)
+		{
+			// Wait until there is more data in the chunk
+			ResetEvent(Resources->ReaderEvents[ReaderIdx]);
+			if (Header->Chunks[readerState.GetChunkIdx()].Get().Value == chunkState.Value && WaitForSingleObject(Resources->ReaderEvents[ReaderIdx], TimeoutMs) == WAIT_TIMEOUT)
+			{
+				return nullptr;
+			}
+		}
+		else if (readerState.GetOffset() < chunkState.GetLength() || chunkState.GetWriteState() == EWriteState::Complete)
+		{
+			// Cannot read the requested amount of data from this chunk.
+			return nullptr;
+		}
+		else if (chunkState.GetWriteState() == EWriteState::MovedToNext)
+		{
+			// Move to the next chunk
+			chunkStatePtr.FinishReading(ReaderIdx);
+			SetEvent(Resources->WriterEvent);
+
+			int chunkIdx = readerState.GetChunkIdx() + 1;
+			if (chunkIdx == Header->NumChunks)
+			{
+				chunkIdx = 0;
+			}
+
+			readerStatePtr.Set(FReaderState(chunkIdx, 0));
+		}
+		else
+		{
+			assert(false);
+		}
+	}
+}
+
+const wchar_t* FComputeBufferReader::GetName() const
+{
+	return Resources->Name;
+}
+
+
+
+
+
+//// FComputeBufferWriter /////
+
+FComputeBufferWriter::FComputeBufferWriter()
+	: Resources(nullptr)
+{
+}
+
+FComputeBufferWriter::FComputeBufferWriter(std::shared_ptr<FComputeBufferResources> Resources)
+	: Resources(std::move(Resources))
+{
+}
+
+FComputeBufferWriter::~FComputeBufferWriter()
+{
+}
+
+void FComputeBufferWriter::MarkComplete()
+{
+	FComputeBufferHeader* Header = Resources->Header;
+	Header->Chunks[Header->WriteChunkIdx].MarkComplete();
+	SetAllReaderEvents();
+}
+
+void FComputeBufferWriter::AdvanceWritePosition(size_t Size)
+{
+	FComputeBufferHeader* Header = Resources->Header;
+	Header->Chunks[Header->WriteChunkIdx].Append(Size);
+	SetAllReaderEvents();
+}
+
+size_t FComputeBufferWriter::GetMaxWriteSize() const
+{
+	FComputeBufferHeader* Header = Resources->Header;
+	FChunkState Value = Header->Chunks[Header->WriteChunkIdx].Get();
+	return Header->ChunkLength - Value.GetLength();
+}
+
+unsigned char* FComputeBufferWriter::WaitToWrite(size_t MinSize, int TimeoutMs)
+{
+	FComputeBufferHeader* Header = Resources->Header;
+	assert(MinSize <= Header->ChunkLength);
+
+	DWORD WaitParam = (TimeoutMs == -1) ? INFINITE : (unsigned int)TimeoutMs;
+	for (; ; )
+	{
+		FChunkStatePtr WriteChunkStatePtr = Header->Chunks[Header->WriteChunkIdx];
+
+		FChunkState ChunkState = WriteChunkStatePtr.Get();
+		if (ChunkState.GetWriteState() == EWriteState::Writing)
+		{
+			int Length = ChunkState.GetLength();
+			if (Length + MinSize <= Header->ChunkLength)
+			{
+				return Resources->ChunkPtrs[Header->WriteChunkIdx] + Length;
+			}
+
+			WriteChunkStatePtr.FinishWriting(); // STATE CHANGE
+			SetAllReaderEvents();
+		}
+
+		int NextWriteChunkIdx = Header->WriteChunkIdx + 1;
+		if (NextWriteChunkIdx == Header->NumChunks)
+		{
+			NextWriteChunkIdx = 0;
+		}
+
+		FChunkStatePtr NextWriteChunkStatePtr = Header->Chunks[NextWriteChunkIdx];
+		while (NextWriteChunkStatePtr.Get().GetReaderFlags() != 0)
+		{
+			if (WaitForSingleObject(Resources->WriterEvent, WaitParam) == WAIT_TIMEOUT)
+			{
+				return nullptr;
+			}
+			ResetEvent(Resources->WriterEvent);
+		}
+
+		Header->WriteChunkIdx = NextWriteChunkIdx;
+
+		NextWriteChunkStatePtr.StartWriting(Header->NumReaders);
+	}
+}
+
+const wchar_t* FComputeBufferWriter::GetName() const
+{
+	return Resources->Name;
+}
+
+void FComputeBufferWriter::SetAllReaderEvents()
+{
+	FComputeBufferHeader* Header = Resources->Header;
+	for (int Idx = 0; Idx < Header->NumReaders; Idx++)
+	{
+		SetEvent(Resources->ReaderEvents[Idx]);
+	}
+}
+
diff --git a/Engine/Source/Programs/Horde/Samples/RemoteWorkerCpp/ComputeBuffer.h b/Engine/Source/Programs/Horde/Samples/RemoteWorkerCpp/ComputeBuffer.h
new file mode 100644
--- /dev/null
+++ b/Engine/Source/Programs/Horde/Samples/RemoteWorkerCpp/ComputeBuffer.h
@@ -0,0 +1,138 @@
+// Copyright Epic Games, Inc. All Rights Reserved.
+
+#pragma once
+
+#include <memory>
+
+class FComputeBufferReader;
+class FComputeBufferWriter;
+
+struct FComputeBufferHeader;
+struct FComputeBufferResources;
+
+//
+// Implements a ring buffer using shared memory, with one writer and multiple readers. 
+// 
+// Contention between the writer and readers is resolved by splitting the buffer into a number of
+// fixed-size chunks, which can be appended to and read immediately, but only reset once consumed
+// by all readers.
+//
+// If the read position is never advanced and only one chunk is used, the buffer functions as an
+// append-only buffer.
+//
+class FComputeBuffer
+{
+public:
+	// Maximum allowed number of readers
+	static const int MaxReaders = 16;
+
+	// Maximum allowed number of chunks
+	static const int MaxChunks = 16;
+
+	// Parameters for constructing a compute buffer
+	struct FParams
+	{
+		const wchar_t* Name = nullptr; // When creating a buffer, we should not normally set the name. We can only *OPEN* buffers with a name. Usually we create buffers and attach them to channels.
+		int NumChunks = 2;
+		int ChunkLength = 64 * 1024;
+		int NumReaders = 1;
+	};
+
+	FComputeBuffer();
+	FComputeBuffer(const FComputeBuffer& Buffer);
+	FComputeBuffer(FComputeBuffer&& Buffer) noexcept;
+	~FComputeBuffer();
+
+	// Creates a new buffer
+	bool CreateNew(const FParams& Params);
+
+	// Opens an existing shared memory buffer (typically from handles created in another process)
+	bool OpenExisting(const wchar_t* Name);
+
+	// Close the current buffer and release all allocated resources
+	void Close();
+
+	// Test if the buffer is currently open
+	bool IsOpen() const { return Resources != nullptr; }
+
+	// Gets a reference to the reader instance
+	FComputeBufferReader& GetReader();
+	const FComputeBufferReader& GetReader() const;
+
+	// Gets a reference to the writer instance
+	FComputeBufferWriter& GetWriter();
+	const FComputeBufferWriter& GetWriter() const;
+
+private:
+	std::shared_ptr<FComputeBufferResources> Resources;
+};
+
+//
+// Facilitates reading data from a compute buffer
+//
+class FComputeBufferReader
+{
+public:
+	FComputeBufferReader();
+	FComputeBufferReader(std::shared_ptr<FComputeBufferResources> Resources, int ReaderIdx);
+	~FComputeBufferReader();
+
+	// Test if the reader is valid
+	bool IsValid() const { return Resources.get() != nullptr; }
+
+	// Test whether the buffer has finished being written to (ie. MarkComplete() has been called by the writer) and all data has been read from it.
+	bool IsComplete() const;
+
+	// Move the read cursor forwards by the given number of bytes
+	void AdvanceReadPosition(size_t Size);
+
+	// Gets the amount of data that is ready to be read from a contiguous block of memory.
+	size_t GetMaxReadSize() const;
+
+	// Waits until the given amount of data has been read, and returns a pointer to it. Returns nullptr if the timeout expires, or
+	// if the requested amount of data is not in a contiguous block of memory.
+	const unsigned char* WaitToRead(size_t MinSize, int TimeoutMs = -1);
+
+private:
+	friend class FWorkerComputeSocket;
+
+	std::shared_ptr<FComputeBufferResources> Resources;
+	int ReaderIdx;
+
+	const wchar_t* GetName() const;
+};
+
+//
+// Facilitates writing data to a compute buffer
+//
+class FComputeBufferWriter
+{
+public:
+	FComputeBufferWriter();
+	FComputeBufferWriter(std::shared_ptr<FComputeBufferResources> Resources);
+	~FComputeBufferWriter();
+
+	// Test if the writer is valid
+	bool IsValid() const { return Resources.get() != nullptr; }
+
+	// Signal that we've finished writing to this buffer
+	void MarkComplete();
+
+	// Move the write cursor forward by the given number of bytes
+	void AdvanceWritePosition(size_t Size);
+
+	// Gets the length of the current write buffer.
+	size_t GetMaxWriteSize() const;
+
+	// Waits until a write buffer of the requested size is available, and returns a pointer to it. Returns nullptr if the 
+	// timeout expires before enough data has been flushed to return a buffer of the given length.
+	unsigned char* WaitToWrite(size_t MinSize, int TimeoutMs = -1);
+
+private:
+	friend class FWorkerComputeSocket;
+
+	std::shared_ptr<FComputeBufferResources> Resources;
+
+	const wchar_t* GetName() const;
+	void SetAllReaderEvents();
+};
diff --git a/Engine/Source/Programs/Horde/Samples/RemoteWorkerCpp/ComputeChannel.cpp b/Engine/Source/Programs/Horde/Samples/RemoteWorkerCpp/ComputeChannel.cpp
--- a/Engine/Source/Programs/Horde/Samples/RemoteWorkerCpp/ComputeChannel.cpp
+++ b/Engine/Source/Programs/Horde/Samples/RemoteWorkerCpp/ComputeChannel.cpp
@@ -1,9 +1,10 @@
 // Copyright Epic Games, Inc. All Rights Reserved.
 
+#include <string.h>
+#include <stdio.h>
 #include <windows.h>
-#include <iostream>
-#include <assert.h>
 #include "ComputeChannel.h"
+#include "ComputeSocket.h"
 
 FComputeChannel::FComputeChannel()
 {
@@ -11,82 +12,90 @@ FComputeChannel::FComputeChannel()
 
 FComputeChannel::~FComputeChannel()
 {
-	Close();
+	if (SendBufferWriter.IsValid())
+	{
+		SendBufferWriter.MarkComplete();
+	}
 }
 
-bool FComputeChannel::Open(FComputeSocket& Socket, int ChannelId)
+void FComputeChannel::Attach(FComputeSocket& Socket, int ChannelId, FComputeBuffer SendBuffer, FComputeBuffer RecvBuffer)
 {
-	return Open(Socket, ChannelId, ChannelId, 1024 * 1024);
+	RecvBufferReader = RecvBuffer.GetReader();
+	SendBufferWriter = SendBuffer.GetWriter();
+	Socket.AttachRecvBuffer(ChannelId, RecvBuffer.GetWriter());
+	Socket.AttachSendBuffer(ChannelId, SendBuffer.GetReader());
 }
 
-bool FComputeChannel::Open(FComputeSocket& Socket, int ChannelId, int NumChunks, int ChunkSize)
+bool FComputeChannel::Attach(FComputeSocket& Socket, int ChannelId, const FComputeBuffer::FParams& Params)
 {
-	return Open(Socket, ChannelId, NumChunks, ChunkSize, NumChunks, ChunkSize);
+	return Attach(Socket, ChannelId, Params, Params);
 }
 
-bool FComputeChannel::Open(FComputeSocket& Socket, int ChannelId, int NumSendChunks, int SendChunkSize, int NumRecvChunks, int RecvChunkSize)
+bool FComputeChannel::Attach(FComputeSocket& Socket, int ChannelId, const FComputeBuffer::FParams& SendBufferParams, const FComputeBuffer::FParams& RecvBufferParams)
 {
-	if (!SendBuffer.CreateNew(nullptr, NumSendChunks, SendChunkSize))
+	Detach();
+
+	FComputeBuffer SendBuffer;
+	if (!SendBuffer.CreateNew(SendBufferParams))
 	{
-		Close();
 		return false;
 	}
-	if (!RecvBuffer.CreateNew(nullptr, NumRecvChunks, RecvChunkSize))
+
+	FComputeBuffer RecvBuffer;
+	if (!RecvBuffer.CreateNew(RecvBufferParams))
 	{
-		Close();
 		return false;
 	}
 
-	Socket.AttachSendBuffer(ChannelId, SendBuffer);
-	Socket.AttachRecvBuffer(ChannelId, RecvBuffer);
-
+	Attach(Socket, ChannelId, SendBuffer, RecvBuffer);
 	return true;
 }
 
-void FComputeChannel::Close()
+void FComputeChannel::Detach()
 {
-	SendBuffer.MarkComplete();
-	SendBuffer.Close();
-	RecvBuffer.Close();
+	RecvBufferReader = FComputeBufferReader();
+	SendBufferWriter = FComputeBufferWriter();
 }
 
-void FComputeChannel::Send(const void* Data, size_t Length)
+void FComputeChannel::MarkComplete()
 {
-	while (Length > 0)
-	{
-		size_t Size;
-		unsigned char* WriteMemory = SendBuffer.GetWriteMemory(Size);
-
-		if (Size == 0)
-		{
-			SendBuffer.WaitToWrite(0);
-			continue;
-		}
+	SendBufferWriter.MarkComplete();
+}
 
-		size_t CopyLength = (Size < Length) ? Size : Length;
-		memcpy(WriteMemory, Data, CopyLength);
-		SendBuffer.AdvanceWritePosition(CopyLength);
+size_t FComputeChannel::Send(const void* Data, size_t Size, int TimeoutMs)
+{
+	unsigned char* SendData = SendBufferWriter.WaitToWrite(1, TimeoutMs);
+	if (SendData == nullptr)
+	{
+		return 0;
+	}
 
-		Data = (unsigned char*)Data + CopyLength;
-		Length -= CopyLength;
+	size_t SendSize = SendBufferWriter.GetMaxWriteSize();
+	if (Size < SendSize)
+	{
+		SendSize = Size;
 	}
+
+	memcpy(SendData, Data, SendSize);
+	SendBufferWriter.AdvanceWritePosition(SendSize);
+	return SendSize;
 }
 
-size_t FComputeChannel::Receive(void* Data, size_t Length)
+size_t FComputeChannel::Recv(void* Data, size_t Size, int TimeoutMs)
 {
-	size_t Size;
-	const unsigned char* ReadMemory = RecvBuffer.GetReadMemory(Size);
-
-	while (Size == 0 && !RecvBuffer.IsComplete())
+	const unsigned char* RecvData = RecvBufferReader.WaitToRead(1, TimeoutMs);
+	if (RecvData == nullptr)
 	{
-		RecvBuffer.WaitToRead(0);
-		ReadMemory = RecvBuffer.GetReadMemory(Size);
+		return 0;
 	}
 
-	size_t CopyLength = (Size < Length) ? Size : Length;
-	memcpy(Data, ReadMemory, CopyLength);
-	RecvBuffer.AdvanceReadPosition(CopyLength);
+	size_t RecvSize = RecvBufferReader.GetMaxReadSize();
+	if (Size < RecvSize)
+	{
+		RecvSize = Size;
+	}
 
-	return CopyLength;
+	memcpy(Data, RecvData, RecvSize);
+	RecvBufferReader.AdvanceReadPosition(RecvSize);
+	return RecvSize;
 }
-
diff --git a/Engine/Source/Programs/Horde/Samples/RemoteWorkerCpp/ComputeChannel.h b/Engine/Source/Programs/Horde/Samples/RemoteWorkerCpp/ComputeChannel.h
--- a/Engine/Source/Programs/Horde/Samples/RemoteWorkerCpp/ComputeChannel.h
+++ b/Engine/Source/Programs/Horde/Samples/RemoteWorkerCpp/ComputeChannel.h
@@ -2,26 +2,41 @@
 
 #pragma once
 
-#include "SharedMemoryBuffer.h"
-#include "ComputeSocket.h"
+#include "ComputeBuffer.h"
 
+class FComputeSocket;
+
+//
+// Allows bi-directional communication between two nodes using compute buffers
+//
 class FComputeChannel
 {
 public:
-	static const int DefaultWorkerChannelId = 100;
-
 	FComputeChannel();
 	~FComputeChannel();
 
-	bool Open(FComputeSocket& Socket, int ChannelId);
-	bool Open(FComputeSocket& Socket, int ChannelId, int NumChunks, int ChunkSize);
-	bool Open(FComputeSocket& Socket, int ChannelId, int NumSendChunks, int SendChunkSize, int NumRecvChunks, int RecvChunkSize);
-	void Close();
+	// Creates a new channel using default parameters
+	void Attach(FComputeSocket& Socket, int ChannelId, FComputeBuffer SendBuffer, FComputeBuffer RecvBuffer);
+
+	// Creates a new channel using the same parameters for the send and receive buffers
+	bool Attach(FComputeSocket& Socket, int ChannelId, const FComputeBuffer::FParams& Params);
+
+	// Creates a new channel using the same parameters for the send and receive buffers
+	bool Attach(FComputeSocket& Socket, int ChannelId, const FComputeBuffer::FParams& SendParams, const FComputeBuffer::FParams& RecvParams);
+
+	// Close the current buffer and release all allocated resources
+	void Detach();
+
+	// Indicate to the remote that no more data will be sent.
+	void MarkComplete();
+
+	// Sends bytes to the remote. 
+	size_t Send(const void* Data, size_t Size, int TimeoutMs = -1);
 
-	void Send(const void* Data, size_t Length);
-	size_t Receive(void* Data, size_t Length);
+	// Reads as many bytes as are available from the socket.
+	size_t Recv(void* Data, size_t Size, int TimeoutMs = -1);
 
 private:
-	FSharedMemoryBuffer RecvBuffer;
-	FSharedMemoryBuffer SendBuffer;
+	FComputeBufferReader RecvBufferReader;
+	FComputeBufferWriter SendBufferWriter;
 };
diff --git a/Engine/Source/Programs/Horde/Samples/RemoteWorkerCpp/ComputeSocket.cpp b/Engine/Source/Programs/Horde/Samples/RemoteWorkerCpp/ComputeSocket.cpp
--- a/Engine/Source/Programs/Horde/Samples/RemoteWorkerCpp/ComputeSocket.cpp
+++ b/Engine/Source/Programs/Horde/Samples/RemoteWorkerCpp/ComputeSocket.cpp
@@ -3,71 +3,181 @@
 #include <windows.h>
 #include <iostream>
 #include <assert.h>
+#include <uchar.h>
+#include <unordered_set>
+#include <unordered_map>
+#include <vector>
 #include "ComputeSocket.h"
 
-const char* const FComputeSocket::EnvVarName = "UE_HORDE_COMPUTE_IPC";
-
-FComputeSocket::FComputeSocket()
+FComputeSocket::~FComputeSocket()
 {
 }
 
-FComputeSocket::~FComputeSocket()
+
+//////////////////////////////////////////////////////
+
+const wchar_t* const FWorkerComputeSocket::EnvVarName = L"UE_HORDE_COMPUTE_IPC";
+
+enum class FWorkerComputeSocket::EMessageType
+{
+	AttachRecvBuffer = 0,
+	AttachSendBuffer = 1,
+};
+
+FWorkerComputeSocket::FWorkerComputeSocket()
 {
-	Close();
 }
 
-bool FComputeSocket::Open()
+FWorkerComputeSocket::~FWorkerComputeSocket()
 {
 	Close();
+}
 
-	char EnvVar[MAX_PATH];
-	int Length = GetEnvironmentVariableA(EnvVarName, EnvVar, sizeof(EnvVar) / sizeof(EnvVar[0]));
+bool FWorkerComputeSocket::Open()
+{
+	wchar_t EnvVar[MAX_PATH];
+	int Length = GetEnvironmentVariableW(EnvVarName, EnvVar, sizeof(EnvVar) / sizeof(EnvVar[0]));
 
 	if (Length <= 0 || Length >= sizeof(EnvVar))
 	{
 		return false;
 	}
 
-	return CommandBuffer.OpenExisting(EnvVar);
+	return Open(EnvVar);
+}
+
+bool FWorkerComputeSocket::Open(const wchar_t* CommandBufferName)
+{
+	return CommandBuffer.OpenExisting(CommandBufferName);
 }
 
-void FComputeSocket::Close()
+void FWorkerComputeSocket::Close()
 {
 	CommandBuffer.Close();
 }
 
-void FComputeSocket::AttachRecvBuffer(int ChannelId, FSharedMemoryBuffer& Buffer)
+void FWorkerComputeSocket::AttachRecvBuffer(int ChannelId, FComputeBufferWriter Writer)
 {
-	const int AttachRecvBufferType = 0;
-	AttachBuffer(ChannelId, AttachRecvBufferType, Buffer);
+	AttachBuffer(ChannelId, EMessageType::AttachRecvBuffer, Writer.GetName());
 }
 
-void FComputeSocket::AttachSendBuffer(int ChannelId, FSharedMemoryBuffer& Buffer)
+void FWorkerComputeSocket::AttachSendBuffer(int ChannelId, FComputeBufferReader Reader)
 {
-	const int AttachSendBufferType = 1;
-	AttachBuffer(ChannelId, AttachSendBufferType, Buffer);
+	AttachBuffer(ChannelId, EMessageType::AttachSendBuffer, Reader.GetName());
 }
 
-void FComputeSocket::AttachBuffer(int ChannelId, int Type, FSharedMemoryBuffer& Buffer)
+void FWorkerComputeSocket::AttachBuffer(int ChannelId, EMessageType Type, const wchar_t* Name)
 {
-	size_t Size;
-	unsigned char* Data = CommandBuffer.GetWriteMemory(Size);
+	FComputeBufferWriter& Writer = CommandBuffer.GetWriter();
+	unsigned char* Data = Writer.WaitToWrite(1024);
 
 	size_t Len = 0;
-	Len += WriteVarUInt(Data + Len, Type);
+	Len += WriteVarUInt(Data + Len, (unsigned char)Type);
 	Len += WriteVarUInt(Data + Len, (unsigned int)ChannelId);
-	Len += WriteString(Data + Len, Buffer.GetName());
+	Len += WriteString(Data + Len, Name);
+
+	Writer.AdvanceWritePosition(Len);
+}
+
+void FWorkerComputeSocket::RunServer(FComputeBufferReader& CommandBufferReader, FComputeSocket& Socket)
+{
+	const unsigned char* Message;
+	while ((Message = CommandBufferReader.WaitToRead(1)) != nullptr)
+	{
+		size_t Len = 0;
+
+		unsigned int Type;
+		Len += ReadVarUInt(Message + Len, &Type);
+
+		EMessageType MessageType = (EMessageType)*Message;
+		switch (MessageType)
+		{
+		case EMessageType::AttachSendBuffer:
+		{
+			unsigned int ChannelId;
+			Len += ReadVarUInt(Message + Len, &ChannelId);
+
+			wchar_t Name[MAX_PATH];
+			Len += ReadString(Message + Len, Name, MAX_PATH);
+
+			FComputeBuffer Buffer;
+			if (Buffer.OpenExisting(Name))
+			{
+				Socket.AttachSendBuffer(ChannelId, Buffer.GetReader());
+			}
+			else
+			{
+				assert(false);
+			}
+		}
+		break;
+		case EMessageType::AttachRecvBuffer:
+		{
+			unsigned int ChannelId;
+			Len += ReadVarUInt(Message + Len, &ChannelId);
+
+			wchar_t Name[MAX_PATH];
+			Len += ReadString(Message + Len, Name, MAX_PATH);
+
+			FComputeBuffer Buffer;
+			if (Buffer.OpenExisting(Name))
+			{
+				Socket.AttachRecvBuffer(ChannelId, Buffer.GetWriter());
+			}
+			else
+			{
+				assert(false);
+			}
+		}
+		break;
+		default:
+			assert(false);
+			return;
+		}
+
+		CommandBufferReader.AdvanceReadPosition(Len);
+	}
+}
+
+size_t FWorkerComputeSocket::ReadVarUInt(const unsigned char* Pos, unsigned int* OutValue)
+{
+	size_t ByteCount = CountLeadingZeros((unsigned char)(~*static_cast<const unsigned char*>(Pos))) - 23;
+
+	unsigned int Value = *Pos++ & (unsigned char)(0xff >> ByteCount);
+	switch (ByteCount - 1)
+	{
+	case 8: Value <<= 8; Value |= *Pos++;
+	case 7: Value <<= 8; Value |= *Pos++;
+	case 6: Value <<= 8; Value |= *Pos++;
+	case 5: Value <<= 8; Value |= *Pos++;
+	case 4: Value <<= 8; Value |= *Pos++;
+	case 3: Value <<= 8; Value |= *Pos++;
+	case 2: Value <<= 8; Value |= *Pos++;
+	case 1: Value <<= 8; Value |= *Pos++;
+	default:
+		break;
+	}
+
+	*OutValue = Value;
+	return ByteCount;
+}
+
+size_t FWorkerComputeSocket::ReadString(const unsigned char* Pos, wchar_t* OutText, size_t OutTextMaxLen)
+{
+	unsigned int TextLen;
+	size_t Len = ReadVarUInt(Pos, &TextLen);
 
-	CommandBuffer.AdvanceWritePosition(Len);
+	int DecodedLen = MultiByteToWideChar(CP_UTF8, 0, (const char*)Pos + Len, TextLen, OutText, (int)OutTextMaxLen);
+	OutText[DecodedLen] = 0;
+
+	return Len + TextLen;
 }
 
-size_t FComputeSocket::WriteVarUInt(unsigned char* Pos, unsigned int Value)
+size_t FWorkerComputeSocket::WriteVarUInt(unsigned char* Pos, unsigned int Value)
 {
 	// Use BSR to return the log2 of the integer
 	// return 0 if value is 0
-	unsigned long BitIndex;
-	unsigned int FloorLog2 = _BitScanReverse(&BitIndex, Value) ? BitIndex : 0;
-	unsigned int ByteCount = (unsigned int)(int(FloorLog2) / 7 + 1);
+	unsigned int ByteCount = (unsigned int)(int(FloorLog2(Value)) / 7 + 1);
 
 	unsigned char* OutBytes = Pos + ByteCount - 1;
 	switch (ByteCount - 1)
@@ -83,11 +193,322 @@ size_t FComputeSocket::WriteVarUInt(unsigned char* Pos, unsigned int Value)
 	return ByteCount;
 }
 
-size_t FComputeSocket::WriteString(unsigned char* Pos, const char* Text)
+size_t FWorkerComputeSocket::WriteString(unsigned char* Pos, const wchar_t* Text)
+{
+	int EncodedLen = WideCharToMultiByte(CP_UTF8, 0, Text, -1, nullptr, 0, nullptr, nullptr);
+	EncodedLen--; // Ignore null terminator
+
+	size_t Len = WriteVarUInt(Pos, EncodedLen);
+	WideCharToMultiByte(CP_UTF8, 0, Text, -1, (char*)Pos + Len, EncodedLen, nullptr, nullptr);
+
+	return Len + EncodedLen;
+}
+
+unsigned int FWorkerComputeSocket::FloorLog2(unsigned int Value)
+{
+	// Use BSR to return the log2 of the integer
+	// return 0 if value is 0
+	unsigned long BitIndex;
+	return _BitScanReverse(&BitIndex, Value) ? BitIndex : 0;
+}
+
+unsigned int FWorkerComputeSocket::CountLeadingZeros(unsigned int Value)
+{
+	// return 32 if value is zero
+	unsigned long BitIndex;
+	_BitScanReverse64(&BitIndex, (unsigned long long)(Value) * 2 + 1);
+	return 32 - BitIndex;
+}
+
+//////////////////////////////////////////////////////
+
+enum class FRemoteComputeSocket::EControlMessageType
+{
+	Attach = -1,
+	Detach = -2,
+};
+
+struct FRemoteComputeSocket::FFrameHeader
 {
-	unsigned int NameLen = (unsigned int)strlen(Text);
-	size_t Len = WriteVarUInt(Pos, NameLen);
-	memcpy(Pos + Len, Text, (unsigned int)NameLen);
-	return Len + NameLen;
+	int Channel;
+	int Size;
+};
+
+struct FRemoteComputeSocket::FDetail
+{
+	class FCriticalSectionLock
+	{
+	public:
+		FCriticalSectionLock(CRITICAL_SECTION* InCriticalSection) : CriticalSection(InCriticalSection)
+		{
+			EnterCriticalSection(CriticalSection);
+		}
+
+		FCriticalSectionLock(const FCriticalSectionLock& Other) = delete;
+
+		~FCriticalSectionLock()
+		{
+			LeaveCriticalSection(CriticalSection);
+		}
+
+	private:
+		CRITICAL_SECTION* const CriticalSection;
+	};
+
+	class FSendThread
+	{
+	public:
+		FSendThread(FDetail& InDetail, int InChannel, FComputeBufferReader& InReader)
+			: Detail(InDetail)
+			, Channel(InChannel)
+			, Reader(InReader)
+			, ThreadHandle(nullptr)
+		{
+		}
+
+		~FSendThread()
+		{
+			Stop();
+		}
+
+		void Start()
+		{
+			Stop();
+			ThreadHandle = CreateThread(nullptr, 0, &SendThreadProc, this, 0, nullptr);
+		}
+
+		void Stop()
+		{
+			if (ThreadHandle != nullptr)
+			{
+				WaitForSingleObject(ThreadHandle, INFINITE);
+				ThreadHandle = nullptr;
+			}
+		}
+
+	private:
+		FDetail& Detail;
+		const int Channel;
+		FComputeBufferReader& Reader;
+		HANDLE ThreadHandle;
+
+		static DWORD __stdcall SendThreadProc(void* Param)
+		{
+			FSendThread* SendBufferInfo = (FSendThread*)Param;
+			FRemoteComputeSocket& Socket = SendBufferInfo->Detail.Socket;
+			FComputeBufferReader& Reader = SendBufferInfo->Reader;
+
+			FFrameHeader Header;
+			Header.Channel = SendBufferInfo->Channel;
+
+			const unsigned char* Data;
+			while ((Data = Reader.WaitToRead(1)) != nullptr)
+			{
+				FCriticalSectionLock Lock(&SendBufferInfo->Detail.CriticalSection);
+				Header.Size = (int)Reader.GetMaxReadSize();
+				Socket.SendFull(&Header, sizeof(Header));
+				Socket.SendFull(Data, Header.Size);
+			}
+
+			if (Reader.IsComplete())
+			{
+				FCriticalSectionLock Lock(&SendBufferInfo->Detail.CriticalSection);
+				Header.Size = (int)EControlMessageType::Detach;
+				Socket.SendFull(&Header, sizeof(Header));
+			}
+
+			delete SendBufferInfo;
+			return 0;
+		}
+	};
+
+	FRemoteComputeSocket& Socket;
+	CRITICAL_SECTION CriticalSection;
+
+	HANDLE RecvThreadHandle;
+
+	std::unordered_map<int, FComputeBufferWriter*> Writers;
+	std::unordered_map<int, std::unique_ptr<FSendThread>> SendThreads;
+
+	std::unordered_set<int> AttachedRemoteBuffers;
+	std::unordered_map<int, HANDLE> AttachEvents;
+
+	FDetail(FRemoteComputeSocket& InSocket)
+		: Socket(InSocket)
+	{
+		InitializeCriticalSection(&CriticalSection);
+		RecvThreadHandle = CreateThread(nullptr, 0, &RecvThread, this, 0, nullptr);
+	}
+
+	~FDetail()
+	{
+		WaitForSingleObject(RecvThreadHandle, INFINITE);
+		DeleteCriticalSection(&CriticalSection);
+	}
+
+	static DWORD _stdcall RecvThread(void* Param)
+	{
+		FDetail& Detail = *(FDetail*)Param;
+		FRemoteComputeSocket& Socket = Detail.Socket;
+
+		std::unordered_map<int, FComputeBufferWriter*> CachedWriters;
+
+		// Process messages from the remote
+		FFrameHeader Header;
+		while (Socket.RecvFull(&Header, sizeof(Header)))
+		{
+			if (Header.Size >= 0)
+			{
+				Detail.ReadFrame(CachedWriters, Header.Channel, Header.Size);
+			}
+			else if (Header.Size == (int)EControlMessageType::Attach)
+			{
+				Detail.AttachRemoteRecvBuffer(Header.Channel);
+			}
+			else if (Header.Size == (int)EControlMessageType::Detach)
+			{
+				Detail.DetachRecvBuffer(CachedWriters, Header.Channel);
+			}
+			else
+			{
+				assert(false);
+			}
+		}
+		return 0;
+	}
+
+	bool ReadFrame(std::unordered_map<int, FComputeBufferWriter*>& CachedWriters, int Channel, int Size)
+	{
+		std::unordered_map<int, FComputeBufferWriter*>::iterator Iter = CachedWriters.find(Channel);
+		if (Iter == CachedWriters.end())
+		{
+			FCriticalSectionLock Lock(&CriticalSection);
+
+			Iter = Writers.find(Channel);
+			if (Iter == Writers.end())
+			{
+				return false;
+			}
+
+			Iter = CachedWriters.insert(*Iter).first;
+		}
+
+		FComputeBufferWriter* Writer = Iter->second;
+
+		unsigned char* Data = Writer->WaitToWrite(Size);
+		if (!Socket.RecvFull(Data, Size))
+		{
+			return false;
+		}
+
+		Writer->AdvanceWritePosition(Size);
+		return true;
+	}
+
+	void AttachRecvBuffer(int ChannelId, FComputeBufferWriter& Writer)
+	{
+		FCriticalSectionLock Lock(&CriticalSection);
+		Writers.insert(std::pair<int, FComputeBufferWriter*>(ChannelId, &Writer));
+
+		FFrameHeader Header;
+		Header.Channel = ChannelId;
+		Header.Size = (int)EControlMessageType::Attach;
+		Socket.SendFull(&Header, sizeof(Header));
+	}
+
+	void AttachSendBuffer(int ChannelId, FComputeBufferReader& Reader)
+	{
+		FCriticalSectionLock Lock(&CriticalSection);
+		SendThreads.insert(std::make_pair(ChannelId, std::make_unique<FSendThread>(*this, ChannelId, Reader)));
+	}
+
+	void AttachRemoteRecvBuffer(int Channel)
+	{
+		FCriticalSectionLock Lock(&CriticalSection);
+		AttachedRemoteBuffers.insert(Channel);
+
+		std::unordered_map<int, HANDLE>::iterator AttachEventIter = AttachEvents.find(Channel);
+		if (AttachEventIter != AttachEvents.end())
+		{
+			SetEvent(AttachEventIter->second);
+			AttachEvents.erase(AttachEventIter);
+		}
+	}
+
+	void DetachRecvBuffer(std::unordered_map<int, FComputeBufferWriter*>& CachedWriters, int Channel)
+	{
+		CachedWriters.erase(Channel);
+
+		FCriticalSectionLock Lock(&CriticalSection);
+
+		std::unordered_map<int, FComputeBufferWriter*>::iterator Iter = Writers.find(Channel);
+		if (Iter != Writers.end())
+		{
+			Iter->second->MarkComplete();
+			Writers.erase(Iter);
+		}
+	}
+};
+
+FRemoteComputeSocket::FRemoteComputeSocket()
+	: Detail(nullptr)
+{
+}
+
+FRemoteComputeSocket::~FRemoteComputeSocket()
+{
+	delete Detail;
+}
+
+void FRemoteComputeSocket::Start()
+{
+	if (Detail == nullptr)
+	{
+		Detail = new FDetail(*this);
+	}
+}
+
+void FRemoteComputeSocket::AttachRecvBuffer(int ChannelId, FComputeBufferWriter Writer)
+{
+	Detail->AttachRecvBuffer(ChannelId, Writer);
+}
+
+void FRemoteComputeSocket::AttachSendBuffer(int ChannelId, FComputeBufferReader Reader)
+{
+	Detail->AttachSendBuffer(ChannelId, Reader);
+}
+
+bool FRemoteComputeSocket::SendFull(const void* Data, size_t Size)
+{
+	const unsigned char* RemainingData = (const unsigned char*)Data;
+	for (size_t RemainingSize = Size; RemainingSize > 0; )
+	{
+		size_t SentSize = Send(RemainingData, RemainingSize);
+		if (SentSize == 0)
+		{
+			return false;
+		}
+
+		RemainingData += SentSize;
+		RemainingSize -= SentSize;
+	}
+	return true;
+}
+
+bool FRemoteComputeSocket::RecvFull(void* Data, size_t Size)
+{
+	unsigned char* RemainingData = (unsigned char*)Data;
+	for (size_t RemainingSize = Size; RemainingSize > 0; )
+	{
+		size_t RecvSize = Recv(RemainingData, RemainingSize);
+		if (RecvSize == 0)
+		{
+			return false;
+		}
+
+		RemainingData += RecvSize;
+		RemainingSize -= RecvSize;
+	}
+	return true;
 }
 
diff --git a/Engine/Source/Programs/Horde/Samples/RemoteWorkerCpp/ComputeSocket.h b/Engine/Source/Programs/Horde/Samples/RemoteWorkerCpp/ComputeSocket.h
--- a/Engine/Source/Programs/Horde/Samples/RemoteWorkerCpp/ComputeSocket.h
+++ b/Engine/Source/Programs/Horde/Samples/RemoteWorkerCpp/ComputeSocket.h
@@ -2,33 +2,104 @@
 
 #pragma once
 
-#include "SharedMemoryBuffer.h"
+#include "ComputeBuffer.h"
+#include "ComputeChannel.h"
 
+//
+// Connection to a remote machine that multiplexes data into and out-of multiple buffers
+// attached to different channel numbers.
+//
 class FComputeSocket
 {
-	static const char* const EnvVarName;
+public:
+	virtual ~FComputeSocket();
+
+	// Attaches a new buffer for receiving data
+	virtual void AttachRecvBuffer(int ChannelId, FComputeBufferWriter Writer) = 0;
+
+	// Attaches a new buffer for sending data */
+	virtual void AttachSendBuffer(int ChannelId, FComputeBufferReader Reader) = 0;
+};
 
+//
+// Socket used by a worker process to communicate with a host running on the same machine
+// using shared memory to attach new buffers.
+//
+class FWorkerComputeSocket final : public FComputeSocket
+{
 public:
-	FComputeSocket();
-	~FComputeSocket();
+	static const wchar_t* const EnvVarName;
+
+	FWorkerComputeSocket();
+	~FWorkerComputeSocket();
 
-	/** Opens a connection to the agent process */
+	// Opens a connection to the agent process using a command buffer read from an environment variable (EnvVarName)
 	bool Open();
 
-	/** Close the current connection */
+	// Opens a connection to the agent process using a specific command buffer name
+	bool Open(const wchar_t* CommandBufferName);
+
+	// Close the current connection
 	void Close();
 
-	/** Attaches a new buffer for receiving data */
-	void AttachRecvBuffer(int ChannelId, FSharedMemoryBuffer& Buffer);
+	// Attaches a new buffer for receiving data
+	virtual void AttachRecvBuffer(int ChannelId, FComputeBufferWriter Writer) override;
+
+	// Attaches a new buffer for sending data
+	virtual void AttachSendBuffer(int ChannelId, FComputeBufferReader Reader) override;
 
-	/** Attaches a new buffer for sending data */
-	void AttachSendBuffer(int ChannelId, FSharedMemoryBuffer& Buffer);
+	// Reads and handles a command from the command buffer
+	static void RunServer(FComputeBufferReader& CommandBufferReader, FComputeSocket& Socket);
 
 private:
-	FSharedMemoryBuffer CommandBuffer;
+	enum class EMessageType;
+
+	FComputeBuffer CommandBuffer;
 
-	void AttachBuffer(int ChannelId, int Type, FSharedMemoryBuffer& Buffer);
+	void AttachBuffer(int ChannelId, EMessageType Type, const wchar_t* Name);
+
+	static size_t ReadVarUInt(const unsigned char* Pos, unsigned int* OutValue);
+	static size_t ReadString(const unsigned char* Pos, wchar_t* OutText, size_t OutTextMaxLen);
 
 	static size_t WriteVarUInt(unsigned char* Pos, unsigned int Value);
-	static size_t WriteString(unsigned char* Pos, const char* Text);
+	static size_t WriteString(unsigned char* Pos, const wchar_t* Text);
+
+	static unsigned int FloorLog2(unsigned int Value);
+	static unsigned int CountLeadingZeros(unsigned int Value);
+};
+
+//
+// Implementation of FComputeSocket that communicates via an overridable transport mechanism.
+//
+class FRemoteComputeSocket : public FComputeSocket
+{
+public:
+	FRemoteComputeSocket();
+	~FRemoteComputeSocket();
+
+	// Attaches a new buffer for receiving data
+	virtual void AttachRecvBuffer(int ChannelId, FComputeBufferWriter Writer) override final;
+
+	// Attaches a new buffer for sending data */
+	virtual void AttachSendBuffer(int ChannelId, FComputeBufferReader Reader) override final;
+
+protected:
+	// Starts reading and dispatching data in a background thread
+	void Start();
+
+	// Sends data to the remote
+	virtual size_t Send(const void* Data, size_t Size) = 0;
+
+	// Receives data from the remote
+	virtual size_t Recv(void* Data, size_t Size) = 0;
+
+private:
+	struct FDetail;
+	struct FFrameHeader;
+	enum class EControlMessageType;
+
+	bool SendFull(const void* Data, size_t Size);
+	bool RecvFull(void* Data, size_t Size);
+
+	FDetail* Detail;
 };
diff --git a/Engine/Source/Programs/Horde/Samples/RemoteWorkerCpp/Main.cpp b/Engine/Source/Programs/Horde/Samples/RemoteWorkerCpp/Main.cpp
--- a/Engine/Source/Programs/Horde/Samples/RemoteWorkerCpp/Main.cpp
+++ b/Engine/Source/Programs/Horde/Samples/RemoteWorkerCpp/Main.cpp
@@ -2,11 +2,16 @@
 
 #include <iostream>
 #include "ComputeChannel.h"
-#include "SharedMemoryBuffer.h"
+#include "ComputeBuffer.h"
+#include "ComputeSocket.h"
+
+#include <Windows.h>
 
 int main()
 {
-	FComputeSocket Socket;
+	const int ChannelId = 100;
+
+	FWorkerComputeSocket Socket;
 	if (!Socket.Open())
 	{
 		std::cout << "Environment variable not set correctly" << std::endl;
@@ -14,26 +19,26 @@ int main()
 	}
 
 	FComputeChannel Channel;
-	if (!Channel.Open(Socket, FComputeChannel::DefaultWorkerChannelId))
+	if (!Channel.Attach(Socket, ChannelId, FComputeBuffer::FParams()))
 	{
-		std::cout << "Unable to create channel to host" << std::endl;
+		std::cout << "Unable to create channel to initiator" << std::endl;
 		return 1;
 	}
 
-	std::cout << "Connected to client" << std::endl;
+	std::cout << "Connected to initiator" << std::endl;
 
 	size_t Length = 0;
 	char Buffer[4];
 
 	for (;;)
 	{
-		size_t Read = Channel.Receive(Buffer + Length, sizeof(Buffer) - Length);
-		if (Read == 0)
+		size_t RecvLength = Channel.Recv(Buffer + Length, sizeof(Buffer) - Length);
+		if (RecvLength == 0)
 		{
 			return 0;
 		}
 
-		Length += Read;
+		Length += RecvLength;
 
 		if (Length >= 4)
 		{
@@ -42,7 +47,7 @@ int main()
 		}
 	}
 
-	Channel.Close();
+	Channel.MarkComplete();
 	return 0;
 }
 
diff --git a/Engine/Source/Programs/Horde/Samples/RemoteWorkerCpp/SharedMemoryBuffer.cpp b/Engine/Source/Programs/Horde/Samples/RemoteWorkerCpp/SharedMemoryBuffer.cpp
deleted file mode 100644
diff --git a/Engine/Source/Programs/Horde/Samples/RemoteWorkerCpp/SharedMemoryBuffer.h b/Engine/Source/Programs/Horde/Samples/RemoteWorkerCpp/SharedMemoryBuffer.h
deleted file mode 100644
diff --git a/Engine/Source/Programs/Shared/EpicGames.Horde.Tests/BufferTests.cs b/Engine/Source/Programs/Shared/EpicGames.Horde.Tests/BufferTests.cs
--- a/Engine/Source/Programs/Shared/EpicGames.Horde.Tests/BufferTests.cs
+++ b/Engine/Source/Programs/Shared/EpicGames.Horde.Tests/BufferTests.cs
@@ -97,8 +97,8 @@ namespace EpicGames.Horde.Tests
 			Pipe sourceToTargetPipe = new Pipe();
 			Pipe targetToSourcePipe = new Pipe();
 
-			await using RemoteComputeSocket producerSocket = new RemoteComputeSocket(new PipeTransport(targetToSourcePipe.Reader, sourceToTargetPipe.Writer), NullLogger.Instance);
-			await using RemoteComputeSocket consumerSocket = new RemoteComputeSocket(new PipeTransport(sourceToTargetPipe.Reader, targetToSourcePipe.Writer), NullLogger.Instance);
+			await using RemoteComputeSocket producerSocket = new RemoteComputeSocket(new PipeTransport(targetToSourcePipe.Reader, sourceToTargetPipe.Writer), ComputeSocketEndpoint.Local, NullLogger.Instance);
+			await using RemoteComputeSocket consumerSocket = new RemoteComputeSocket(new PipeTransport(sourceToTargetPipe.Reader, targetToSourcePipe.Writer), ComputeSocketEndpoint.Remote, NullLogger.Instance);
 
 			using IComputeBuffer consumerBuffer = createBuffer(Length);
 			consumerSocket.AttachRecvBuffer(ChannelId, consumerBuffer.Writer);
diff --git a/Engine/Source/Programs/Shared/EpicGames.Horde/Compute/Buffers/SharedMemoryBuffer.cs b/Engine/Source/Programs/Shared/EpicGames.Horde/Compute/Buffers/SharedMemoryBuffer.cs
--- a/Engine/Source/Programs/Shared/EpicGames.Horde/Compute/Buffers/SharedMemoryBuffer.cs
+++ b/Engine/Source/Programs/Shared/EpicGames.Horde/Compute/Buffers/SharedMemoryBuffer.cs
@@ -114,7 +114,7 @@ namespace EpicGames.Horde.Compute.Buffers
 			MemoryMappedView memoryMappedView = new MemoryMappedView(memoryMappedViewAccessor);
 
 			Native.EventHandle writerEvent = Native.EventHandle.CreateNew($"{name}_W", EventResetMode.ManualReset, true, HandleInheritability.Inheritable);
-			Native.EventHandle readerEvent = Native.EventHandle.CreateNew($"{name}_R", EventResetMode.ManualReset, true, HandleInheritability.Inheritable);
+			Native.EventHandle readerEvent = Native.EventHandle.CreateNew($"{name}_R0", EventResetMode.ManualReset, true, HandleInheritability.Inheritable);
 
 			HeaderPtr headerPtr = new HeaderPtr((ulong*)memoryMappedView.GetPointer(), 1, numChunks, chunkLength);
 			Memory<byte>[] chunks = CreateChunks(headerPtr.NumChunks, headerPtr.ChunkLength, memoryMappedView);
@@ -137,7 +137,7 @@ namespace EpicGames.Horde.Compute.Buffers
 			MemoryMappedViewAccessor memoryMappedViewAccessor = memoryMappedFile.CreateViewAccessor();
 			MemoryMappedView memoryMappedView = new MemoryMappedView(memoryMappedViewAccessor);
 
-			Native.EventHandle readerEvent = Native.EventHandle.OpenExisting($"{name}_R");
+			Native.EventHandle readerEvent = Native.EventHandle.OpenExisting($"{name}_R0");
 			Native.EventHandle writerEvent = Native.EventHandle.OpenExisting($"{name}_W");
 
 			HeaderPtr headerPtr = new HeaderPtr((ulong*)memoryMappedView.GetPointer());
diff --git a/Engine/Source/Programs/Shared/EpicGames.Horde/Compute/Clients/AgentComputeClient.cs b/Engine/Source/Programs/Shared/EpicGames.Horde/Compute/Clients/AgentComputeClient.cs
--- a/Engine/Source/Programs/Shared/EpicGames.Horde/Compute/Clients/AgentComputeClient.cs
+++ b/Engine/Source/Programs/Shared/EpicGames.Horde/Compute/Clients/AgentComputeClient.cs
@@ -92,7 +92,7 @@ namespace EpicGames.Horde.Compute.Clients
 			using BackgroundTask agentTask = BackgroundTask.StartNew(ctx => RunAgentAsync(_hordeAgentAssembly, _port, _logger, ctx));
 			using Socket tcpSocket = await listener.AcceptAsync(cancellationToken);
 
-			await using RemoteComputeSocket socket = new RemoteComputeSocket(new TcpTransport(tcpSocket), _logger);
+			await using RemoteComputeSocket socket = new RemoteComputeSocket(new TcpTransport(tcpSocket), ComputeSocketEndpoint.Local, _logger);
 			yield return socket;
 
 			await socket.CloseAsync(cancellationToken);
diff --git a/Engine/Source/Programs/Shared/EpicGames.Horde/Compute/Clients/LocalComputeClient.cs b/Engine/Source/Programs/Shared/EpicGames.Horde/Compute/Clients/LocalComputeClient.cs
--- a/Engine/Source/Programs/Shared/EpicGames.Horde/Compute/Clients/LocalComputeClient.cs
+++ b/Engine/Source/Programs/Shared/EpicGames.Horde/Compute/Clients/LocalComputeClient.cs
@@ -77,7 +77,7 @@ namespace EpicGames.Horde.Compute.Clients
 
 			using MemoryCache memoryCache = new MemoryCache(new MemoryCacheOptions { SizeLimit = 10 * 1024 * 1024 });
 
-			await using (IComputeSocket socket = new RemoteComputeSocket(new TcpTransport(tcpSocket), logger))
+			await using (IComputeSocket socket = new RemoteComputeSocket(new TcpTransport(tcpSocket), ComputeSocketEndpoint.Remote, logger))
 			{
 				ComputeWorker worker = new ComputeWorker(sandboxDir, memoryCache, logger);
 				await worker.RunAsync(socket, cancellationToken);
@@ -89,7 +89,7 @@ namespace EpicGames.Horde.Compute.Clients
 		public Task<IComputeLease?> TryAssignWorkerAsync(ClusterId clusterId, Requirements? requirements, CancellationToken cancellationToken)
 		{
 #pragma warning disable CA2000 // Dispose objects before losing scope
-			RemoteComputeSocket socket = new RemoteComputeSocket(new TcpTransport(_socket), _logger);
+			RemoteComputeSocket socket = new RemoteComputeSocket(new TcpTransport(_socket), ComputeSocketEndpoint.Local, _logger);
 			return Task.FromResult<IComputeLease?>(new LeaseImpl(socket));
 #pragma warning restore CA2000 // Dispose objects before losing scope
 		}
diff --git a/Engine/Source/Programs/Shared/EpicGames.Horde/Compute/Clients/ServerComputeClient.cs b/Engine/Source/Programs/Shared/EpicGames.Horde/Compute/Clients/ServerComputeClient.cs
--- a/Engine/Source/Programs/Shared/EpicGames.Horde/Compute/Clients/ServerComputeClient.cs
+++ b/Engine/Source/Programs/Shared/EpicGames.Horde/Compute/Clients/ServerComputeClient.cs
@@ -144,28 +144,28 @@ namespace EpicGames.Horde.Compute.Clients
 			_logger.LogInformation("Requesting compute resource");
 
 			// Assign a compute worker
-				HttpClient client = _createHttpClient();
+			HttpClient client = _createHttpClient();
 
 			AssignComputeRequest request = new AssignComputeRequest();
-				request.Requirements = requirements;
+			request.Requirements = requirements;
 
 			AssignComputeResponse? responseMessage;
-				using (HttpResponseMessage response = await client.PostAsync($"api/v2/compute/{clusterId}", request, _cancellationSource.Token))
-				{
-				if (response.StatusCode == HttpStatusCode.ServiceUnavailable)
+			using (HttpResponseMessage response = await client.PostAsync($"api/v2/compute/{clusterId}", request, _cancellationSource.Token))
 			{
+				if (response.StatusCode == HttpStatusCode.ServiceUnavailable)
+				{
 					_logger.LogInformation("No compute resource is available.");
 					yield break;
-			}
+				}
 
 				response.EnsureSuccessStatusCode();
 
 				responseMessage = await response.Content.ReadFromJsonAsync<AssignComputeResponse>(cancellationToken: cancellationToken);
 				if (responseMessage == null)
-			{
-					throw new InvalidOperationException();
+				{
+						throw new InvalidOperationException();
+				}
 			}
-		}
 
 			_logger.LogInformation("Connecting to {Ip} with nonce {Nonce}...", responseMessage.Ip, responseMessage.Nonce);
 
@@ -181,7 +181,7 @@ namespace EpicGames.Horde.Compute.Clients
 			// Pass the rest of the call over to the handler
 			byte[] key = StringUtils.ParseHexString(responseMessage.Key);
 
-			await using RemoteComputeSocket computeSocket = new RemoteComputeSocket(new TcpTransport(socket), _logger);
+			await using RemoteComputeSocket computeSocket = new RemoteComputeSocket(new TcpTransport(socket), ComputeSocketEndpoint.Local, _logger);
 			yield return new LeaseInfo(responseMessage.Properties, responseMessage.AssignedResources, computeSocket);
 		}
 	}
diff --git a/Engine/Source/Programs/Shared/EpicGames.Horde/Compute/ComputeWorker.cs b/Engine/Source/Programs/Shared/EpicGames.Horde/Compute/ComputeWorker.cs
--- a/Engine/Source/Programs/Shared/EpicGames.Horde/Compute/ComputeWorker.cs
+++ b/Engine/Source/Programs/Shared/EpicGames.Horde/Compute/ComputeWorker.cs
@@ -83,7 +83,7 @@ namespace EpicGames.Horde.Compute
 						case ComputeMessageType.Execute:
 							{
 								ExecuteProcessMessage executeProcess = message.ParseExecuteProcessMessage();
-								await ExecuteProcessAsync(socket, channel, executeProcess.ChannelId, executeProcess.Executable, executeProcess.Arguments, executeProcess.WorkingDir, executeProcess.EnvVars, cancellationToken);
+								await ExecuteProcessAsync(socket, channel, executeProcess.Executable, executeProcess.Arguments, executeProcess.WorkingDir, executeProcess.EnvVars, cancellationToken);
 							}
 							break;
 						case ComputeMessageType.XorRequest:
@@ -146,13 +146,13 @@ namespace EpicGames.Horde.Compute
 			}
 		}
 
-		async Task ExecuteProcessAsync(IComputeSocket socket, IComputeMessageChannel channel, int newChannelId, string executable, IReadOnlyList<string> arguments, string? workingDir, IReadOnlyDictionary<string, string?>? envVars, CancellationToken cancellationToken)
+		async Task ExecuteProcessAsync(IComputeSocket socket, IComputeMessageChannel channel, string executable, IReadOnlyList<string> arguments, string? workingDir, IReadOnlyDictionary<string, string?>? envVars, CancellationToken cancellationToken)
 		{
 			try
 			{
 				if (RuntimeInformation.IsOSPlatform(OSPlatform.Windows))
 				{
-					await ExecuteProcessWindowsAsync(socket, channel, newChannelId, executable, arguments, workingDir, envVars, cancellationToken);
+					await ExecuteProcessWindowsAsync(socket, channel, executable, arguments, workingDir, envVars, cancellationToken);
 				}
 				else
 				{
@@ -165,7 +165,7 @@ namespace EpicGames.Horde.Compute
 			}
 		}
 
-		async Task ExecuteProcessWindowsAsync(IComputeSocket socket, IComputeMessageChannel channel, int newChannelId, string executable, IReadOnlyList<string> arguments, string? workingDir, IReadOnlyDictionary<string, string?>? envVars, CancellationToken cancellationToken)
+		async Task ExecuteProcessWindowsAsync(IComputeSocket socket, IComputeMessageChannel channel, string executable, IReadOnlyList<string> arguments, string? workingDir, IReadOnlyDictionary<string, string?>? envVars, CancellationToken cancellationToken)
 		{
 			Dictionary<string, string?> newEnvVars = new Dictionary<string, string?>();
 			if (envVars != null)
@@ -176,57 +176,84 @@ namespace EpicGames.Horde.Compute
 				}
 			}
 
-			using SharedMemoryBuffer ipcBuffer = SharedMemoryBuffer.CreateNew(null, 1, 64 * 1024);
-			newEnvVars[WorkerComputeSocket.IpcEnvVar] = ipcBuffer.Name;
-
-			using BackgroundTask backgroundTask = BackgroundTask.StartNew(ctx => ProcessIpcMessagesAsync(socket, ipcBuffer.Reader, cancellationToken));
+			using (SharedMemoryBuffer ipcBuffer = SharedMemoryBuffer.CreateNew(null, 1, 64 * 1024))
+			{
+				newEnvVars[WorkerComputeSocket.IpcEnvVar] = ipcBuffer.Name;
 
-			_logger.LogInformation("Launching {Executable} {Arguments}", CommandLineArguments.Quote(executable), CommandLineArguments.Join(arguments));
-			await ExecuteProcessInternalAsync(channel, executable, arguments, workingDir, newEnvVars, cancellationToken);
-			_logger.LogInformation("Finished executing process");
+				using (BackgroundTask backgroundTask = BackgroundTask.StartNew(ctx => ProcessIpcMessagesAsync(socket, ipcBuffer.Reader, cancellationToken)))
+				{
+					_logger.LogInformation("Launching {Executable} {Arguments}", CommandLineArguments.Quote(executable), CommandLineArguments.Join(arguments));
+					await ExecuteProcessInternalAsync(channel, executable, arguments, workingDir, newEnvVars, cancellationToken);
+					_logger.LogInformation("Finished executing process");
+					ipcBuffer.Writer.MarkComplete();
+				}
+			}
 
 			_logger.LogInformation("Child process has shut down");
-
-			ipcBuffer.Writer.MarkComplete();
 		}
 
-		static async Task ProcessIpcMessagesAsync(IComputeSocket socket, IComputeBufferReader ipcReader, CancellationToken cancellationToken)
+		async Task ProcessIpcMessagesAsync(IComputeSocket socket, IComputeBufferReader ipcReader, CancellationToken cancellationToken)
 		{
-			while(await ipcReader.WaitToReadAsync(1, cancellationToken))
+			List<(IpcMessage, int, SharedMemoryBuffer)> buffers = new();
+			try
 			{
-				ReadOnlyMemory<byte> memory = ipcReader.GetReadBuffer();
-				MemoryReader reader = new MemoryReader(memory);
-
-				IpcMessage message = (IpcMessage)reader.ReadUnsignedVarInt();
-				switch (message)
+				List<(int, IComputeBufferWriter)> writers = new List<(int, IComputeBufferWriter)>();
+				while (await ipcReader.WaitToReadAsync(1, cancellationToken))
 				{
-					case IpcMessage.AttachSendBuffer:
-						{
-							int channelId = (int)reader.ReadUnsignedVarInt();
-							string name = reader.ReadString();
+					ReadOnlyMemory<byte> memory = ipcReader.GetReadBuffer();
+					MemoryReader reader = new MemoryReader(memory);
 
-							using (SharedMemoryBuffer buffer = SharedMemoryBuffer.OpenExisting(name))
-							{
-								socket.AttachSendBuffer(channelId, buffer.Reader);
-							}
-						}
-						break;
-					case IpcMessage.AttachRecvBuffer:
+					IpcMessage message = (IpcMessage)reader.ReadUnsignedVarInt();
+					try
+					{
+						switch (message)
 						{
-							int channelId = (int)reader.ReadUnsignedVarInt();
-							string name = reader.ReadString();
+							case IpcMessage.AttachSendBuffer:
+								{
+									int channelId = (int)reader.ReadUnsignedVarInt();
+									string name = reader.ReadString();
+									_logger.LogDebug("Attaching send buffer for channel {ChannelId} to {Name}", channelId, name);
 
-							using (SharedMemoryBuffer buffer = SharedMemoryBuffer.OpenExisting(name))
-							{
-								socket.AttachRecvBuffer(channelId, buffer.Writer);
-							}
+									SharedMemoryBuffer buffer = SharedMemoryBuffer.OpenExisting(name);
+									buffers.Add((message, channelId, buffer));
+
+									socket.AttachSendBuffer(channelId, buffer.Reader);
+								}
+								break;
+							case IpcMessage.AttachRecvBuffer:
+								{
+									int channelId = (int)reader.ReadUnsignedVarInt();
+									string name = reader.ReadString();
+									_logger.LogDebug("Attaching recv buffer for channel {ChannelId} to {Name}", channelId, name);
+
+									SharedMemoryBuffer buffer = SharedMemoryBuffer.OpenExisting(name);
+									buffers.Add((message, channelId, buffer));
+
+									socket.AttachRecvBuffer(channelId, buffer.Writer);
+								}
+								break;
+							default:
+								throw new InvalidOperationException($"Invalid IPC message: {message}");
 						}
-						break;
-					default:
-						throw new InvalidOperationException($"Invalid IPC message: {message}");
-				}
+					}
+					catch (Exception ex)
+					{
+						_logger.LogError(ex, "Exception while processing messages from child process: {Message}", ex.Message);
+					}
 
-				ipcReader.AdvanceReadPosition(memory.Length - reader.RemainingMemory.Length);
+					ipcReader.AdvanceReadPosition(memory.Length - reader.RemainingMemory.Length);
+				}
+			}
+			finally
+			{
+				foreach ((IpcMessage message, int channelId, SharedMemoryBuffer buffer) in buffers)
+				{
+					if (buffer.Writer.MarkComplete())
+					{
+						_logger.LogWarning("Buffer added via {Message} on channel {ChannelId} was not marked complete", message, channelId);
+					}
+					buffer.Dispose();
+				}
 			}
 		}
 
diff --git a/Engine/Source/Programs/Shared/EpicGames.Horde/Compute/IComputeChannel.cs b/Engine/Source/Programs/Shared/EpicGames.Horde/Compute/IComputeChannel.cs
--- a/Engine/Source/Programs/Shared/EpicGames.Horde/Compute/IComputeChannel.cs
+++ b/Engine/Source/Programs/Shared/EpicGames.Horde/Compute/IComputeChannel.cs
@@ -4,7 +4,6 @@ using System;
 using System.IO;
 using System.Threading;
 using System.Threading.Tasks;
-using EpicGames.Horde.Compute.Buffers;
 
 namespace EpicGames.Horde.Compute
 {
@@ -33,11 +32,6 @@ namespace EpicGames.Horde.Compute
 	/// </summary>
 	public static class ComputeChannel
 	{
-		/// <summary>
-		/// Default channel id to use for worker messages
-		/// </summary>
-		public const int DefaultWorkerChannelId = 100;
-
 		class BufferedReaderChannel : IComputeChannel
 		{
 			readonly RemoteComputeSocket _socket;
@@ -65,61 +59,28 @@ namespace EpicGames.Horde.Compute
 
 		internal sealed class BufferedReaderWriterChannel : IComputeChannel
 		{
-			readonly IComputeBuffer _recvBuffer;
-			readonly IComputeBuffer _sendBuffer;
+			readonly IComputeBufferReader _recvBufferReader;
+			readonly IComputeBufferWriter _sendBufferWriter;
 
-			public BufferedReaderWriterChannel(IComputeBuffer recvBuffer, IComputeBuffer sendBuffer)
+			public BufferedReaderWriterChannel(IComputeBufferReader recvBufferReader, IComputeBufferWriter sendBufferWriter)
 			{
-				_recvBuffer = recvBuffer;
-				_sendBuffer = sendBuffer;
+				_recvBufferReader = recvBufferReader.AddRef();
+				_sendBufferWriter = sendBufferWriter.AddRef();
 			}
 
 			public void Dispose()
 			{
-				_recvBuffer.Dispose();
-				_sendBuffer.Dispose();
-			}
-
-			public ValueTask<int> ReceiveAsync(Memory<byte> buffer, CancellationToken cancellationToken = default) => _recvBuffer.Reader.ReadAsync(buffer, cancellationToken);
-
-			public ValueTask SendAsync(ReadOnlyMemory<byte> memory, CancellationToken cancellationToken = default) => _sendBuffer.Writer.WriteAsync(memory, cancellationToken);
-		}
-
-		class SimpleWorkerChannel : IComputeChannel
-		{
-			readonly WorkerComputeSocket _socket;
-			readonly SharedMemoryBuffer _recvBuffer;
-			readonly SharedMemoryBuffer _sendBuffer;
-
-			public SimpleWorkerChannel(int channelId)
-			{
-				_socket = new WorkerComputeSocket();
-
-				_recvBuffer = SharedMemoryBuffer.CreateNew(null, 1024 * 1024);
-				_socket.AttachRecvBuffer(channelId, _recvBuffer);
+				_recvBufferReader.Dispose();
 
-				_sendBuffer = SharedMemoryBuffer.CreateNew(null, 1024 * 1024);
-				_socket.AttachSendBuffer(channelId, _sendBuffer);
+				_sendBufferWriter.MarkComplete();
+				_sendBufferWriter.Dispose();
 			}
 
-			public void Dispose()
-			{
-				_sendBuffer.Writer.MarkComplete();
-				_sendBuffer.Dispose();
-				_recvBuffer.Dispose();
-				_socket.Dispose();
-			}
-
-			public ValueTask<int> ReceiveAsync(Memory<byte> buffer, CancellationToken cancellationToken = default) => _recvBuffer.Reader.ReadAsync(buffer, cancellationToken);
+			public ValueTask<int> ReceiveAsync(Memory<byte> buffer, CancellationToken cancellationToken = default) => _recvBufferReader.ReadAsync(buffer, cancellationToken);
 
-			public ValueTask SendAsync(ReadOnlyMemory<byte> memory, CancellationToken cancellationToken = default) => _sendBuffer.Writer.WriteAsync(memory, cancellationToken);
+			public ValueTask SendAsync(ReadOnlyMemory<byte> memory, CancellationToken cancellationToken = default) => _sendBufferWriter.WriteAsync(memory, cancellationToken);
 		}
 
-		/// <summary>
-		/// Creates a socket for a worker
-		/// </summary>
-		public static IComputeChannel ConnectAsWorker(int channelId = DefaultWorkerChannelId) => new SimpleWorkerChannel(channelId);
-
 		/// <summary>
 		/// Creates a channel using a socket and receive buffer
 		/// </summary>
@@ -139,7 +100,7 @@ namespace EpicGames.Horde.Compute
 		{
 			socket.AttachRecvBuffer(channelId, recvBuffer.Writer);
 			socket.AttachSendBuffer(channelId, sendBuffer.Reader);
-			return new BufferedReaderWriterChannel(recvBuffer, sendBuffer);
+			return new BufferedReaderWriterChannel(recvBuffer.Reader, sendBuffer.Writer);
 		}
 	}
 
diff --git a/Engine/Source/Programs/Shared/EpicGames.Horde/Compute/IComputeMessage.cs b/Engine/Source/Programs/Shared/EpicGames.Horde/Compute/IComputeMessage.cs
--- a/Engine/Source/Programs/Shared/EpicGames.Horde/Compute/IComputeMessage.cs
+++ b/Engine/Source/Programs/Shared/EpicGames.Horde/Compute/IComputeMessage.cs
@@ -139,12 +139,11 @@ namespace EpicGames.Horde.Compute
 	/// <summary>
 	/// Message to execute a new child process
 	/// </summary>
-	/// <param name="ChannelId">Channel id</param>
 	/// <param name="Executable">Executable path</param>
 	/// <param name="Arguments">Arguments for the executable</param>
 	/// <param name="WorkingDir">Working directory to execute in</param>
 	/// <param name="EnvVars">Environment variables for the child process. Null values unset variables.</param>
-	public record struct ExecuteProcessMessage(int ChannelId, string Executable, IReadOnlyList<string> Arguments, string? WorkingDir, IReadOnlyDictionary<string, string?> EnvVars);
+	public record struct ExecuteProcessMessage(string Executable, IReadOnlyList<string> Arguments, string? WorkingDir, IReadOnlyDictionary<string, string?> EnvVars);
 
 	/// <summary>
 	/// Response from executing a child process
@@ -309,17 +308,15 @@ namespace EpicGames.Horde.Compute
 		/// Executes a remote process
 		/// </summary>
 		/// <param name="channel">Current channel</param>
-		/// <param name="newChannelId">Channel to use for communicating with the child process</param>
 		/// <param name="executable">Executable to run, relative to the sandbox root</param>
 		/// <param name="arguments">Arguments for the child process</param>
 		/// <param name="workingDir">Working directory for the process</param>
 		/// <param name="envVars">Environment variables for the child process</param>
 		/// <param name="cancellationToken">Cancellation token for the operation</param>
-		public static async Task<IComputeProcess> ExecuteAsync(this IComputeMessageChannel channel, int newChannelId, string executable, IReadOnlyList<string> arguments, string? workingDir, IReadOnlyDictionary<string, string?>? envVars, CancellationToken cancellationToken = default)
+		public static async Task<IComputeProcess> ExecuteAsync(this IComputeMessageChannel channel, string executable, IReadOnlyList<string> arguments, string? workingDir, IReadOnlyDictionary<string, string?>? envVars, CancellationToken cancellationToken = default)
 		{
 			using (IComputeMessageBuilder request = await channel.CreateMessageAsync(ComputeMessageType.Execute, cancellationToken))
 			{
-				request.WriteInt32(newChannelId);
 				request.WriteString(executable);
 				request.WriteList(arguments, MemoryWriterExtensions.WriteString);
 				request.WriteOptionalString(workingDir);
@@ -334,12 +331,11 @@ namespace EpicGames.Horde.Compute
 		/// </summary>
 		public static ExecuteProcessMessage ParseExecuteProcessMessage(this IComputeMessage message)
 		{
-			int channelId = message.ReadInt32();
 			string executable = message.ReadString();
 			List<string> arguments = message.ReadList(MemoryReaderExtensions.ReadString);
 			string? workingDir = message.ReadOptionalString();
 			Dictionary<string, string?> envVars = message.ReadDictionary(MemoryReaderExtensions.ReadString, MemoryReaderExtensions.ReadOptionalString);
-			return new ExecuteProcessMessage(channelId, executable, arguments, workingDir, envVars);
+			return new ExecuteProcessMessage(executable, arguments, workingDir, envVars);
 		}
 
 		/// <summary>
diff --git a/Engine/Source/Programs/Shared/EpicGames.Horde/Compute/RemoteComputeSocket.cs b/Engine/Source/Programs/Shared/EpicGames.Horde/Compute/RemoteComputeSocket.cs
--- a/Engine/Source/Programs/Shared/EpicGames.Horde/Compute/RemoteComputeSocket.cs
+++ b/Engine/Source/Programs/Shared/EpicGames.Horde/Compute/RemoteComputeSocket.cs
@@ -13,6 +13,22 @@ using Microsoft.Extensions.Logging;
 
 namespace EpicGames.Horde.Compute
 {
+	/// <summary>
+	/// Enum identifying which end of the socket a particular machine is
+	/// </summary>
+	public enum ComputeSocketEndpoint
+	{
+		/// <summary>
+		/// The initiating machine
+		/// </summary>
+		Local,
+
+		/// <summary>
+		/// The remote machine
+		/// </summary>
+		Remote,
+	}
+
 	/// <summary>
 	/// Manages a set of readers and writers to buffers across a transport layer
 	/// </summary>
@@ -29,6 +45,7 @@ namespace EpicGames.Horde.Compute
 		bool _complete;
 
 		readonly IComputeTransport _transport;
+		readonly ComputeSocketEndpoint _endpoint;
 		readonly ILogger _logger;
 
 		readonly CancellationTokenSource _cancellationSource = new CancellationTokenSource();
@@ -46,16 +63,20 @@ namespace EpicGames.Horde.Compute
 		readonly BackgroundTask _attachTask;
 		readonly Channel<int> _attachChannel = Channel.CreateUnbounded<int>();
 
+		string Tag => (_endpoint == ComputeSocketEndpoint.Local)? "LOCAL": "REMOTE";
+
 		/// <summary>
 		/// Constructor
 		/// </summary>
 		/// <param name="transport">Transport to communicate with the remote</param>
+		/// <param name="endpoint">Tag for log messages</param>
 		/// <param name="logger">Logger for trace output</param>
-		public RemoteComputeSocket(IComputeTransport transport, ILogger logger)
+		public RemoteComputeSocket(IComputeTransport transport, ComputeSocketEndpoint endpoint, ILogger logger)
 		{
 			_transport = transport;
+			_endpoint = endpoint;
 			_logger = logger;
-			_recvTask = BackgroundTask.StartNew(ctx => RunReaderAsync(transport, ctx));
+			_recvTask = BackgroundTask.StartNew(ctx => RunRecvTaskAsync(transport, ctx));
 			_attachTask = BackgroundTask.StartNew(ctx => SendAttachNotificationsAsync(ctx));
 		}
 
@@ -64,6 +85,8 @@ namespace EpicGames.Horde.Compute
 		/// </summary>
 		public async ValueTask CloseAsync(CancellationToken cancellationToken)
 		{
+			_cancellationSource.CancelAfter(TimeSpan.FromSeconds(2.0));
+
 			// Make sure we close all buffers that are attached, otherwise we'll lock up waiting for send tasks to complete
 			Task[] sendTasks;
 			lock (_lockObject)
@@ -102,9 +125,9 @@ namespace EpicGames.Horde.Compute
 			GC.SuppressFinalize(this);
 		}
 
-		async Task RunReaderAsync(IComputeTransport transport, CancellationToken cancellationToken)
+		async Task RunRecvTaskAsync(IComputeTransport transport, CancellationToken cancellationToken)
 		{
-			_logger.LogTrace("Started socket reader");
+			_logger.LogTrace("[{Tag}] Started socket reader", Tag);
 
 			byte[] header = new byte[8];
 			try
@@ -120,7 +143,7 @@ namespace EpicGames.Horde.Compute
 					// Read the next packet header
 					if (!await transport.ReadOptionalAsync(header, cancellationToken))
 					{
-						_logger.LogTrace("End of socket");
+						_logger.LogTrace("[{Tag}] End of socket", Tag);
 						break;
 					}
 
@@ -136,17 +159,17 @@ namespace EpicGames.Horde.Compute
 					}
 					else if (size == (int)ControlMessageType.Attach)
 					{
-						_logger.LogTrace("Attaching recv buffer {Id}", id);
+						_logger.LogTrace("[{Tag}] Attaching recv buffer {Id}", Tag, id);
 						AttachRemoteRecvBuffer(id);
 					}
 					else if (size == (int)ControlMessageType.Detach)
 					{
-						_logger.LogTrace("Detaching recv buffer {Id}", id);
+						_logger.LogTrace("[{Tag}] Detaching recv buffer {Id}", Tag, id);
 						DetachRecvBuffer(cachedWriters, id);
 					}
 					else
 					{
-						_logger.LogWarning("Unrecognized control message: {Message}", size);
+						_logger.LogWarning("[{Tag}] Unrecognized control message: {Message}", Tag, size);
 					}
 				}
 			}
@@ -164,7 +187,7 @@ namespace EpicGames.Horde.Compute
 				}
 			}
 
-			_logger.LogTrace("Closing reader");
+			_logger.LogTrace("[{Tag}] Closing reader", Tag);
 		}
 
 		async Task ReadPacketAsync(IComputeTransport transport, int id, int size, IComputeBufferWriter writer, CancellationToken cancellationToken)
@@ -172,7 +195,7 @@ namespace EpicGames.Horde.Compute
 			Memory<byte> memory = writer.GetWriteBuffer();
 			while (memory.Length < size)
 			{
-				_logger.LogTrace("No space in buffer {Id}, flushing", id);
+				_logger.LogTrace("[{Tag}] No space in buffer {Id}, flushing", Tag, id);
 				await writer.WaitToWriteAsync(size, cancellationToken);
 				memory = writer.GetWriteBuffer();
 			}
diff --git a/Engine/Source/Programs/Shared/EpicGames.Horde/Compute/WorkerComputeSocket.cs b/Engine/Source/Programs/Shared/EpicGames.Horde/Compute/WorkerComputeSocket.cs
--- a/Engine/Source/Programs/Shared/EpicGames.Horde/Compute/WorkerComputeSocket.cs
+++ b/Engine/Source/Programs/Shared/EpicGames.Horde/Compute/WorkerComputeSocket.cs
@@ -27,7 +27,15 @@ namespace EpicGames.Horde.Compute
 		/// <summary>
 		/// Creates a socket for a worker
 		/// </summary>
-		public WorkerComputeSocket()
+		private WorkerComputeSocket(SharedMemoryBuffer commandBuffer)
+		{
+			_commandBuffer = commandBuffer;
+		}
+
+		/// <summary>
+		/// Opens a socket which allows a worker to communicate with the Horde Agent
+		/// </summary>
+		public static WorkerComputeSocket Open()
 		{
 			string? baseName = Environment.GetEnvironmentVariable(IpcEnvVar);
 			if (baseName == null)
@@ -35,7 +43,8 @@ namespace EpicGames.Horde.Compute
 				throw new InvalidOperationException($"Environment variable {IpcEnvVar} is not defined; cannot connect as worker.");
 			}
 
-			_commandBuffer = SharedMemoryBuffer.OpenExisting(baseName);
+			SharedMemoryBuffer commandBuffer = SharedMemoryBuffer.OpenExisting(baseName);
+			return new WorkerComputeSocket(commandBuffer);
 		}
 
 		/// <inheritdoc/>
@@ -73,6 +82,18 @@ namespace EpicGames.Horde.Compute
 	/// </summary>
 	public static class WorkerComputeSocketExtensions
 	{
+		/// <summary>
+		/// Creates a channel using a socket and receive buffer
+		/// </summary>
+		/// <param name="socket">Socket to use for sending data</param>
+		/// <param name="channelId">Channel id to send and receive data</param>
+		public static IComputeChannel CreateChannel(this WorkerComputeSocket socket, int channelId)
+		{
+			using SharedMemoryBuffer recvBuffer = SharedMemoryBuffer.CreateNew(null, 65536);
+			using SharedMemoryBuffer sendBuffer = SharedMemoryBuffer.CreateNew(null, 65536);
+			return CreateChannel(socket, channelId, recvBuffer, sendBuffer);
+		}
+
 		/// <summary>
 		/// Creates a channel using a socket and receive buffer
 		/// </summary>
@@ -84,7 +105,7 @@ namespace EpicGames.Horde.Compute
 		{
 			socket.AttachRecvBuffer(channelId, recvBuffer);
 			socket.AttachSendBuffer(channelId, sendBuffer);
-			return new ComputeChannel.BufferedReaderWriterChannel(recvBuffer, sendBuffer);
+			return new ComputeChannel.BufferedReaderWriterChannel(recvBuffer.Reader, sendBuffer.Writer);
 		}
 	}
 }
