Horde: Allow underscores in identifiers for *Test files

For better naming of individual test methods. See https://learn.microsoft.com/en-us/dotnet/core/testing/unit-testing-best-practices#better

diff --git a/Engine/Source/Programs/Horde/.editorconfig b/Engine/Source/Programs/Horde/.editorconfig
--- a/Engine/Source/Programs/Horde/.editorconfig
+++ b/Engine/Source/Programs/Horde/.editorconfig
@@ -1073,3 +1073,9 @@ dotnet_diagnostic.CA5404.severity = none
 
 # CA5405: Do not always skip token validation in delegates
 dotnet_diagnostic.CA5405.severity = warning
+
+
+[*Test*.cs]
+
+# CA1707: Identifiers should not contain underscores
+dotnet_diagnostic.CA1707.severity = none
\ No newline at end of file
