Horde: Fix naming of PoolUpdateServiceTest

diff --git a/Engine/Source/Programs/Horde/Horde.Server.Tests/Agents/Pools/PoolUpdateServiceTest.cs b/Engine/Source/Programs/Horde/Horde.Server.Tests/Agents/Pools/PoolUpdateServiceTest.cs
--- a/Engine/Source/Programs/Horde/Horde.Server.Tests/Agents/Pools/PoolUpdateServiceTest.cs
+++ b/Engine/Source/Programs/Horde/Horde.Server.Tests/Agents/Pools/PoolUpdateServiceTest.cs
@@ -30,10 +30,7 @@ namespace Horde.Server.Tests.Agents.Pools
 			_pool = await PoolService.CreatePoolAsync("testPool", null, true);
 			_enabledAgent = await CreateAgentAsync(_pool, true);
 			_disabledAgent = await CreateAgentAsync(_pool, false);
-			DateTime now = Clock.UtcNow;
-			Clock.UtcNow = now - TimeSpan.FromHours(9);
-			_disabledAgentBeyondGracePeriod = await CreateAgentAsync(_pool, false);
-			Clock.UtcNow = now;
+			_disabledAgentBeyondGracePeriod = await CreateAgentAsync(_pool, enabled: false, adjustClockBy: -TimeSpan.FromHours(9));
 		}
 
 		private async Task RefreshAgents()
@@ -54,43 +51,46 @@ namespace Horde.Server.Tests.Agents.Pools
 		}
 
 		[TestMethod]
-		public async Task ShutdownDisabledAgentBasedOnGlobalConfig()
+		public async Task ShutdownDisabledAgents_WithGlobalGracePeriod_RequestsShutdown()
 		{
-			Assert.IsFalse(_enabledAgent.RequestShutdown);
-			Assert.IsFalse(_disabledAgent.RequestShutdown);
-			Assert.IsFalse(_disabledAgentBeyondGracePeriod.RequestShutdown);
-			
+			// Act
 			await _pus.ShutdownDisabledAgentsAsync(CancellationToken.None);
 			await RefreshAgents();
 			
+			// Assert
 			Assert.IsFalse(_enabledAgent.RequestShutdown);
 			Assert.IsFalse(_disabledAgent.RequestShutdown);
 			Assert.IsTrue(_disabledAgentBeyondGracePeriod.RequestShutdown);
 		}
 		
 		[TestMethod]
-		public async Task ShutdownDisabledAgentBasedOnPoolConfig()
+		public async Task ShutdownDisabledAgents_WithPerPoolGracePeriod_DoesNotRequestShutdown()
 		{
+			// Arrange
 			// Explicitly set the grace period for the pool to be longer than the default of 8 hours
 			await PoolCollection.TryUpdateAsync(_pool, shutdownIfDisabledGracePeriod: TimeSpan.FromHours(24));
 			
+			// Act
 			await _pus.ShutdownDisabledAgentsAsync(CancellationToken.None);
 			await RefreshAgents();
 			
-			// Should result in no shutdown requests
+			// Assert
 			Assert.IsFalse(_enabledAgent.RequestShutdown);
 			Assert.IsFalse(_disabledAgent.RequestShutdown);
 			Assert.IsFalse(_disabledAgentBeyondGracePeriod.RequestShutdown);
 		}
 		
 		[TestMethod]
-		public async Task OnlyShutdownDisabledAgentsThatAreAutoScaled()
+		public async Task ShutdownDisabledAgents_WithAutoScalingOff_DoesNotRequestShutdown()
 		{
+			// Arrange
 			await PoolCollection.TryUpdateAsync(_pool, newEnableAutoscaling: false);
 			
+			// Act
 			await _pus.ShutdownDisabledAgentsAsync(CancellationToken.None);
 			await RefreshAgents();
 			
+			// Assert
 			Assert.IsFalse(_enabledAgent.RequestShutdown);
 			Assert.IsFalse(_disabledAgent.RequestShutdown);
 			Assert.IsFalse(_disabledAgentBeyondGracePeriod.RequestShutdown);
diff --git a/Engine/Source/Programs/Horde/Horde.Server.Tests/TestSetup.cs b/Engine/Source/Programs/Horde/Horde.Server.Tests/TestSetup.cs
--- a/Engine/Source/Programs/Horde/Horde.Server.Tests/TestSetup.cs
+++ b/Engine/Source/Programs/Horde/Horde.Server.Tests/TestSetup.cs
@@ -341,8 +341,14 @@ namespace Horde.Server.Tests
 		}
 		
 		private static int s_agentIdCounter = 1;
-		public async Task<IAgent> CreateAgentAsync(IPool pool, bool enabled = true, bool requestShutdown = false)
+		public async Task<IAgent> CreateAgentAsync(IPool pool, bool enabled = true, bool requestShutdown = false, TimeSpan? adjustClockBy = null)
 		{
+			DateTime now = Clock.UtcNow;
+			if (adjustClockBy != null)
+			{
+				Clock.UtcNow = now + adjustClockBy.Value;
+			}
+			
 			IAgent agent = await AgentService.CreateAgentAsync("TestAgent" + s_agentIdCounter++, enabled, new List<PoolId> { pool.Id });
 			agent = await AgentService.CreateSessionAsync(agent, AgentStatus.Ok, new List<string>(), new Dictionary<string, int>(), null);
 			if (requestShutdown)
@@ -350,6 +356,7 @@ namespace Horde.Server.Tests
 				await AgentCollection.TryUpdateSettingsAsync(agent, requestShutdown: true);
 			}
 			
+			Clock.UtcNow = now;
 			return agent;
 		}
 
