Horde: Fix incorrect logging of job executor being used

diff --git a/Engine/Source/Programs/Horde/Horde.Agent/Leases/Handlers/JobHandler.cs b/Engine/Source/Programs/Horde/Horde.Agent/Leases/Handlers/JobHandler.cs
--- a/Engine/Source/Programs/Horde/Horde.Agent/Leases/Handlers/JobHandler.cs
+++ b/Engine/Source/Programs/Horde/Horde.Agent/Leases/Handlers/JobHandler.cs
@@ -161,11 +161,11 @@ namespace Horde.Agent.Leases.Handlers
 		{
 			IRpcConnection rpcClient = session.RpcConnection;
 
-			batchLogger.LogInformation("Executing batch {BatchId} using {Executor} executor", options.BatchId, _settings.Executor.ToString());
-			await session.TerminateProcessesAsync(batchLogger, cancellationToken);
-
 			// Create an executor for this job
 			string executorName = String.IsNullOrEmpty(options.JobOptions.Executor) ? _settings.Executor : options.JobOptions.Executor;
+			
+			batchLogger.LogInformation("Executing batch {BatchId} using {Executor} executor", options.BatchId, executorName);
+			await session.TerminateProcessesAsync(batchLogger, cancellationToken);
 
 			IJobExecutorFactory? executorFactory = _executorFactories.FirstOrDefault(x => x.Name.Equals(executorName, StringComparison.OrdinalIgnoreCase));
 			if (executorFactory == null)
