Horde: Fix overflow issue with stream view again, this time with links surviving
#rnx
#skipci

diff --git a/Engine/Source/Programs/Horde/HordeDashboard/src/components/StreamView.tsx b/Engine/Source/Programs/Horde/HordeDashboard/src/components/StreamView.tsx
--- a/Engine/Source/Programs/Horde/HordeDashboard/src/components/StreamView.tsx
+++ b/Engine/Source/Programs/Horde/HordeDashboard/src/components/StreamView.tsx
@@ -3,7 +3,7 @@
 import { DefaultButton, IContextualMenuProps, mergeStyleSets, Pivot, PivotItem, PrimaryButton, Stack, Text, TextField } from '@fluentui/react';
 import { observer } from 'mobx-react-lite';
 import React, { useEffect, useState } from 'react';
-import { Navigate, useNavigate, useLocation, useParams } from 'react-router-dom';
+import { Navigate, useNavigate, useLocation, useParams, Link } from 'react-router-dom';
 import { useBackend } from '../backend';
 import dashboard from '../backend/Dashboard';
 import { JobFilterSimple } from '../base/utilities/filter';
@@ -131,12 +131,11 @@ const StreamViewInner: React.FC = observer(() => {
    const crumbTitle = `Horde: //${stream.project?.name}/${stream.name}`;
 
    const pivotItems = stream.tabs.map(tab => {
-      return <PivotItem headerText={tab.title} itemKey={tab.title} key={tab.title}  />;
+      return <PivotItem headerText={tab.title} itemKey={tab.title} key={tab.title} onRenderItemLink={() => <Link to={`/stream/${streamId}?tab=${encodeURIComponent(tab.title)}`} style={{color: modeColors.text}}>{tab.title}</Link>} />;
    });
 
-   pivotItems.unshift(<PivotItem headerText="All" itemKey="all" key="pivot_item_all"  />)
-
-   pivotItems.unshift(<PivotItem headerText="Summary" itemKey="summary" key="pivot_item_summary" />)      
+   pivotItems.unshift(<PivotItem headerText="All" itemKey="all" key="pivot_item_all" onRenderItemLink={() => <Link to={`/stream/${streamId}?tab=all`} style={{color: modeColors.text}}>All</Link>}/>)
+   pivotItems.unshift(<PivotItem headerText="Summary" itemKey="summary" key="pivot_item_summary" onRenderItemLink={() => <Link to={`/stream/${streamId}?tab=summary`} style={{color: modeColors.text}}>Summary</Link>}/>)
 
    let findJobsItems: IContextualMenuProps = { items: [] };
 
@@ -190,12 +189,12 @@ const StreamViewInner: React.FC = observer(() => {
 
    const newBuildTab = queryTab?.toLowerCase() === "summary" ? "all" : queryTab;
    const isSwarmTab = queryTab?.toLowerCase() === "swarm" || queryTab?.toLowerCase() === "presubmit";
-   let isIncrementalTab = queryTab?.toLowerCase() === "incremental";   
+   let isIncrementalTab = queryTab?.toLowerCase() === "incremental";
 
    if (!isIncrementalTab && currentTab) {
       isIncrementalTab = currentTab.showNames === false; /*|| currentTab.templates?.length === 1;*/
    }
-   
+
    if (!isIncrementalTab && queryTab) {
       const tabName = queryTab.toLowerCase();
       isIncrementalTab = tabName === "primary inc" || tabName === "secondary" || tabName.indexOf("incremental") !== -1;
@@ -215,7 +214,7 @@ const StreamViewInner: React.FC = observer(() => {
    */
 
    const windowWidth = windowSize.width;
-   
+
 
    return (
       <Stack className={hordeClasses.horde}>
@@ -233,7 +232,7 @@ const StreamViewInner: React.FC = observer(() => {
          }} />
          {findJobsShown && <JobSearchSimpleModal onClose={() => { setFindJobsShown(false) }} streamId={stream.id} />}
          <Stack horizontal>
-            <div key={`windowsize_streamview_${windowSize.width}_${windowSize.height}`} style={{ width: (vw / 2 - (1440/2)) - 12, flexShrink: 0, backgroundColor: modeColors.background }} />
+            <div key={`windowsize_streamview_${windowSize.width}_${windowSize.height}`} style={{ width: (vw / 2 - (1440 / 2)) - 12, flexShrink: 0, backgroundColor: modeColors.background }} />
             <Stack tokens={{ childrenGap: 0 }} styles={{ root: { backgroundColor: modeColors.background, width: "100%" } }}>
                <Stack style={{ width: 1440, paddingTop: 12, marginLeft: 4 }}>
                   <Stack style={{ maxWidth: windowWidth - 12 }} horizontal verticalAlign='center' verticalFill={true}>
