Horde: Allow uploading tools through multiple blob writes.

diff --git a/Engine/Source/Programs/Horde/Horde.Server/Logs/LogsController.cs b/Engine/Source/Programs/Horde/Horde.Server/Logs/LogsController.cs
--- a/Engine/Source/Programs/Horde/Horde.Server/Logs/LogsController.cs
+++ b/Engine/Source/Programs/Horde/Horde.Server/Logs/LogsController.cs
@@ -117,7 +117,8 @@ namespace Horde.Server.Logs
 				return Forbid();
 			}
 
-			return await StorageController.WriteBlobAsync(_storageService, Namespace.Logs, file, $"{logFile.RefName}", cancellationToken);
+			IStorageClient storageClient = await _storageService.GetClientAsync(Namespace.Logs, cancellationToken);
+			return await StorageController.WriteBlobAsync(storageClient, file, $"{logFile.RefName}", cancellationToken);
 		}
 
 		/// <summary>
diff --git a/Engine/Source/Programs/Horde/Horde.Server/Storage/StorageController.cs b/Engine/Source/Programs/Horde/Horde.Server/Storage/StorageController.cs
--- a/Engine/Source/Programs/Horde/Horde.Server/Storage/StorageController.cs
+++ b/Engine/Source/Programs/Horde/Horde.Server/Storage/StorageController.cs
@@ -198,59 +198,42 @@ namespace Horde.Server.Storage
 				return Forbid(StorageAclAction.WriteBlobs, namespaceId);
 			}
 
-			IStorageClientImpl client = await _storageService.GetClientAsync(namespaceId, cancellationToken);
-			if (file == null)
-			{
-				(BlobLocator Locator, Uri UploadUrl)? result = await client.GetWriteRedirectAsync(prefix ?? String.Empty, cancellationToken);
-				if (result == null)
-				{
-					return new WriteBlobResponse { SupportsRedirects = false };
-				}
-				else
-				{
-					return new WriteBlobResponse { Blob = result.Value.Locator, UploadUrl = result.Value.UploadUrl };
-				}
-			}
-			else
-			{
-				using (Stream stream = file.OpenReadStream())
-				{
-					BlobLocator locator = await client.WriteBlobAsync(stream, prefix: (prefix == null) ? Utf8String.Empty : new Utf8String(prefix), cancellationToken: cancellationToken);
-					return new WriteBlobResponse { Blob = locator, SupportsRedirects = client.SupportsRedirects ? (bool?)true : null };
-				}
-			}
+			IStorageClientImpl storageClient = await _storageService.GetClientAsync(namespaceId, cancellationToken);
+			return await WriteBlobAsync(storageClient, file, prefix, cancellationToken);
 		}
 
 		/// <summary>
 		/// Writes a blob to storage. Exposed as a public utility method to allow other routes with their own authentication methods to wrap their own authentication/redirection.
 		/// </summary>
-		/// <param name="storageService">The storage service</param>
-		/// <param name="namespaceId">Namespace to write the blob to</param>
+		/// <param name="storageClient">The client to write to service</param>
 		/// <param name="file">File to be written</param>
 		/// <param name="prefix">Prefix for uploaded blobs</param>
 		/// <param name="cancellationToken">Cancellation token</param>
 		/// <returns>Information about the written blob, or redirect information</returns>
-		public static async Task<ActionResult<WriteBlobResponse>> WriteBlobAsync(StorageService storageService, NamespaceId namespaceId, IFormFile? file, [FromForm] string? prefix = default, CancellationToken cancellationToken = default)
+		public static async Task<ActionResult<WriteBlobResponse>> WriteBlobAsync(IStorageClient storageClient, IFormFile? file, [FromForm] string? prefix = default, CancellationToken cancellationToken = default)
 		{
-			IStorageClientImpl client = await storageService.GetClientAsync(namespaceId, cancellationToken);
+			IStorageClientImpl? storageClientImpl = storageClient as IStorageClientImpl;
 			if (file == null)
 			{
-				(BlobLocator Locator, Uri UploadUrl)? result = await client.GetWriteRedirectAsync(prefix ?? String.Empty, cancellationToken);
-				if (result == null)
+				if (storageClientImpl == null)
 				{
 					return new WriteBlobResponse { SupportsRedirects = false };
 				}
-				else
+
+				(BlobLocator Locator, Uri UploadUrl)? result = await storageClientImpl.GetWriteRedirectAsync(prefix ?? String.Empty, cancellationToken);
+				if (result == null)
 				{
-					return new WriteBlobResponse { Blob = result.Value.Locator, UploadUrl = result.Value.UploadUrl };
+					return new WriteBlobResponse { SupportsRedirects = false };
 				}
+
+				return new WriteBlobResponse { Blob = result.Value.Locator, UploadUrl = result.Value.UploadUrl };
 			}
 			else
 			{
 				using (Stream stream = file.OpenReadStream())
 				{
-					BlobLocator locator = await client.WriteBlobAsync(stream, prefix: (prefix == null) ? Utf8String.Empty : new Utf8String(prefix), cancellationToken: cancellationToken);
-					return new WriteBlobResponse { Blob = locator, SupportsRedirects = client.SupportsRedirects ? (bool?)true : null };
+					BlobLocator locator = await storageClient.WriteBlobAsync(stream, prefix: (prefix == null) ? Utf8String.Empty : new Utf8String(prefix), cancellationToken: cancellationToken);
+					return new WriteBlobResponse { Blob = locator, SupportsRedirects = storageClientImpl?.SupportsRedirects };
 				}
 			}
 		}
diff --git a/Engine/Source/Programs/Horde/Horde.Server/Tools/IToolCollection.cs b/Engine/Source/Programs/Horde/Horde.Server/Tools/IToolCollection.cs
--- a/Engine/Source/Programs/Horde/Horde.Server/Tools/IToolCollection.cs
+++ b/Engine/Source/Programs/Horde/Horde.Server/Tools/IToolCollection.cs
@@ -5,6 +5,7 @@ using System.Threading;
 using System.Threading.Tasks;
 using EpicGames.Horde.Storage;
 using Horde.Server.Server;
+using Horde.Server.Storage;
 
 namespace Horde.Server.Tools
 {
@@ -32,6 +33,17 @@ namespace Horde.Server.Tools
 		/// <returns>Updated tool document, or null if it does not exist</returns>
 		Task<ITool?> CreateDeploymentAsync(ITool tool, ToolDeploymentConfig options, Stream stream, GlobalConfig globalConfig, CancellationToken cancellationToken);
 
+		/// <summary>
+		/// Adds a new deployment to the given tool. The new deployment will replace the current active deployment.
+		/// </summary>
+		/// <param name="tool">The tool to update</param>
+		/// <param name="options">Options for the new deployment</param>
+		/// <param name="handle">Handle to the root node containing the tool data</param>
+		/// <param name="globalConfig">The current configuration</param>
+		/// <param name="cancellationToken">Cancellation token for the operation</param>
+		/// <returns>Updated tool document, or null if it does not exist</returns>
+		Task<ITool?> CreateDeploymentAsync(ITool tool, ToolDeploymentConfig options, NodeHandle handle, GlobalConfig globalConfig, CancellationToken cancellationToken);
+
 		/// <summary>
 		/// Updates the state of the current deployment
 		/// </summary>
diff --git a/Engine/Source/Programs/Horde/Horde.Server/Tools/ToolCollection.cs b/Engine/Source/Programs/Horde/Horde.Server/Tools/ToolCollection.cs
--- a/Engine/Source/Programs/Horde/Horde.Server/Tools/ToolCollection.cs
+++ b/Engine/Source/Programs/Horde/Horde.Server/Tools/ToolCollection.cs
@@ -247,15 +247,9 @@ namespace Horde.Server.Tools
 		/// <returns>Updated tool document, or null if it does not exist</returns>
 		public async Task<ITool?> CreateDeploymentAsync(ITool tool, ToolDeploymentConfig options, Stream stream, GlobalConfig globalConfig, CancellationToken cancellationToken)
 		{
-			if (tool.Config is BundledToolConfig)
-			{
-				throw new InvalidOperationException("Cannot update the state of bundled tools.");
-			}
-
 			ToolDeploymentId deploymentId = ToolDeploymentId.GenerateNewId();
 			RefName refName = new RefName($"{tool.Id}/{deploymentId}");
 
-			// Upload the tool data first
 			IStorageClient client = await _storageService.GetClientAsync(Namespace.Tools, cancellationToken);
 
 			NodeHandle handle;
@@ -266,6 +260,36 @@ namespace Horde.Server.Tools
 				handle = await writer.WriteAsync(refName, directoryNode, cancellationToken: cancellationToken);
 			}
 
+			return await CreateDeploymentAsync(tool, options, handle, globalConfig, cancellationToken);
+		}
+
+		/// <summary>
+		/// Adds a new deployment to the given tool. The new deployment will replace the current active deployment.
+		/// </summary>
+		/// <param name="tool">The tool to update</param>
+		/// <param name="options">Options for the new deployment</param>
+		/// <param name="handle">Handle for the tool data</param>
+		/// <param name="globalConfig">The current configuration</param>
+		/// <param name="cancellationToken">Cancellation token for the operation</param>
+		/// <returns>Updated tool document, or null if it does not exist</returns>
+		public async Task<ITool?> CreateDeploymentAsync(ITool tool, ToolDeploymentConfig options, NodeHandle handle, GlobalConfig globalConfig, CancellationToken cancellationToken)
+		{
+			ToolDeploymentId deploymentId = ToolDeploymentId.GenerateNewId();
+			RefName refName = new RefName($"{tool.Id}/{deploymentId}");
+
+			IStorageClient client = await _storageService.GetClientAsync(Namespace.Tools, cancellationToken);
+			await client.WriteRefTargetAsync(refName, handle, cancellationToken: cancellationToken);
+
+			return await CreateDeploymentInternalAsync(tool, deploymentId, options, refName, globalConfig, cancellationToken);
+		}
+
+		async Task<ITool?> CreateDeploymentInternalAsync(ITool tool, ToolDeploymentId deploymentId, ToolDeploymentConfig options, RefName refName, GlobalConfig globalConfig, CancellationToken cancellationToken)
+		{
+			if (tool.Config is BundledToolConfig)
+			{
+				throw new InvalidOperationException("Cannot update the state of bundled tools.");
+			}
+
 			// Create the new deployment object
 			ToolDeployment deployment = new ToolDeployment(deploymentId, options, refName);
 
diff --git a/Engine/Source/Programs/Horde/Horde.Server/Tools/ToolsController.cs b/Engine/Source/Programs/Horde/Horde.Server/Tools/ToolsController.cs
--- a/Engine/Source/Programs/Horde/Horde.Server/Tools/ToolsController.cs
+++ b/Engine/Source/Programs/Horde/Horde.Server/Tools/ToolsController.cs
@@ -142,6 +142,17 @@ namespace Horde.Server.Tools
 		}
 	}
 
+	/// <summary>
+	/// Request for creating a new deployment
+	/// </summary>
+	public class CreateDeploymentRequest : ToolDeploymentConfig
+	{
+		/// <summary>
+		/// Handle to the root node
+		/// </summary>
+		public string Node { get; } = null!;
+	}
+
 	/// <summary>
 	/// Response from creating a deployment
 	/// </summary>
@@ -205,6 +216,30 @@ namespace Horde.Server.Tools
 			_globalConfig = globalConfig;
 		}
 
+		/// <summary>
+		/// Uploads blob data for a new tool deployment.
+		/// </summary>
+		/// <param name="id">Identifier of the tool to upload</param>
+		/// <param name="file">Data to be uploaded. May be null, in which case the server may return a separate url.</param>
+		/// <param name="cancellationToken">Cancellation token for the operation</param>
+		[HttpPost]
+		[Route("/api/v1/tools/{id}/blobs")]
+		public async Task<ActionResult<WriteBlobResponse>> WriteBlobAsync(ToolId id, IFormFile? file, CancellationToken cancellationToken = default)
+		{
+			ITool? tool = await _toolCollection.GetAsync(id, _globalConfig.Value);
+			if (tool == null)
+			{
+				return NotFound(id);
+			}
+			if (!tool.Config.Authorize(ToolAclAction.UploadTool, User))
+			{
+				return Forbid(ToolAclAction.UploadTool, id);
+			}
+
+			IStorageClient storageClient = await _toolCollection.GetStorageClientAsync(tool, cancellationToken);
+			return await StorageController.WriteBlobAsync(storageClient, file, cancellationToken: cancellationToken);
+		}
+
 		/// <summary>
 		/// Create a new deployment of the given tool.
 		/// </summary>
@@ -236,6 +271,33 @@ namespace Horde.Server.Tools
 			return new CreateDeploymentResponse { Id = tool.Deployments[^1].Id };
 		}
 
+		/// <summary>
+		/// Create a new deployment of the given tool.
+		/// </summary>
+		[HttpPost]
+		[Route("/api/v2/tools/{id}/deployments")]
+		public async Task<ActionResult<CreateDeploymentResponse>> CreateDeploymentAsync(ToolId id, CreateDeploymentRequest request, CancellationToken cancellationToken)
+		{
+			ITool? tool = await _toolCollection.GetAsync(id, _globalConfig.Value);
+
+			if (tool == null)
+			{
+				return NotFound(id);
+			}
+			if (!tool.Config.Authorize(ToolAclAction.UploadTool, User))
+			{
+				return Forbid(ToolAclAction.UploadTool, id);
+			}
+
+			tool = await _toolCollection.CreateDeploymentAsync(tool, request, NodeHandle.Parse(request.Node), _globalConfig.Value, cancellationToken);
+			if (tool == null)
+			{
+				return NotFound(id);
+			}
+
+			return new CreateDeploymentResponse { Id = tool.Deployments[^1].Id };
+		}
+
 		/// <summary>
 		/// Updates the state of an active deployment.
 		/// </summary>
@@ -272,7 +334,6 @@ namespace Horde.Server.Tools
 	public class PublicToolsController : HordeControllerBase
 	{
 		readonly IToolCollection _toolCollection;
-		readonly StorageService _storageService;
 		readonly IOptionsSnapshot<GlobalConfig> _globalConfig;
 		readonly IClock _clock;
 		readonly IMemoryCache _cache;
@@ -281,10 +342,9 @@ namespace Horde.Server.Tools
 		/// <summary>
 		/// Constructor
 		/// </summary>
-		public PublicToolsController(IToolCollection toolCollection, StorageService storageService, IClock clock, IOptionsSnapshot<GlobalConfig> globalConfig, IMemoryCache cache, ILogger<ToolsController> logger)
+		public PublicToolsController(IToolCollection toolCollection, IClock clock, IOptionsSnapshot<GlobalConfig> globalConfig, IMemoryCache cache, ILogger<ToolsController> logger)
 		{
 			_toolCollection = toolCollection;
-			_storageService = storageService;
 			_clock = clock;
 			_globalConfig = globalConfig;
 			_cache = cache;
@@ -463,7 +523,7 @@ namespace Horde.Server.Tools
 		}
 
 		/// <summary>
-		/// Retrieves blobs for a particular artifact
+		/// Retrieves blobs for a particular tool
 		/// </summary>
 		/// <param name="id">Identifier of the tool to retrieve</param>
 		/// <param name="locator">The blob locator</param>
