Horde: Accept duration of tool deployment in minutes, rather than parsing a timestamp value.

diff --git a/Engine/Source/Programs/Horde/Horde.Server/Tools/ToolsController.cs b/Engine/Source/Programs/Horde/Horde.Server/Tools/ToolsController.cs
--- a/Engine/Source/Programs/Horde/Horde.Server/Tools/ToolsController.cs
+++ b/Engine/Source/Programs/Horde/Horde.Server/Tools/ToolsController.cs
@@ -145,8 +145,21 @@ namespace Horde.Server.Tools
 	/// <summary>
 	/// Request for creating a new deployment
 	/// </summary>
-	public class CreateDeploymentRequest : ToolDeploymentConfig
+	public class CreateDeploymentRequest
 	{
+		/// <inheritdoc cref="IToolDeployment.Version"/>
+		public string Version { get; set; } = "Unknown";
+
+		/// <summary>
+		/// Number of minutes over which to do the deployment
+		/// </summary>
+		public double? Duration { get; set; }
+
+		/// <summary>
+		/// Whether to create the deployment in a paused state
+		/// </summary>
+		public bool? CreatePaused { get; set; }
+
 		/// <summary>
 		/// Handle to the root node
 		/// </summary>
@@ -289,7 +302,9 @@ namespace Horde.Server.Tools
 				return Forbid(ToolAclAction.UploadTool, id);
 			}
 
-			tool = await _toolCollection.CreateDeploymentAsync(tool, request, NodeHandle.Parse(request.Node), _globalConfig.Value, cancellationToken);
+			ToolDeploymentConfig options = new ToolDeploymentConfig { Version = request.Version, Duration = TimeSpan.FromMinutes(request.Duration ?? 0.0), CreatePaused = request.CreatePaused ?? false };
+
+			tool = await _toolCollection.CreateDeploymentAsync(tool, options, NodeHandle.Parse(request.Node), _globalConfig.Value, cancellationToken);
 			if (tool == null)
 			{
 				return NotFound(id);
