Horde: Add missing setter for new tool deployments.

diff --git a/Engine/Source/Programs/Horde/Horde.Server/Tools/ToolsController.cs b/Engine/Source/Programs/Horde/Horde.Server/Tools/ToolsController.cs
--- a/Engine/Source/Programs/Horde/Horde.Server/Tools/ToolsController.cs
+++ b/Engine/Source/Programs/Horde/Horde.Server/Tools/ToolsController.cs
@@ -163,7 +163,7 @@ namespace Horde.Server.Tools
 		/// <summary>
 		/// Handle to the root node
 		/// </summary>
-		public string Node { get; } = null!;
+		public string Node { get; set; } = null!;
 	}
 
 	/// <summary>
