Horde: Fix connection leak if remote execution clients do not send a valid nonce.

diff --git a/Engine/Source/Programs/Horde/Horde.Agent/Services/ComputeListenerService.cs b/Engine/Source/Programs/Horde/Horde.Agent/Services/ComputeListenerService.cs
--- a/Engine/Source/Programs/Horde/Horde.Agent/Services/ComputeListenerService.cs
+++ b/Engine/Source/Programs/Horde/Horde.Agent/Services/ComputeListenerService.cs
@@ -40,7 +40,7 @@ namespace Horde.Agent.Services
 				return null;
 			}
 
-			public Task WaitForClientToBeTaken() => _takenTaskSource.Task;
+			public Task WaitForTransferAsync() => _takenTaskSource.Task;
 
 			public void Dispose()
 			{
@@ -167,18 +167,23 @@ namespace Horde.Agent.Services
 
 		async Task WaitForLeaseAsync(TcpClient tcpClient, CancellationToken cancellationToken)
 		{
-			WaitingClient? waitingClient = null;
+			WaitingClient? waitingClient = new WaitingClient(tcpClient);
 			try
 			{
 				Socket socket = tcpClient.Client;
 
 				// Read the nonce for the connection
 				byte[] nonceBuffer = new byte[ServerComputeClient.NonceLength];
-				await socket.ReceiveMessageAsync(nonceBuffer, SocketFlags.None, cancellationToken);
-				_logger.LogInformation("Incoming connection with nonce {Nonce}", StringUtils.FormatHexString(nonceBuffer));
+				if (!await socket.TryReceiveMessageAsync(nonceBuffer, SocketFlags.None, cancellationToken))
+				{
+					_logger.LogInformation("No nonce received from {Remote}; closing connection.", tcpClient.Client.RemoteEndPoint);
+					return;
+				}
 
-				// Register the socket
 				ByteString nonce = new ByteString(nonceBuffer);
+				_logger.LogInformation("Incoming connection with nonce {Nonce}", nonce);
+
+				// Register the socket
 				try
 				{
 					// Check if there's already a waiting connection, and add one if not
@@ -187,7 +192,6 @@ namespace Horde.Agent.Services
 					{
 						if (!_waitingLeases.TryGetValue(nonce, out waitingLease))
 						{
-							waitingClient = new WaitingClient(tcpClient);
 							_waitingClients.Add(nonce, waitingClient);
 						}
 					}
@@ -195,16 +199,16 @@ namespace Horde.Agent.Services
 					// Offer ownership of the socket to any interested compute lease, or dispose of the connection if it's already closed.
 					if (waitingLease != null)
 					{
-						if (!waitingLease.TrySetResult(tcpClient))
+						if (waitingLease.TrySetResult(tcpClient))
 						{
-							tcpClient.Dispose();
+							_ = waitingClient.TakeClient();
 						}
 						return;
 					}
 
 					// Wait 2m for it to be claimed
 					Task delayTask = Task.Delay(TimeSpan.FromMinutes(2.0), cancellationToken);
-					await Task.WhenAny(waitingClient!.WaitForClientToBeTaken(), delayTask);
+					await Task.WhenAny(waitingClient!.WaitForTransferAsync(), delayTask);
 				}
 				finally
 				{
