Horde: Fix dead code warning.

diff --git a/Engine/Source/Programs/Horde/Horde.Agent/Services/ComputeListenerService.cs b/Engine/Source/Programs/Horde/Horde.Agent/Services/ComputeListenerService.cs
--- a/Engine/Source/Programs/Horde/Horde.Agent/Services/ComputeListenerService.cs
+++ b/Engine/Source/Programs/Horde/Horde.Agent/Services/ComputeListenerService.cs
@@ -167,7 +167,7 @@ namespace Horde.Agent.Services
 
 		async Task WaitForLeaseAsync(TcpClient tcpClient, CancellationToken cancellationToken)
 		{
-			WaitingClient? waitingClient = new WaitingClient(tcpClient);
+			using WaitingClient waitingClient = new WaitingClient(tcpClient);
 			try
 			{
 				Socket socket = tcpClient.Client;
@@ -222,10 +222,6 @@ namespace Horde.Agent.Services
 			{
 				_logger.LogError(ex, "Exception while processing compute task: {Message}", ex.Message);
 			}
-			finally
-			{
-				waitingClient?.Dispose();
-			}
 		}
 	}
 }
