Horde: Fix rendering of AclScopeName.

diff --git a/Engine/Source/Programs/Horde/Horde.Server/Acls/AclScopeName.cs b/Engine/Source/Programs/Horde/Horde.Server/Acls/AclScopeName.cs
--- a/Engine/Source/Programs/Horde/Horde.Server/Acls/AclScopeName.cs
+++ b/Engine/Source/Programs/Horde/Horde.Server/Acls/AclScopeName.cs
@@ -40,6 +40,9 @@ namespace Horde.Server.Acls
 
 		/// <inheritdoc/>
 		public override int GetHashCode() => Text.GetHashCode(StringComparison.Ordinal);
+
+		/// <inheritdoc/>
+		public override string ToString() => Text;
 	}
 
 	/// <summary>
