Horde: Fix code inspection warnings

diff --git a/Engine/Source/Programs/Horde/Horde.Server/Agents/Pools/PoolAclAction.cs b/Engine/Source/Programs/Horde/Horde.Server/Agents/Pools/PoolAclAction.cs
--- a/Engine/Source/Programs/Horde/Horde.Server/Agents/Pools/PoolAclAction.cs
+++ b/Engine/Source/Programs/Horde/Horde.Server/Agents/Pools/PoolAclAction.cs
@@ -12,26 +12,26 @@ namespace Horde.Server.Agents.Pools
 		/// <summary>
 		/// Create a global pool of agents
 		/// </summary>
-		public static AclAction CreatePool = new AclAction("CreatePool");
+		public static readonly AclAction CreatePool = new AclAction("CreatePool");
 
 		/// <summary>
 		/// Modify an agent pool
 		/// </summary>
-		public static AclAction UpdatePool = new AclAction("UpdatePool");
+		public static readonly AclAction UpdatePool = new AclAction("UpdatePool");
 
 		/// <summary>
 		/// Delete an agent pool
 		/// </summary>
-		public static AclAction DeletePool = new AclAction("DeletePool");
+		public static readonly AclAction DeletePool = new AclAction("DeletePool");
 
 		/// <summary>
 		/// Ability to view a pool
 		/// </summary>
-		public static AclAction ViewPool = new AclAction("ViewPool");
+		public static readonly AclAction ViewPool = new AclAction("ViewPool");
 
 		/// <summary>
 		/// View all the available agent pools
 		/// </summary>
-		public static AclAction ListPools = new AclAction("ListPools");
+		public static readonly AclAction ListPools = new AclAction("ListPools");
 	}
 }
diff --git a/Engine/Source/Programs/Horde/Horde.Server/Agents/Sessions/SessionAclAction.cs b/Engine/Source/Programs/Horde/Horde.Server/Agents/Sessions/SessionAclAction.cs
--- a/Engine/Source/Programs/Horde/Horde.Server/Agents/Sessions/SessionAclAction.cs
+++ b/Engine/Source/Programs/Horde/Horde.Server/Agents/Sessions/SessionAclAction.cs
@@ -7,16 +7,16 @@ namespace Horde.Server.Agents.Sessions
 	/// <summary>
 	/// Actions valid for agent sessions
 	/// </summary>
-	public class SessionAclAction
+	public static class SessionAclAction
 	{
 		/// <summary>
 		/// Granted to agents to call CreateSession, which returns a bearer token identifying themselves valid to call UpdateSesssion via gRPC.
 		/// </summary>
-		public static AclAction CreateSession = new AclAction("CreateSession");
+		public static readonly AclAction CreateSession = new AclAction("CreateSession");
 
 		/// <summary>
 		/// Allows viewing information about an agent session
 		/// </summary>
-		public static AclAction ViewSession = new AclAction("ViewSession");
+		public static readonly AclAction ViewSession = new AclAction("ViewSession");
 	}
 }
diff --git a/Engine/Source/Programs/Horde/Horde.Server/Jobs/JobAclAction.cs b/Engine/Source/Programs/Horde/Horde.Server/Jobs/JobAclAction.cs
--- a/Engine/Source/Programs/Horde/Horde.Server/Jobs/JobAclAction.cs
+++ b/Engine/Source/Programs/Horde/Horde.Server/Jobs/JobAclAction.cs
@@ -12,31 +12,31 @@ namespace Horde.Server.Jobs
 		/// <summary>
 		/// Ability to start new jobs
 		/// </summary>
-		public static AclAction CreateJob = new AclAction("CreateJob");
+		public static readonly AclAction CreateJob = new AclAction("CreateJob");
 
 		/// <summary>
 		/// Rename a job, modify its priority, etc...
 		/// </summary>
-		public static AclAction UpdateJob = new AclAction("UpdateJob");
+		public static readonly AclAction UpdateJob = new AclAction("UpdateJob");
 
 		/// <summary>
 		/// Delete a job properties
 		/// </summary>
-		public static AclAction DeleteJob = new AclAction("DeleteJob");
+		public static readonly AclAction DeleteJob = new AclAction("DeleteJob");
 
 		/// <summary>
 		/// Allows updating a job metadata (name, changelist number, step properties, new groups, job states, etc...). Typically granted to agents. Not user facing.
 		/// </summary>
-		public static AclAction ExecuteJob = new AclAction("ExecuteJob");
+		public static readonly AclAction ExecuteJob = new AclAction("ExecuteJob");
 
 		/// <summary>
 		/// Ability to retry a failed job step
 		/// </summary>
-		public static AclAction RetryJobStep = new AclAction("RetryJobStep");
+		public static readonly AclAction RetryJobStep = new AclAction("RetryJobStep");
 
 		/// <summary>
 		/// Ability to view a job
 		/// </summary>
-		public static AclAction ViewJob = new AclAction("ViewJob");
+		public static readonly AclAction ViewJob = new AclAction("ViewJob");
 	}
 }
diff --git a/Engine/Source/Programs/Horde/Horde.Server/Jobs/JobCollection.cs b/Engine/Source/Programs/Horde/Horde.Server/Jobs/JobCollection.cs
--- a/Engine/Source/Programs/Horde/Horde.Server/Jobs/JobCollection.cs
+++ b/Engine/Source/Programs/Horde/Horde.Server/Jobs/JobCollection.cs
@@ -124,7 +124,7 @@ namespace Horde.Server.Jobs
 					{
 						description.Append($" - {Error}");
 					}
-					description.Append($")");
+					description.Append(')');
 				}
 				return description.ToString();
 			}
diff --git a/Engine/Source/Programs/Horde/Horde.Server/Notifications/NotificationAclAction.cs b/Engine/Source/Programs/Horde/Horde.Server/Notifications/NotificationAclAction.cs
--- a/Engine/Source/Programs/Horde/Horde.Server/Notifications/NotificationAclAction.cs
+++ b/Engine/Source/Programs/Horde/Horde.Server/Notifications/NotificationAclAction.cs
@@ -12,6 +12,6 @@ namespace Horde.Server.Notifications
 		/// <summary>
 		/// Ability to subscribe to notifications
 		/// </summary>
-		public static AclAction CreateSubscription = new AclAction("CreateSubscription");
+		public static readonly AclAction CreateSubscription = new AclAction("CreateSubscription");
 	}
 }
diff --git a/Engine/Source/Programs/Horde/Horde.Server/ServerSettings.cs b/Engine/Source/Programs/Horde/Horde.Server/ServerSettings.cs
--- a/Engine/Source/Programs/Horde/Horde.Server/ServerSettings.cs
+++ b/Engine/Source/Programs/Horde/Horde.Server/ServerSettings.cs
@@ -846,7 +846,7 @@ namespace Horde.Server
 		/// Create the default ACL for the server, including all predefined roles.
 		/// </summary>
 		/// <returns></returns>
-		AclConfig GetDefaultAcl()
+		static AclConfig GetDefaultAcl()
 		{
 			AclConfig defaultAcl = new AclConfig();
 			defaultAcl.Entries.Add(new AclEntryConfig(new AclClaimConfig(ClaimTypes.Role, "internal:AgentRegistration"), new[] { AgentAclAction.CreateAgent, SessionAclAction.CreateSession }));
diff --git a/Engine/Source/Programs/Horde/Horde.Server/Storage/StorageAclAction.cs b/Engine/Source/Programs/Horde/Horde.Server/Storage/StorageAclAction.cs
--- a/Engine/Source/Programs/Horde/Horde.Server/Storage/StorageAclAction.cs
+++ b/Engine/Source/Programs/Horde/Horde.Server/Storage/StorageAclAction.cs
@@ -12,26 +12,26 @@ namespace Horde.Server.Storage
 		/// <summary>
 		/// Ability to read blobs from the storage service
 		/// </summary>
-		public static AclAction ReadBlobs = new AclAction("ReadBlobs");
+		public static readonly AclAction ReadBlobs = new AclAction("ReadBlobs");
 
 		/// <summary>
 		/// Ability to write blobs to the storage service
 		/// </summary>
-		public static AclAction WriteBlobs = new AclAction("WriteBlobs");
+		public static readonly AclAction WriteBlobs = new AclAction("WriteBlobs");
 
 		/// <summary>
 		/// Ability to read refs from the storage service
 		/// </summary>
-		public static AclAction ReadRefs = new AclAction("ReadRefs");
+		public static readonly AclAction ReadRefs = new AclAction("ReadRefs");
 
 		/// <summary>
 		/// Ability to write refs to the storage service
 		/// </summary>
-		public static AclAction WriteRefs = new AclAction("WriteRefs");
+		public static readonly AclAction WriteRefs = new AclAction("WriteRefs");
 
 		/// <summary>
 		/// Ability to delete refs
 		/// </summary>
-		public static AclAction DeleteRefs = new AclAction("DeleteRefs");
+		public static readonly AclAction DeleteRefs = new AclAction("DeleteRefs");
 	}
 }
