Horde: Configure materializer per job when using WorkspaceExecutor

diff --git a/Engine/Source/Programs/Horde/Horde.Agent/Execution/WorkspaceExecutor.cs b/Engine/Source/Programs/Horde/Horde.Agent/Execution/WorkspaceExecutor.cs
--- a/Engine/Source/Programs/Horde/Horde.Agent/Execution/WorkspaceExecutor.cs
+++ b/Engine/Source/Programs/Horde/Horde.Agent/Execution/WorkspaceExecutor.cs
@@ -129,31 +129,44 @@ namespace Horde.Agent.Execution
 
 	class WorkspaceExecutorFactory : IJobExecutorFactory
 	{
+		private readonly IWorkspaceMaterializerFactory _materializerFactory;
 		private readonly ILoggerFactory _loggerFactory;
 		
 		public string Name => WorkspaceExecutor.Name;
 
-		public WorkspaceExecutorFactory(ILoggerFactory loggerFactory)
+		public WorkspaceExecutorFactory(IWorkspaceMaterializerFactory materializerFactory, ILoggerFactory loggerFactory)
 		{
+			_materializerFactory = materializerFactory;
 			_loggerFactory = loggerFactory;
 		}
 
 		public IJobExecutor CreateExecutor(AgentWorkspace workspaceInfo, AgentWorkspace? autoSdkWorkspaceInfo, JobExecutorOptions options)
 		{
-			IWorkspaceMaterializer workspaceMaterializer;
-			IWorkspaceMaterializer? autoSdkMaterializer = null;
-
-			// Default to ManagedWorkspaceMaterializer for now. Different types of materializer can later be set via AgentWorkspace or JobOptions.
-			workspaceMaterializer = new ManagedWorkspaceMaterializer(
-				workspaceInfo, options.Session.WorkingDir, false, false, _loggerFactory.CreateLogger<ManagedWorkspaceMaterializer>());
+			WorkspaceMaterializerType type = GetMaterializerType(options.JobOptions.WorkspaceMaterializer, WorkspaceMaterializerType.ManagedWorkspace);
+			IWorkspaceMaterializer workspaceMaterializer = _materializerFactory.CreateMaterializer(type, workspaceInfo, options);
 			
+			IWorkspaceMaterializer? autoSdkMaterializer = null;
 			if (autoSdkWorkspaceInfo != null)
 			{
-				autoSdkMaterializer = new ManagedWorkspaceMaterializer(
-					autoSdkWorkspaceInfo, options.Session.WorkingDir, true, true, _loggerFactory.CreateLogger<ManagedWorkspaceMaterializer>());
+				autoSdkMaterializer = _materializerFactory.CreateMaterializer(type, workspaceInfo, options, forAutoSdk: true);
 			}
 			
 			return new WorkspaceExecutor(options, workspaceMaterializer, autoSdkMaterializer, _loggerFactory.CreateLogger<WorkspaceExecutor>());
 		}
+
+		private static WorkspaceMaterializerType GetMaterializerType(string name, WorkspaceMaterializerType defaultValue)
+		{
+			if (String.IsNullOrEmpty(name))
+			{
+				return defaultValue;
+			}
+			
+			if (Enum.TryParse(name, true, out WorkspaceMaterializerType enumType))
+			{
+				return enumType;
+			}
+			
+			throw new ArgumentException($"Unable to find materializer type '{name}'");
+		}
 	}
 }
diff --git a/Engine/Source/Programs/Horde/Horde.Agent/Execution/WorkspaceMaterializer.cs b/Engine/Source/Programs/Horde/Horde.Agent/Execution/WorkspaceMaterializer.cs
--- a/Engine/Source/Programs/Horde/Horde.Agent/Execution/WorkspaceMaterializer.cs
+++ b/Engine/Source/Programs/Horde/Horde.Agent/Execution/WorkspaceMaterializer.cs
@@ -4,6 +4,8 @@ using System;
 using System.Threading;
 using System.Threading.Tasks;
 using EpicGames.Core;
+using HordeCommon.Rpc.Messages;
+using Microsoft.Extensions.Logging;
 
 namespace Horde.Agent.Execution;
 
@@ -126,4 +128,50 @@ public interface IWorkspaceMaterializer
 	/// <exception cref="Horde.Agent.Execution.WorkspaceMaterializationException">Thrown if unshelving fails</exception>
 	/// <returns>Async task</returns>
 	public Task UnshelveAsync(int changeNum, CancellationToken cancellationToken);
-}
\ No newline at end of file
+}
+
+enum WorkspaceMaterializerType
+{
+	ManagedWorkspace,
+}
+
+/// <summary>
+/// Factory for creating new workspace materializers
+/// </summary>
+interface IWorkspaceMaterializerFactory
+{
+	/// <summary>
+	/// Creates a new workspace materializer instance
+	/// </summary>
+	/// <param name="type">Type of materializer to instantiate</param>
+	/// <param name="workspaceInfo">Agent workspace</param>
+	/// <param name="options">Job options</param>
+	/// <param name="forAutoSdk">Whether intended for AutoSDK materialization</param>
+	/// <returns>A new workspace materializer instance</returns>
+	IWorkspaceMaterializer CreateMaterializer(WorkspaceMaterializerType type, AgentWorkspace workspaceInfo, JobExecutorOptions options, bool forAutoSdk = false);
+}
+
+class WorkspaceMaterializerFactory : IWorkspaceMaterializerFactory
+{
+	private readonly ILoggerFactory _loggerFactory;
+	
+	public WorkspaceMaterializerFactory(ILoggerFactory loggerFactory)
+	{
+		_loggerFactory = loggerFactory;
+	}
+	
+	/// <inheritdoc/>
+	public IWorkspaceMaterializer CreateMaterializer(WorkspaceMaterializerType type, AgentWorkspace workspaceInfo, JobExecutorOptions options, bool forAutoSdk)
+	{
+		switch (type)
+		{
+			case WorkspaceMaterializerType.ManagedWorkspace:
+				return forAutoSdk
+					? new ManagedWorkspaceMaterializer(workspaceInfo, options.Session.WorkingDir, true, true, _loggerFactory.CreateLogger<ManagedWorkspaceMaterializer>())
+					: new ManagedWorkspaceMaterializer(workspaceInfo, options.Session.WorkingDir, false, false, _loggerFactory.CreateLogger<ManagedWorkspaceMaterializer>());
+
+			default:
+				throw new Exception("Unhandled materializer option: " + type);
+		};
+	}
+}
diff --git a/Engine/Source/Programs/Horde/Horde.Agent/Program.cs b/Engine/Source/Programs/Horde/Horde.Agent/Program.cs
--- a/Engine/Source/Programs/Horde/Horde.Agent/Program.cs
+++ b/Engine/Source/Programs/Horde/Horde.Agent/Program.cs
@@ -194,6 +194,8 @@ namespace Horde.Agent
 			services.AddSingleton<IJobExecutorFactory, WorkspaceExecutorFactory>();
 			services.AddSingleton<IJobExecutorFactory, LocalExecutorFactory>();
 			services.AddSingleton<IJobExecutorFactory, TestExecutorFactory>();
+			
+			services.AddSingleton<IWorkspaceMaterializerFactory, WorkspaceMaterializerFactory>();
 
 			services.AddSingleton<JobHandler>();
 			services.AddSingleton<StatusService>();
