Horde: Record lease, job and batch ID in telemetry events

diff --git a/Engine/Source/Programs/Horde/Horde.Agent.Tests/Services/TelemetryServiceTest.cs b/Engine/Source/Programs/Horde/Horde.Agent.Tests/Services/TelemetryServiceTest.cs
--- a/Engine/Source/Programs/Horde/Horde.Agent.Tests/Services/TelemetryServiceTest.cs
+++ b/Engine/Source/Programs/Horde/Horde.Agent.Tests/Services/TelemetryServiceTest.cs
@@ -4,6 +4,8 @@ using System;
 using System.Collections.Generic;
 using System.Threading;
 using System.Threading.Tasks;
+using Horde.Agent.Execution;
+using Horde.Agent.Leases.Handlers;
 using Horde.Agent.Services;
 using Microsoft.Extensions.Logging.Abstractions;
 using Microsoft.Extensions.Options;
@@ -21,7 +23,9 @@ public sealed class TelemetryServiceTest : IDisposable
 	{
 		AgentSettings settings = new() { Server = "Test", ServerProfiles = { new ServerProfile() { Name = "Test", Url = new Uri("http://localhost:1234") } }};
 		GrpcService grpcService = new (new OptionsWrapper<AgentSettings>(settings), NullLogger<GrpcService>.Instance, _loggerFactory);
-		_telemetryService = new TelemetryService(grpcService, new OptionsWrapper<AgentSettings>(settings), NullLogger<TelemetryService>.Instance);
+		JobHandler jobHandler = new (new List<IJobExecutorFactory>(), new OptionsWrapper<AgentSettings>(settings),
+			null!, null!, NullLogger<JobHandler>.Instance);
+		_telemetryService = new TelemetryService(jobHandler, grpcService, new OptionsWrapper<AgentSettings>(settings), NullLogger<TelemetryService>.Instance);
 	}
 
 	[TestMethod]
diff --git a/Engine/Source/Programs/Horde/Horde.Agent/Leases/Handlers/JobHandler.cs b/Engine/Source/Programs/Horde/Horde.Agent/Leases/Handlers/JobHandler.cs
--- a/Engine/Source/Programs/Horde/Horde.Agent/Leases/Handlers/JobHandler.cs
+++ b/Engine/Source/Programs/Horde/Horde.Agent/Leases/Handlers/JobHandler.cs
@@ -39,6 +39,21 @@ namespace Horde.Agent.Leases.Handlers
 		/// </summary>
 		internal TimeSpan _stepAbortPollInterval = TimeSpan.FromSeconds(5);
 
+		/// <summary>
+		/// Current lease ID being executed
+		/// </summary>
+		public string? CurrentLeaseId { get; private set; } = null;
+		
+		/// <summary>
+		/// Current job ID being executed
+		/// </summary>
+		public string? CurrentJobId { get; private set; } = null;
+		
+		/// <summary>
+		/// Current job batch ID being executed
+		/// </summary>
+		public string? CurrentBatchId { get; private set; } = null;
+
 		readonly IEnumerable<IJobExecutorFactory> _executorFactories;
 		readonly AgentSettings _settings;
 		readonly IServerStorageFactory _serverStorageFactory;
@@ -60,52 +75,70 @@ namespace Horde.Agent.Leases.Handlers
 		/// <inheritdoc/>
 		public override async Task<LeaseResult> ExecuteAsync(ISession session, string leaseId, ExecuteJobTask executeTask, CancellationToken cancellationToken)
 		{
-			executeTask.JobOptions ??= new JobOptions();
-
-			if (executeTask.JobOptions.RunInSeparateProcess ?? false)
+			try
 			{
-				using (ManagedProcessGroup processGroup = new ManagedProcessGroup())
+				CurrentLeaseId = leaseId;
+				CurrentJobId = executeTask.JobId;
+				CurrentBatchId = executeTask.BatchId;
+				
+				executeTask.JobOptions ??= new JobOptions();
+
+				if (executeTask.JobOptions.RunInSeparateProcess ?? false)
 				{
-					List<string> arguments = new List<string>();
-					arguments.Add(Assembly.GetExecutingAssembly().Location);
-					arguments.Add("execute");
-					arguments.Add("job");
-					arguments.Add($"-Server={_settings.GetCurrentServerProfile().Name}");
-					arguments.Add($"-AgentId={session.AgentId}");
-					arguments.Add($"-SessionId={session.SessionId}");
-					arguments.Add($"-LeaseId={leaseId}");
-					arguments.Add($"-WorkingDir={session.WorkingDir}");
-					arguments.Add($"-Task={Convert.ToBase64String(executeTask.ToByteArray())}");
-
-					string commandLine = CommandLineArguments.Join(arguments);
-					_defaultLogger.LogInformation("Running child process with arguments: {CommandLine}", commandLine);
-
-					using (ManagedProcess process = new ManagedProcess(processGroup, "dotnet", commandLine, null, null, ProcessPriorityClass.Normal))
+					using (ManagedProcessGroup processGroup = new ManagedProcessGroup())
 					{
-						using (LogEventParser parser = new LogEventParser(_defaultLogger))
+						List<string> arguments = new List<string>();
+						arguments.Add(Assembly.GetExecutingAssembly().Location);
+						arguments.Add("execute");
+						arguments.Add("job");
+						arguments.Add($"-Server={_settings.GetCurrentServerProfile().Name}");
+						arguments.Add($"-AgentId={session.AgentId}");
+						arguments.Add($"-SessionId={session.SessionId}");
+						arguments.Add($"-LeaseId={leaseId}");
+						arguments.Add($"-WorkingDir={session.WorkingDir}");
+						arguments.Add($"-Task={Convert.ToBase64String(executeTask.ToByteArray())}");
+
+						string commandLine = CommandLineArguments.Join(arguments);
+						_defaultLogger.LogInformation("Running child process with arguments: {CommandLine}",
+							commandLine);
+
+						using (ManagedProcess process = new ManagedProcess(processGroup, "dotnet", commandLine, null,
+							       null, ProcessPriorityClass.Normal))
 						{
-							for (; ; )
+							using (LogEventParser parser = new LogEventParser(_defaultLogger))
 							{
-								string? line = await process.ReadLineAsync(cancellationToken);
-								if (line == null)
+								for (;;)
 								{
-									break;
+									string? line = await process.ReadLineAsync(cancellationToken);
+									if (line == null)
+									{
+										break;
+									}
+
+									parser.WriteLine(line);
 								}
-								parser.WriteLine(line);
 							}
-						}
-						process.WaitForExit();
 
-						if (process.ExitCode != 0)
-						{
-							return LeaseResult.Failed;
+							process.WaitForExit();
+
+							if (process.ExitCode != 0)
+							{
+								return LeaseResult.Failed;
+							}
 						}
 					}
+
+					return LeaseResult.Success;
 				}
-				return LeaseResult.Success;
-			}
 
-			return await ExecuteInternalAsync(session, leaseId, executeTask, cancellationToken);
+				return await ExecuteInternalAsync(session, leaseId, executeTask, cancellationToken);
+			}
+			finally
+			{
+				CurrentLeaseId = null;
+				CurrentJobId = null;
+				CurrentBatchId = null;
+			}
 		}
 
 		internal async Task<LeaseResult> ExecuteInternalAsync(ISession session, string leaseId, ExecuteJobTask executeTask, CancellationToken cancellationToken)
diff --git a/Engine/Source/Programs/Horde/Horde.Agent/Services/TelemetryService.cs b/Engine/Source/Programs/Horde/Horde.Agent/Services/TelemetryService.cs
--- a/Engine/Source/Programs/Horde/Horde.Agent/Services/TelemetryService.cs
+++ b/Engine/Source/Programs/Horde/Horde.Agent/Services/TelemetryService.cs
@@ -14,6 +14,7 @@ using EpicGames.Core;
 using Google.Protobuf.WellKnownTypes;
 using Grpc.Core;
 using Grpc.Net.Client;
+using Horde.Agent.Leases.Handlers;
 using HordeCommon.Rpc;
 using HordeCommon.Rpc.Messages.Telemetry;
 using Microsoft.Extensions.Logging;
@@ -211,6 +212,7 @@ class TelemetryService : BackgroundService
 	private readonly TimeSpan _heartbeatInterval = TimeSpan.FromSeconds(60);
 	private readonly TimeSpan _heartbeatMaxAllowedDiff = TimeSpan.FromSeconds(5);
 	
+	private readonly JobHandler _jobHandler;
 	private readonly GrpcService _grpcService;
 	private readonly AgentSettings _agentSettings;
 	private readonly ILogger<TelemetryService> _logger;
@@ -226,8 +228,9 @@ class TelemetryService : BackgroundService
 	/// <summary>
 	/// Constructor
 	/// </summary>
-	public TelemetryService(GrpcService grpcService, IOptions<AgentSettings> settings, ILogger<TelemetryService> logger)
+	public TelemetryService(JobHandler jobHandler, GrpcService grpcService, IOptions<AgentSettings> settings, ILogger<TelemetryService> logger)
 	{
+		_jobHandler = jobHandler;
 		_grpcService = grpcService;
 		_agentSettings = settings.Value;
 		_logger = logger;
@@ -478,7 +481,12 @@ class TelemetryService : BackgroundService
 			
 			SendTelemetryEventsRequest request = new();
 			Timestamp utcNow = Timestamp.FromDateTime(DateTime.UtcNow);
-			ExecutionMetadata em = new ();
+			ExecutionMetadata em = new()
+			{
+				LeaseId = _jobHandler.CurrentLeaseId,
+				JobId = _jobHandler.CurrentJobId,
+				JobBatchId = _jobHandler.CurrentBatchId,
+			};
 			
 			{
 				AgentCpuMetricsEvent cpuMetricsEvent = _systemMetrics.GetCpu().ToEvent();
