Horde: Record pool IDs the agent is a member of in telemetry events

diff --git a/Engine/Source/Programs/Horde/Horde.Agent.Tests/Services/TelemetryServiceTest.cs b/Engine/Source/Programs/Horde/Horde.Agent.Tests/Services/TelemetryServiceTest.cs
--- a/Engine/Source/Programs/Horde/Horde.Agent.Tests/Services/TelemetryServiceTest.cs
+++ b/Engine/Source/Programs/Horde/Horde.Agent.Tests/Services/TelemetryServiceTest.cs
@@ -5,6 +5,7 @@ using System.Collections.Generic;
 using System.Threading;
 using System.Threading.Tasks;
 using Horde.Agent.Execution;
+using Horde.Agent.Leases;
 using Horde.Agent.Leases.Handlers;
 using Horde.Agent.Services;
 using Microsoft.Extensions.Logging.Abstractions;
@@ -22,10 +23,13 @@ public sealed class TelemetryServiceTest : IDisposable
 	public TelemetryServiceTest()
 	{
 		AgentSettings settings = new() { Server = "Test", ServerProfiles = { new ServerProfile() { Name = "Test", Url = new Uri("http://localhost:1234") } }};
-		GrpcService grpcService = new (new OptionsWrapper<AgentSettings>(settings), NullLogger<GrpcService>.Instance, _loggerFactory);
-		JobHandler jobHandler = new (new List<IJobExecutorFactory>(), new OptionsWrapper<AgentSettings>(settings),
-			null!, null!, NullLogger<JobHandler>.Instance);
-		_telemetryService = new TelemetryService(jobHandler, grpcService, new OptionsWrapper<AgentSettings>(settings), NullLogger<TelemetryService>.Instance);
+		OptionsWrapper<AgentSettings> settingsOpt = new (settings);
+		
+		using WorkerService workerService = new (settingsOpt, null!, null!, null!, new List<LeaseHandler>(), null!, null!);
+		JobHandler jobHandler = new (new List<IJobExecutorFactory>(), settingsOpt, null!, null!, NullLogger<JobHandler>.Instance);
+		GrpcService grpcService = new (settingsOpt, NullLogger<GrpcService>.Instance, _loggerFactory);
+		
+		_telemetryService = new TelemetryService(workerService, jobHandler, grpcService, settingsOpt, NullLogger<TelemetryService>.Instance);
 	}
 
 	[TestMethod]
diff --git a/Engine/Source/Programs/Horde/Horde.Agent/Leases/LeaseManager.cs b/Engine/Source/Programs/Horde/Horde.Agent/Leases/LeaseManager.cs
--- a/Engine/Source/Programs/Horde/Horde.Agent/Leases/LeaseManager.cs
+++ b/Engine/Source/Programs/Horde/Horde.Agent/Leases/LeaseManager.cs
@@ -110,6 +110,11 @@ namespace Horde.Agent.Leases
 		/// </summary>
 		public event LeaseActiveEvent? OnLeaseActive;
 
+		/// <summary>
+		/// List of pool IDs this session is a member of
+		/// </summary>
+		public IReadOnlyList<string> PoolIds { get; private set; } = new List<string>();
+
 		/// <summary>
 		/// How long to wait before trying to reacquire a new connection
 		/// Exposed as internal to ease testing. Using a lower delay can speed up tests. 
@@ -331,6 +336,8 @@ namespace Horde.Agent.Leases
 								_sessionResult = new SessionResult(SessionOutcome.Terminate);
 							}
 
+							PoolIds = updateSessionResponse.PoolIds;
+
 							// Remove any leases which have completed
 							int numRemoved = _activeLeases.RemoveAll(x => (x.Lease.State == LeaseState.Completed || x.Lease.State == LeaseState.Cancelled) && !updateSessionResponse.Leases.Any(y => y.Id == x.Lease.Id && y.State != LeaseState.Cancelled));
 							NumLeasesCompleted += numRemoved;
diff --git a/Engine/Source/Programs/Horde/Horde.Agent/Services/TelemetryService.cs b/Engine/Source/Programs/Horde/Horde.Agent/Services/TelemetryService.cs
--- a/Engine/Source/Programs/Horde/Horde.Agent/Services/TelemetryService.cs
+++ b/Engine/Source/Programs/Horde/Horde.Agent/Services/TelemetryService.cs
@@ -212,13 +212,14 @@ class TelemetryService : BackgroundService
 	private readonly TimeSpan _heartbeatInterval = TimeSpan.FromSeconds(60);
 	private readonly TimeSpan _heartbeatMaxAllowedDiff = TimeSpan.FromSeconds(5);
 	
+	private readonly WorkerService _workerService;
 	private readonly JobHandler _jobHandler;
 	private readonly GrpcService _grpcService;
 	private readonly AgentSettings _agentSettings;
 	private readonly ILogger<TelemetryService> _logger;
 	private readonly TimeSpan _reportInterval;
 	private readonly AgentMetadataEvent _agentMetadataEvent;
-	private readonly TimeSpan _agentMetadataReportInterval = TimeSpan.FromMinutes(10);
+	private readonly TimeSpan _agentMetadataReportInterval = TimeSpan.FromMinutes(2);
 	private ISystemMetrics? _systemMetrics;
 	private DateTime _lastTimeAgentMetadataSent = DateTime.UnixEpoch;
 
@@ -228,8 +229,9 @@ class TelemetryService : BackgroundService
 	/// <summary>
 	/// Constructor
 	/// </summary>
-	public TelemetryService(JobHandler jobHandler, GrpcService grpcService, IOptions<AgentSettings> settings, ILogger<TelemetryService> logger)
+	public TelemetryService(WorkerService workerService, JobHandler jobHandler, GrpcService grpcService, IOptions<AgentSettings> settings, ILogger<TelemetryService> logger)
 	{
+		_workerService = workerService;
 		_jobHandler = jobHandler;
 		_grpcService = grpcService;
 		_agentSettings = settings.Value;
@@ -528,11 +530,12 @@ class TelemetryService : BackgroundService
 			Region = null,
 			AvailabilityZone = null,
 			Environment = null,
-			AgentVersion = null,
+			AgentVersion = Program.Version,
 			Os = GetOs(),
 			OsVersion = Environment.OSVersion.Version.ToString(),
-			Architecture = RuntimeInformation.OSArchitecture.ToString()
+			Architecture = RuntimeInformation.OSArchitecture.ToString(),
 		};
+		e.PoolIds.AddRange(_workerService.PoolIds);
 		e.AgentId = e.CalculateAgentId();
 		return e;
 	}
diff --git a/Engine/Source/Programs/Horde/Horde.Agent/Services/WorkerService.cs b/Engine/Source/Programs/Horde/Horde.Agent/Services/WorkerService.cs
--- a/Engine/Source/Programs/Horde/Horde.Agent/Services/WorkerService.cs
+++ b/Engine/Source/Programs/Horde/Horde.Agent/Services/WorkerService.cs
@@ -27,6 +27,9 @@ namespace Horde.Agent.Services
 		readonly StatusService _statusService;
 		readonly LeaseHandler[] _leaseHandlers;
 		readonly IServiceProvider _serviceProvider;
+		LeaseManager? _currentLeaseManager;
+
+		public IReadOnlyList<string> PoolIds => _currentLeaseManager?.PoolIds ?? Array.Empty<string>();
 
 		static readonly TimeSpan[] s_sessionBackOffTime =
 		{
@@ -106,8 +109,8 @@ namespace Horde.Agent.Services
 
 						await using (ISession session = await _sessionFactory.CreateAsync(stoppingToken))
 						{
-							LeaseManager leaseManager = new LeaseManager(session, _capabilitiesService, _statusService, _leaseHandlers, _settings, _logger);
-							result = await leaseManager.RunAsync(false, stoppingToken);
+							_currentLeaseManager = new LeaseManager(session, _capabilitiesService, _statusService, _leaseHandlers, _settings, _logger);
+							result = await _currentLeaseManager.RunAsync(false, stoppingToken);
 						}
 
 						failureCount = 0;
diff --git a/Engine/Source/Programs/Horde/Horde.Agent/Utilities/HordeRpc.cs b/Engine/Source/Programs/Horde/Horde.Agent/Utilities/HordeRpc.cs
--- a/Engine/Source/Programs/Horde/Horde.Agent/Utilities/HordeRpc.cs
+++ b/Engine/Source/Programs/Horde/Horde.Agent/Utilities/HordeRpc.cs
@@ -213,12 +213,17 @@ namespace HordeCommon.Rpc.Messages.Telemetry
 			bw.Write(OsVersion ?? "<empty os version>");
 			bw.Write(Architecture ?? "<empty os architecture>");
 			
-			foreach (KeyValuePair<string,string> pair in Properties)
+			foreach (KeyValuePair<string, string> pair in Properties)
 			{
 				bw.Write(pair.Key ?? "<empty key>");
 				bw.Write(pair.Value ?? "<empty value>");
 			}
 			
+			foreach (string poolId in PoolIds)
+			{
+				bw.Write(poolId);
+			}
+			
 			ms.Position = 0;
 			byte[] hash = sha256.ComputeHash(ms);
 			return BitConverter.ToUInt64(hash, 0);
