Horde: Expose WorkerService as a singleton in DI

diff --git a/Engine/Source/Programs/Horde/Horde.Agent/Program.cs b/Engine/Source/Programs/Horde/Horde.Agent/Program.cs
--- a/Engine/Source/Programs/Horde/Horde.Agent/Program.cs
+++ b/Engine/Source/Programs/Horde/Horde.Agent/Program.cs
@@ -212,7 +212,8 @@ namespace Horde.Agent
 			services.AddSingleton<ISessionFactory, SessionFactory>();
 			services.AddSingleton<IServerLoggerFactory, ServerLoggerFactory>();
 			services.AddSingleton<IServerStorageFactory, HttpServerStorageFactory>();
-			services.AddHostedService<WorkerService>();
+			services.AddSingleton<WorkerService>();
+			services.AddHostedService(sp => sp.GetRequiredService<WorkerService>());
 
 			services.AddSingleton<IStorageClientFactory, StorageClientFactory>();
 
