Horde: Add support for creating partial MongoDB indexes.

diff --git a/Engine/Source/Programs/Horde/Horde.Server/Server/MongoService.cs b/Engine/Source/Programs/Horde/Horde.Server/Server/MongoService.cs
--- a/Engine/Source/Programs/Horde/Horde.Server/Server/MongoService.cs
+++ b/Engine/Source/Programs/Horde/Horde.Server/Server/MongoService.cs
@@ -53,6 +53,12 @@ namespace Horde.Server.Server
 		[BsonElement("key"), BsonRequired]
 		public BsonDocument KeysDocument { get; private set; }
 
+		/// <summary>
+		/// Optional partial filter for the index
+		/// </summary>
+		[BsonElement("partialFilterExpression")]
+		public BsonDocument? PartialFilterDocument { get; private set; }
+
 		/// <summary>
 		/// The index should be unique
 		/// </summary>
@@ -67,39 +73,43 @@ namespace Horde.Server.Server
 		{
 			Name = String.Empty;
 			KeysDocument = new BsonDocument();
+			PartialFilterDocument = null;
 		}
 
 		/// <summary>
 		/// Constructor
 		/// </summary>
+		/// <param name="name">Name of the index</param>
 		/// <param name="keysDocument">Keys for the index</param>
+		/// <param name="partialFilterDocument">Partial filter for the index</param>
 		/// <param name="unique">Whether the index should be unique</param>
-		public MongoIndex(BsonDocument keysDocument, bool unique)
-			: this(GetDefaultName(keysDocument), keysDocument, unique)
+		public MongoIndex(string? name, BsonDocument keysDocument, BsonDocument? partialFilterDocument, bool unique)
 		{
+			Name = name ?? GetDefaultName(keysDocument);
+			KeysDocument = keysDocument;
+			PartialFilterDocument = partialFilterDocument;
+			Unique = unique;
 		}
 
 		/// <summary>
 		/// Constructor
 		/// </summary>
-		/// <param name="name">Name of the index</param>
-		/// <param name="keysDocument">Keys for the index</param>
+		/// <param name="keysFunc">Callback to configure keys for this document</param>
 		/// <param name="unique">Whether the index should be unique</param>
-		public MongoIndex(string name, BsonDocument keysDocument, bool unique)
+		public static MongoIndex<T> Create<T>(Func<IndexKeysDefinitionBuilder<T>, IndexKeysDefinition<T>> keysFunc, bool unique = false)
 		{
-			Name = name;
-			KeysDocument = keysDocument;
-			Unique = unique;
+			return new MongoIndex<T>(null, keysFunc(Builders<T>.IndexKeys), null, unique);
 		}
 
 		/// <summary>
 		/// Constructor
 		/// </summary>
 		/// <param name="keysFunc">Callback to configure keys for this document</param>
+		/// <param name="partialFilter">Partial filter for the index</param>
 		/// <param name="unique">Whether the index should be unique</param>
-		public static MongoIndex<T> Create<T>(Func<IndexKeysDefinitionBuilder<T>, IndexKeysDefinition<T>> keysFunc, bool unique = false)
+		public static MongoIndex<T> Create<T>(Func<IndexKeysDefinitionBuilder<T>, IndexKeysDefinition<T>> keysFunc, FilterDefinition<T>? partialFilter, bool unique = false)
 		{
-			return new MongoIndex<T>(keysFunc(Builders<T>.IndexKeys), unique);
+			return new MongoIndex<T>(null, keysFunc(Builders<T>.IndexKeys), partialFilter, unique);
 		}
 
 		/// <summary>
@@ -110,7 +120,19 @@ namespace Horde.Server.Server
 		/// <param name="unique">Whether the index should be unique</param>
 		public static MongoIndex<T> Create<T>(string name, Func<IndexKeysDefinitionBuilder<T>, IndexKeysDefinition<T>> keysFunc, bool unique = false)
 		{
-			return new MongoIndex<T>(name, keysFunc(Builders<T>.IndexKeys), unique);
+			return new MongoIndex<T>(name, keysFunc(Builders<T>.IndexKeys), null, unique);
+		}
+
+		/// <summary>
+		/// Constructor
+		/// </summary>
+		/// <param name="name">Name of the index</param>
+		/// <param name="keysFunc">Callback to configure keys for this document</param>
+		/// <param name="partialFilter">Partial filter for the index</param>
+		/// <param name="unique">Whether the index should be unique</param>
+		public static MongoIndex<T> Create<T>(string name, Func<IndexKeysDefinitionBuilder<T>, IndexKeysDefinition<T>> keysFunc, FilterDefinition<T>? partialFilter, bool unique = false)
+		{
+			return new MongoIndex<T>(name, keysFunc(Builders<T>.IndexKeys), partialFilter, unique);
 		}
 
 		/// <inheritdoc/>
@@ -156,26 +178,22 @@ namespace Horde.Server.Server
 		public IndexKeysDefinition<T> Keys { get; }
 
 		/// <summary>
-		/// Constructor
+		/// Optional expression specifying documents to include in the index
 		/// </summary>
-		/// <param name="keys">Keys for the index</param>
-		/// <param name="unique">Whether the index should be unique</param>
-		public MongoIndex(IndexKeysDefinition<T> keys, bool unique)
-			: base(keys.Render(BsonSerializer.LookupSerializer<T>(), BsonSerializer.SerializerRegistry), unique)
-		{
-			Keys = keys;
-		}
+		public FilterDefinition<T>? PartialFilter { get; }
 
 		/// <summary>
 		/// Constructor
 		/// </summary>
 		/// <param name="name">Name of the index</param>
 		/// <param name="keys">Keys for the index</param>
+		/// <param name="partialFilter">Partial filter for the index</param>
 		/// <param name="unique">Whether the index should be unique</param>
-		public MongoIndex(string name, IndexKeysDefinition<T> keys, bool unique)
-			: base(name, keys.Render(BsonSerializer.LookupSerializer<T>(), BsonSerializer.SerializerRegistry), unique)
+		public MongoIndex(string? name, IndexKeysDefinition<T> keys, FilterDefinition<T>? partialFilter, bool unique)
+			: base(name, keys.Render(BsonSerializer.LookupSerializer<T>(), BsonSerializer.SerializerRegistry), partialFilter?.Render(BsonSerializer.LookupSerializer<T>(), BsonSerializer.SerializerRegistry), unique)
 		{
 			Keys = keys;
+			PartialFilter = partialFilter;
 		}
 	}
 
@@ -189,10 +207,20 @@ namespace Horde.Server.Server
 			list.Add(MongoIndex.Create(keyFunc, unique));
 		}
 
+		public static void Add<T>(this List<MongoIndex<T>> list, Func<IndexKeysDefinitionBuilder<T>, IndexKeysDefinition<T>> keyFunc, FilterDefinition<T>? partialFilter, bool unique = false)
+		{
+			list.Add(MongoIndex.Create(keyFunc, partialFilter, unique));
+		}
+
 		public static void Add<T>(this List<MongoIndex<T>> list, string name, Func<IndexKeysDefinitionBuilder<T>, IndexKeysDefinition<T>> keyFunc, bool unique = false)
 		{
 			list.Add(MongoIndex.Create(name, keyFunc, unique));
 		}
+
+		public static void Add<T>(this List<MongoIndex<T>> list, string name, Func<IndexKeysDefinitionBuilder<T>, IndexKeysDefinition<T>> keyFunc, FilterDefinition<T> partialIndex, bool unique = false)
+		{
+			list.Add(MongoIndex.Create(name, keyFunc, partialIndex, unique));
+		}
 	}
 
 	/// <summary>
@@ -751,6 +779,7 @@ namespace Horde.Server.Server
 
 					CreateIndexOptions<T> options = new CreateIndexOptions<T>();
 					options.Name = createIndex.Name;
+					options.PartialFilterExpression = createIndex.PartialFilter;
 					options.Unique = createIndex.Unique;
 
 					CreateIndexModel<T> model = new CreateIndexModel<T>(createIndex.Keys, options);
