Horde: Add configurable expiry times for each artifact type. See example in Globals.json.

diff --git a/Engine/Source/Programs/Horde/Horde.Server.Tests/ArtifactTests.cs b/Engine/Source/Programs/Horde/Horde.Server.Tests/ArtifactTests.cs
--- a/Engine/Source/Programs/Horde/Horde.Server.Tests/ArtifactTests.cs
+++ b/Engine/Source/Programs/Horde/Horde.Server.Tests/ArtifactTests.cs
@@ -1,12 +1,15 @@
 // Copyright Epic Games, Inc. All Rights Reserved.
 
+using System;
 using System.Collections.Generic;
 using System.Linq;
+using System.Threading;
 using System.Threading.Tasks;
 using EpicGames.Horde.Storage;
 using Horde.Server.Acls;
 using Horde.Server.Artifacts;
 using Horde.Server.Storage;
+using HordeCommon;
 using Microsoft.Extensions.DependencyInjection;
 using Microsoft.VisualStudio.TestTools.UnitTesting;
 
@@ -19,7 +22,7 @@ namespace Horde.Server.Tests
 		public async Task CreateArtifact()
 		{
 			IArtifactCollection artifactCollection = ServiceProvider.GetRequiredService<IArtifactCollection>();
-			IArtifact artifact = await artifactCollection.AddAsync(ArtifactId.GenerateNewId(), ArtifactType.StepOutput, new string[] { "test1", "test2" }, Namespace.Artifacts, new RefName("test"), AclScopeName.Root);
+			IArtifact artifact = await artifactCollection.AddAsync(ArtifactId.GenerateNewId(), ArtifactType.StepOutput, new string[] { "test1", "test2" }, Namespace.Artifacts, new RefName("test"), null, AclScopeName.Root);
 
 			{
 				List<IArtifact> artifacts = await artifactCollection.FindAsync(keys: new[] { "test1" }).ToListAsync();
@@ -38,5 +41,30 @@ namespace Horde.Server.Tests
 				Assert.AreEqual(0, artifacts.Count);
 			}
 		}
+
+		[TestMethod]
+		public async Task ExpireArtifact()
+		{
+			FakeClock clock = ServiceProvider.GetRequiredService<FakeClock>();
+			ArtifactExpirationService expirationService = ServiceProvider.GetRequiredService<ArtifactExpirationService>();
+
+			await expirationService.StartAsync(CancellationToken.None);
+			
+			IArtifactCollection artifactCollection = ServiceProvider.GetRequiredService<IArtifactCollection>();
+			IArtifact artifact = await artifactCollection.AddAsync(ArtifactId.GenerateNewId(), ArtifactType.StepOutput, new string[] { "test1", "test2" }, Namespace.Artifacts, new RefName("test"), clock.UtcNow + TimeSpan.FromHours(1.0), AclScopeName.Root);
+
+			{
+				List<IArtifact> artifacts = await artifactCollection.FindAsync(keys: new[] { "test1" }).ToListAsync();
+				Assert.AreEqual(1, artifacts.Count);
+				Assert.AreEqual(artifact.Id, artifacts[0].Id);
+			}
+
+			await clock.AdvanceAsync(TimeSpan.FromHours(2.0));
+
+			{
+				List<IArtifact> artifacts = await artifactCollection.FindAsync(keys: new[] { "test1" }).ToListAsync();
+				Assert.AreEqual(0, artifacts.Count);
+			}
+		}
 	}
 }
diff --git a/Engine/Source/Programs/Horde/Horde.Server.Tests/TestSetup.cs b/Engine/Source/Programs/Horde/Horde.Server.Tests/TestSetup.cs
--- a/Engine/Source/Programs/Horde/Horde.Server.Tests/TestSetup.cs
+++ b/Engine/Source/Programs/Horde/Horde.Server.Tests/TestSetup.cs
@@ -237,6 +237,7 @@ namespace Horde.Server.Tests
 
 			services.AddSingleton<AclService>();
 			services.AddSingleton<AgentService>();
+			services.AddSingleton<ArtifactExpirationService>();
 			services.AddSingleton<FleetService>();
 			services.AddSingleton<ConsistencyService>();
 			services.AddSingleton<RequestTrackerService>();
diff --git a/Engine/Source/Programs/Horde/Horde.Server/Artifacts/ArtifactCollection.cs b/Engine/Source/Programs/Horde/Horde.Server/Artifacts/ArtifactCollection.cs
--- a/Engine/Source/Programs/Horde/Horde.Server/Artifacts/ArtifactCollection.cs
+++ b/Engine/Source/Programs/Horde/Horde.Server/Artifacts/ArtifactCollection.cs
@@ -11,6 +11,7 @@ using Horde.Server.Acls;
 using Horde.Server.Server;
 using Horde.Server.Storage;
 using Horde.Server.Utilities;
+using HordeCommon;
 using MongoDB.Bson.Serialization.Attributes;
 using MongoDB.Driver;
 
@@ -43,6 +44,9 @@ namespace Horde.Server.Artifacts
 			[BsonElement("scp")]
 			public AclScopeName AclScope { get; set; }
 
+			[BsonElement("exp")]
+			public DateTime? ExpireAtUtc { get; set; }
+
 			[BsonElement("upd")]
 			public int UpdateIndex { get; set; }
 
@@ -51,13 +55,14 @@ namespace Horde.Server.Artifacts
 			{
 			}
 
-			public Artifact(ArtifactId id, ArtifactType type, IEnumerable<string> keys, NamespaceId namespaceId, RefName refName, AclScopeName scopeName)
+			public Artifact(ArtifactId id, ArtifactType type, IEnumerable<string> keys, NamespaceId namespaceId, RefName refName, DateTime? expireAtUtc, AclScopeName scopeName)
 			{
 				Id = id;
 				Type = type;
 				Keys.AddRange(keys);
 				NamespaceId = namespaceId;
 				RefName = refName;
+				ExpireAtUtc = expireAtUtc;
 				AclScope = scopeName;
 			}
 		}
@@ -70,17 +75,28 @@ namespace Horde.Server.Artifacts
 		/// <param name="mongoService">The database service</param>
 		public ArtifactCollection(MongoService mongoService)
 		{
-			_artifacts = mongoService.GetCollection<Artifact>("ArtifactsV2", x => x.Ascending(x => x.Keys));
+			List<MongoIndex<Artifact>> indexes = new List<MongoIndex<Artifact>>();
+			indexes.Add(keys => keys.Ascending(x => x.Keys));
+			indexes.Add(keys => keys.Ascending(x => x.ExpireAtUtc), Builders<Artifact>.Filter.Exists(x => x.ExpireAtUtc));
+			_artifacts = mongoService.GetCollection<Artifact>("ArtifactsV2", indexes);
 		}
 
 		/// <inheritdoc/>
-		public async Task<IArtifact> AddAsync(ArtifactId id, ArtifactType type, IEnumerable<string> keys, NamespaceId namespaceId, RefName refName, AclScopeName scopeName, CancellationToken cancellationToken)
+		public async Task<IArtifact> AddAsync(ArtifactId id, ArtifactType type, IEnumerable<string> keys, NamespaceId namespaceId, RefName refName, DateTime? expireAtUtc, AclScopeName scopeName, CancellationToken cancellationToken)
 		{
-			Artifact artifact = new Artifact(id, type, keys, namespaceId, refName, scopeName);
+			Artifact artifact = new Artifact(id, type, keys, namespaceId, refName, expireAtUtc, scopeName);
 			await _artifacts.InsertOneAsync(artifact, null, cancellationToken);
 			return artifact;
 		}
 
+
+		/// <inheritdoc/>
+		public async Task DeleteAsync(IEnumerable<ArtifactId> ids, CancellationToken cancellationToken = default)
+		{
+			FilterDefinition<Artifact> filter = Builders<Artifact>.Filter.In(x => x.Id, ids);
+			await _artifacts.DeleteManyAsync(filter, cancellationToken);
+		}
+
 		/// <inheritdoc/>
 		public async IAsyncEnumerable<IArtifact> FindAsync(IEnumerable<ArtifactId>? ids = null, IEnumerable<string>? keys = null, [EnumeratorCancellation] CancellationToken cancellationToken = default)
 		{
@@ -106,10 +122,47 @@ namespace Horde.Server.Artifacts
 			}
 		}
 
+		/// <summary>
+		/// Finds artifacts which are ready for expiry
+		/// </summary>
+		/// <param name="utcNow">Current time for expiring artifacts</param>
+		/// <param name="cancellationToken">Cancellation token for the operation</param>
+		/// <returns>Sequence of artifacts</returns>
+		public async IAsyncEnumerable<IEnumerable<IArtifact>> FindExpiredAsync(DateTime utcNow, [EnumeratorCancellation] CancellationToken cancellationToken = default)
+		{
+			using (IAsyncCursor<Artifact> cursor = await _artifacts.Find(x => x.ExpireAtUtc!.Value < utcNow).ToCursorAsync(cancellationToken))
+			{
+				while (await cursor.MoveNextAsync(cancellationToken))
+				{
+					yield return cursor.Current;
+				}
+			}
+		}
+
 		/// <inheritdoc/>
 		public async Task<IArtifact?> GetAsync(ArtifactId artifactId, CancellationToken cancellationToken)
 		{
 			return await _artifacts.Find(x => x.Id == artifactId).FirstOrDefaultAsync(cancellationToken);
 		}
+
+		/// <inheritdoc/>
+		public async Task<IArtifact?> TryUpdateAsync(IArtifact artifact, DateTime? expireAtUtc, CancellationToken cancellationToken)
+		{
+			List<UpdateDefinition<Artifact>> updates = new List<UpdateDefinition<Artifact>>();
+			if (expireAtUtc != null)
+			{
+				updates.Add(Builders<Artifact>.Update.SetOrUnsetNull(x => x.ExpireAtUtc, expireAtUtc));
+			}
+			if (updates.Count == 0)
+			{
+				return artifact;
+			}
+
+			Artifact artifactDoc = (Artifact)artifact;
+			FilterDefinition<Artifact> filter = Builders<Artifact>.Filter.Expr(x => x.Id == artifact.Id && x.UpdateIndex == artifactDoc.UpdateIndex);
+			UpdateDefinition<Artifact> update = Builders<Artifact>.Update.Combine(updates);
+
+			return await _artifacts.FindOneAndUpdateAsync(filter, update, new FindOneAndUpdateOptions<Artifact> { ReturnDocument = ReturnDocument.After }, cancellationToken);
+		}
 	}
 }
diff --git a/Engine/Source/Programs/Horde/Horde.Server/Artifacts/ArtifactExpirationService.cs b/Engine/Source/Programs/Horde/Horde.Server/Artifacts/ArtifactExpirationService.cs
new file mode 100644
--- /dev/null
+++ b/Engine/Source/Programs/Horde/Horde.Server/Artifacts/ArtifactExpirationService.cs
@@ -0,0 +1,66 @@
+// Copyright Epic Games, Inc. All Rights Reserved.
+
+using System;
+using System.Collections.Generic;
+using System.Linq;
+using System.Threading;
+using System.Threading.Tasks;
+using EpicGames.Horde.Storage;
+using Horde.Server.Storage;
+using HordeCommon;
+using Microsoft.Extensions.Hosting;
+using Microsoft.Extensions.Logging;
+
+namespace Horde.Server.Artifacts
+{
+	/// <summary>
+	/// Expires artifacts according to their 
+	/// </summary>
+	class ArtifactExpirationService : IHostedService
+	{
+		readonly IArtifactCollection _artifactCollection;
+		readonly StorageService _storageService;
+		readonly IClock _clock;
+		readonly ITicker _ticker;
+		readonly ILogger _logger;
+
+		public ArtifactExpirationService(IArtifactCollection artifactCollection, StorageService storageService, IClock clock, ILogger<ArtifactExpirationService> logger)
+		{
+			_artifactCollection = artifactCollection;
+			_storageService = storageService;
+			_clock = clock;
+			_ticker = clock.AddSharedTicker<ArtifactExpirationService>(TimeSpan.FromHours(1.0), TickAsync, logger);
+			_logger = logger;
+		}
+
+		/// <inheritdoc/>
+		public async Task StartAsync(CancellationToken cancellationToken)
+		{
+			await _ticker.StartAsync();
+		}
+
+		/// <inheritdoc/>
+		public async Task StopAsync(CancellationToken cancellationToken)
+		{
+			await _ticker.StopAsync();
+		}
+
+		async ValueTask TickAsync(CancellationToken cancellationToken)
+		{
+			DateTime utcNow = _clock.UtcNow;
+			await foreach (IEnumerable<IArtifact> artifacts in _artifactCollection.FindExpiredAsync(utcNow, cancellationToken))
+			{
+				foreach (IGrouping<NamespaceId, IArtifact> group in artifacts.GroupBy(x => x.NamespaceId))
+				{
+					IStorageClient storageClient = await _storageService.GetClientAsync(group.Key, cancellationToken);
+					foreach (IArtifact artifact in group)
+					{
+						_logger.LogDebug("Expiring artifact {ArtifactId}, ref {RefName}", artifact.Id, artifact.RefName);
+						await storageClient.DeleteRefAsync(artifact.RefName, cancellationToken);
+					}
+				}
+				await _artifactCollection.DeleteAsync(artifacts.Select(x => x.Id), cancellationToken);
+			}
+		}
+	}
+}
diff --git a/Engine/Source/Programs/Horde/Horde.Server/Artifacts/IArtifact.cs b/Engine/Source/Programs/Horde/Horde.Server/Artifacts/IArtifact.cs
--- a/Engine/Source/Programs/Horde/Horde.Server/Artifacts/IArtifact.cs
+++ b/Engine/Source/Programs/Horde/Horde.Server/Artifacts/IArtifact.cs
@@ -1,5 +1,6 @@
 // Copyright Epic Games, Inc. All Rights Reserved.
 
+using System;
 using System.Collections.Generic;
 using EpicGames.Horde.Storage;
 using Horde.Server.Acls;
@@ -36,6 +37,11 @@ namespace Horde.Server.Artifacts
 		/// </summary>
 		public RefName RefName { get; }
 
+		/// <summary>
+		/// Time at which the artifact is scheduled to be deleted
+		/// </summary>
+		public DateTime? ExpireAtUtc { get; }
+
 		/// <summary>
 		/// Permissions scope for this object
 		/// </summary>
diff --git a/Engine/Source/Programs/Horde/Horde.Server/Artifacts/IArtifactCollection.cs b/Engine/Source/Programs/Horde/Horde.Server/Artifacts/IArtifactCollection.cs
--- a/Engine/Source/Programs/Horde/Horde.Server/Artifacts/IArtifactCollection.cs
+++ b/Engine/Source/Programs/Horde/Horde.Server/Artifacts/IArtifactCollection.cs
@@ -22,10 +22,27 @@ namespace Horde.Server.Artifacts
 		/// <param name="keys">Keys for the artifact</param>
 		/// <param name="namespaceId">Namespace containing the data</param>
 		/// <param name="refName">Artifact ref name</param>
+		/// <param name="expireAtUtc">Time at which to expire the artifact</param>
 		/// <param name="scopeName">Inherited scope used for permissions</param>
 		/// <param name="cancellationToken">Cancellation token for the operation</param>
 		/// <returns>The new log file document</returns>
-		Task<IArtifact> AddAsync(ArtifactId artifactId, ArtifactType type, IEnumerable<string> keys, NamespaceId namespaceId, RefName refName, AclScopeName scopeName, CancellationToken cancellationToken = default);
+		Task<IArtifact> AddAsync(ArtifactId artifactId, ArtifactType type, IEnumerable<string> keys, NamespaceId namespaceId, RefName refName, DateTime? expireAtUtc, AclScopeName scopeName, CancellationToken cancellationToken = default);
+
+		/// <summary>
+		/// Deletes artifacts
+		/// </summary>
+		/// <param name="ids">Ids to search for</param>
+		/// <param name="cancellationToken">Cancellation token for the operation</param>
+		/// <returns>Sequence of artifacts</returns>
+		Task DeleteAsync(IEnumerable<ArtifactId> ids, CancellationToken cancellationToken = default);
+
+		/// <summary>
+		/// Finds artifacts which are ready for expiry
+		/// </summary>
+		/// <param name="utcNow">The current time, as UTC</param>
+		/// <param name="cancellationToken">Cancellation token for the operation</param>
+		/// <returns>Sequence of artifacts</returns>
+		IAsyncEnumerable<IEnumerable<IArtifact>> FindExpiredAsync(DateTime utcNow, CancellationToken cancellationToken = default);
 
 		/// <summary>
 		/// Finds artifacts with the given keys
@@ -43,5 +60,26 @@ namespace Horde.Server.Artifacts
 		/// <param name="cancellationToken">Cancellation token for the operation</param>
 		/// <returns>The artifact document</returns>
 		Task<IArtifact?> GetAsync(ArtifactId artifactId, CancellationToken cancellationToken = default);
+
+		/// <summary>
+		/// Updates an artifact
+		/// </summary>
+		/// <param name="artifact">Artifact to update</param>
+		/// <param name="expiresAtUtc">Expiry time for the artifact. Set to a default value to clear the expiry time.</param>
+		/// <param name="cancellationToken">Cancellation token for the operation</param>
+		/// <returns>Updated artifact</returns>
+		Task<IArtifact?> TryUpdateAsync(IArtifact artifact, DateTime? expiresAtUtc, CancellationToken cancellationToken = default);
+	}
+
+	/// <summary>
+	/// Extension methods for <see cref="IArtifactCollection"/>
+	/// </summary>
+	public static class ArtifactCollectionExtensions
+	{
+		/// <inheritdoc cref="IArtifactCollection.DeleteAsync(IEnumerable{ArtifactId}, CancellationToken)"/>>
+		public static Task DeleteAsync(this IArtifactCollection artifactCollection, ArtifactId artifactId, CancellationToken cancellationToken = default)
+		{
+			return artifactCollection.DeleteAsync(new[] { artifactId }, cancellationToken);
+		}
 	}
 }
diff --git a/Engine/Source/Programs/Horde/Horde.Server/Defaults/globals.json b/Engine/Source/Programs/Horde/Horde.Server/Defaults/globals.json
--- a/Engine/Source/Programs/Horde/Horde.Server/Defaults/globals.json
+++ b/Engine/Source/Programs/Horde/Horde.Server/Defaults/globals.json
@@ -78,5 +78,11 @@
 			"name": "Linux",
 			"condition": "Platform == 'Linux'"
 		}
+	],
+	"artifactTypes": [
+		{
+			"name": "step-output",
+			"keepDays": 1
+		}
 	]
 }
diff --git a/Engine/Source/Programs/Horde/Horde.Server/Jobs/JobRpcService.cs b/Engine/Source/Programs/Horde/Horde.Server/Jobs/JobRpcService.cs
--- a/Engine/Source/Programs/Horde/Horde.Server/Jobs/JobRpcService.cs
+++ b/Engine/Source/Programs/Horde/Horde.Server/Jobs/JobRpcService.cs
@@ -97,7 +97,13 @@ namespace Horde.Server.Jobs
 			NamespaceId namespaceId = String.IsNullOrEmpty(request.NamespaceId)? Namespace.Artifacts : new NamespaceId(request.NamespaceId);
 			RefName refName = new RefName(String.IsNullOrEmpty(request.RefName) ? $"job-{job.Id}/step-{step.Id}/{artifactId}" : request.RefName);
 
-			IArtifact artifact = await _artifactCollection.AddAsync(artifactId, type, keys, namespaceId, refName, templateConfig.ScopeName, context.CancellationToken);
+			DateTime? expireAt = null;
+			if (_globalConfig.TryGetArtifactType(type, out ArtifactTypeConfig? typeConfig) && typeConfig.KeepDays != null && typeConfig.KeepDays.Value >= 0)
+			{
+				expireAt = DateTime.UtcNow + TimeSpan.FromDays(typeConfig.KeepDays.Value);
+			}
+			
+			IArtifact artifact = await _artifactCollection.AddAsync(artifactId, type, keys, namespaceId, refName, expireAt, templateConfig.ScopeName, context.CancellationToken);
 
 			List<AclClaimConfig> claims = new List<AclClaimConfig>();
 			claims.Add(new AclClaimConfig(HordeClaimTypes.WriteNamespace, namespaceId.ToString()));
diff --git a/Engine/Source/Programs/Horde/Horde.Server/Server/GlobalConfig.cs b/Engine/Source/Programs/Horde/Horde.Server/Server/GlobalConfig.cs
--- a/Engine/Source/Programs/Horde/Horde.Server/Server/GlobalConfig.cs
+++ b/Engine/Source/Programs/Horde/Horde.Server/Server/GlobalConfig.cs
@@ -16,6 +16,7 @@ using EpicGames.Serialization;
 using Horde.Server.Acls;
 using Horde.Server.Agents.Pools;
 using Horde.Server.Agents.Software;
+using Horde.Server.Artifacts;
 using Horde.Server.Configuration;
 using Horde.Server.Projects;
 using Horde.Server.Storage;
@@ -75,6 +76,22 @@ namespace Horde.Server.Server
 		public bool ShowTests { get; set; } = true;
 	}
 
+	/// <summary>
+	/// Configuration for an artifact
+	/// </summary>
+	public class ArtifactTypeConfig
+	{
+		/// <summary>
+		/// Name of the artifact type
+		/// </summary>
+		public ArtifactType Name { get; }
+
+		/// <summary>
+		/// Number of days to retain artifacts of this type
+		/// </summary>
+		public int? KeepDays { get; }
+	}
+
 	/// <summary>
 	/// Global configuration
 	/// </summary>
@@ -174,6 +191,11 @@ namespace Horde.Server.Server
 		/// </summary>
 		public StorageConfig Storage { get; set; } = new StorageConfig();
 
+		/// <summary>
+		/// Configuration for different artifact types
+		/// </summary>
+		public List<ArtifactTypeConfig> ArtifactTypes { get; set; } = new List<ArtifactTypeConfig>();
+
 		/// <summary>
 		/// Access control list
 		/// </summary>
@@ -190,6 +212,7 @@ namespace Horde.Server.Server
 		private readonly Dictionary<ToolId, ToolConfig> _toolLookup = new Dictionary<ToolId, ToolConfig>();
 		private readonly Dictionary<ClusterId, ComputeClusterConfig> _computeClusterLookup = new Dictionary<ClusterId, ComputeClusterConfig>();
 		private readonly Dictionary<AclScopeName, IAclScope> _aclScopeLookup = new Dictionary<AclScopeName, IAclScope>();
+		private readonly Dictionary<ArtifactType, ArtifactTypeConfig> _artifactTypeLookup = new Dictionary<ArtifactType, ArtifactTypeConfig>();
 
 		/// <summary>
 		/// Called after the config file has been read
@@ -242,9 +265,23 @@ namespace Horde.Server.Server
 				}
 			}
 
+			_artifactTypeLookup.Clear();
+			foreach (ArtifactTypeConfig artifactType in ArtifactTypes)
+			{
+				_artifactTypeLookup.Add(artifactType.Name, artifactType);
+			}
+
 			Storage.PostLoad(this);
 		}
 
+		/// <summary>
+		/// Attempts to get configuration for a specific artifact type
+		/// </summary>
+		/// <param name="type">The artifact type name</param>
+		/// <param name="config">Configuration for the stream</param>
+		/// <returns>True if the stream configuration was found</returns>
+		public bool TryGetArtifactType(ArtifactType type, [NotNullWhen(true)] out ArtifactTypeConfig? config) => _artifactTypeLookup.TryGetValue(type, out config);
+
 		/// <summary>
 		/// Attempts to get configuration for a project from this object
 		/// </summary>
diff --git a/Engine/Source/Programs/Horde/Horde.Server/Startup.cs b/Engine/Source/Programs/Horde/Horde.Server/Startup.cs
--- a/Engine/Source/Programs/Horde/Horde.Server/Startup.cs
+++ b/Engine/Source/Programs/Horde/Horde.Server/Startup.cs
@@ -443,6 +443,7 @@ namespace Horde.Server
 
 			services.AddSingleton<AclService>();
 			services.AddSingleton<AgentService>();
+			services.AddHostedService<ArtifactExpirationService>();
 			services.AddSingleton<ConsistencyService>();
 			services.AddSingleton<RequestTrackerService>();
 			services.AddSingleton<ComputeService>();
