Horde: Disable HTTPS by default.

diff --git a/Engine/Source/Programs/Horde/Horde.Server/appsettings.json b/Engine/Source/Programs/Horde/Horde.Server/appsettings.json
--- a/Engine/Source/Programs/Horde/Horde.Server/appsettings.json
+++ b/Engine/Source/Programs/Horde/Horde.Server/appsettings.json
@@ -1,7 +1,7 @@
 {
 	"Horde": {
 		"HttpPort": 5000,
-		"HttpsPort": 5001,
+		"HttpsPort": 0,
 		"Http2Port": 5002,
 		"ConsoleLogLevel": "Debug",
 		"DashboardUrl": "https://localhost:3000",
