Horde: Add an better error message for trying to set a property without a setter method.

diff --git a/Engine/Source/Programs/Horde/Horde.Server/Configuration/ConfigType.cs b/Engine/Source/Programs/Horde/Horde.Server/Configuration/ConfigType.cs
--- a/Engine/Source/Programs/Horde/Horde.Server/Configuration/ConfigType.cs
+++ b/Engine/Source/Programs/Horde/Horde.Server/Configuration/ConfigType.cs
@@ -230,6 +230,12 @@ namespace Horde.Server.Configuration
 				{
 					value = JsonSerializer.Deserialize(node, PropertyInfo.PropertyType, context.JsonOptions);
 				}
+
+				if (!PropertyInfo.CanWrite)
+				{
+					throw new ConfigException(context, $"Property {context.CurrentScope}{Name} does not have a setter.");
+				}
+
 				PropertyInfo.SetValue(target, value);
 
 				return Task.CompletedTask;
@@ -262,6 +268,11 @@ namespace Horde.Server.Configuration
 			{
 				context.AddProperty(Name);
 
+				if (!PropertyInfo.CanWrite)
+				{
+					throw new ConfigException(context, $"Property {context.CurrentScope}{Name} does not have a setter.");
+				}
+
 				Uri uri = CombinePaths(context.CurrentFile, JsonSerializer.Deserialize<string>(node, context.JsonOptions) ?? String.Empty);
 				IConfigFile file = await ReadFileAsync(uri, context, cancellationToken);
 
