Horde: Fix errors configuring artifact type properties.

diff --git a/Engine/Source/Programs/Horde/Horde.Server/Configuration/ConfigType.cs b/Engine/Source/Programs/Horde/Horde.Server/Configuration/ConfigType.cs
--- a/Engine/Source/Programs/Horde/Horde.Server/Configuration/ConfigType.cs
+++ b/Engine/Source/Programs/Horde/Horde.Server/Configuration/ConfigType.cs
@@ -233,7 +233,7 @@ namespace Horde.Server.Configuration
 
 				if (!PropertyInfo.CanWrite)
 				{
-					throw new ConfigException(context, $"Property {context.CurrentScope}{Name} does not have a setter.");
+					throw new ConfigException(context, $"Property {context.CurrentScope}.{Name} does not have a setter.");
 				}
 
 				PropertyInfo.SetValue(target, value);
@@ -403,7 +403,7 @@ namespace Horde.Server.Configuration
 					{
 						if (otherFile != null)
 						{
-							throw new ConfigException(context, $"Property {context.CurrentScope}{Name} conflicts with value in {otherFile}.");
+							throw new ConfigException(context, $"Property {context.CurrentScope}.{Name} conflicts with value in {otherFile}.");
 						}
 
 						childTarget = await _classConfigType.ReadAsync(node, context, cancellationToken);
diff --git a/Engine/Source/Programs/Horde/Horde.Server/Server/GlobalConfig.cs b/Engine/Source/Programs/Horde/Horde.Server/Server/GlobalConfig.cs
--- a/Engine/Source/Programs/Horde/Horde.Server/Server/GlobalConfig.cs
+++ b/Engine/Source/Programs/Horde/Horde.Server/Server/GlobalConfig.cs
@@ -84,12 +84,12 @@ namespace Horde.Server.Server
 		/// <summary>
 		/// Name of the artifact type
 		/// </summary>
-		public ArtifactType Name { get; }
+		public ArtifactType Name { get; set; }
 
 		/// <summary>
 		/// Number of days to retain artifacts of this type
 		/// </summary>
-		public int? KeepDays { get; }
+		public int? KeepDays { get; set; }
 	}
 
 	/// <summary>
