Horde: Add a duration field to the jobstepref telemetry data.

diff --git a/Engine/Source/Programs/Horde/Horde.Server/Jobs/JobStepRefCollection.cs b/Engine/Source/Programs/Horde/Horde.Server/Jobs/JobStepRefCollection.cs
--- a/Engine/Source/Programs/Horde/Horde.Server/Jobs/JobStepRefCollection.cs
+++ b/Engine/Source/Programs/Horde/Horde.Server/Jobs/JobStepRefCollection.cs
@@ -138,7 +138,8 @@ namespace Horde.Server.Jobs
 					StartTime = startTimeUtc,
 					StreamId = streamId,
 					TemplateId = templateId,
-					UpdateIssues = updateIssues
+					UpdateIssues = updateIssues,
+					Duration = (finishTimeUtc != null) ? (double?)(finishTimeUtc.Value - startTimeUtc).TotalSeconds : null
 				});
 			}
 
