Horde: Add a handler for BuildGraph script parsing events.

diff --git a/Engine/Source/Programs/Horde/Horde.Server/Issues/Handlers/BuildGraphIssueHandler.cs b/Engine/Source/Programs/Horde/Horde.Server/Issues/Handlers/BuildGraphIssueHandler.cs
new file mode 100644
--- /dev/null
+++ b/Engine/Source/Programs/Horde/Horde.Server/Issues/Handlers/BuildGraphIssueHandler.cs
@@ -0,0 +1,69 @@
+// Copyright Epic Games, Inc. All Rights Reserved.
+
+using System;
+using System.Collections.Generic;
+using System.Linq;
+using EpicGames.Core;
+using Horde.Server.Jobs;
+using Horde.Server.Jobs.Graphs;
+using Microsoft.Extensions.Logging;
+using MongoDB.Driver;
+
+namespace Horde.Server.Issues.Handlers
+{
+	/// <summary>
+	/// Instance of a particular BuildGraph script error
+	/// </summary>
+	class BuildGraphIssueHandler : SourceFileIssueHandler
+	{
+		/// <inheritdoc/>
+		public override string Type => "BuildGraph";
+
+		/// <inheritdoc/>
+		public override int Priority => 10;
+
+		/// <summary>
+		/// Determines if the given event id matches
+		/// </summary>
+		/// <param name="eventId">The event id to compare</param>
+		/// <returns>True if the given event id matches</returns>
+		public static bool IsMatchingEventId(EventId eventId)
+		{
+			return eventId == KnownLogEvents.AutomationTool_BuildGraphScript;
+		}
+
+		static bool IsMaskedEventId(EventId id) => id == KnownLogEvents.ExitCode || id == KnownLogEvents.Systemic_Xge_BuildFailed;
+
+		/// <inheritdoc/>
+		public override void TagEvents(IJob job, INode node, IReadOnlyNodeAnnotations annotations, IReadOnlyList<IssueEvent> stepEvents)
+		{
+			bool hasMatches = false;
+			foreach (IssueEvent stepEvent in stepEvents)
+			{
+				if (stepEvent.EventId.HasValue)
+				{
+					EventId eventId = stepEvent.EventId.Value;
+					if (IsMatchingEventId(eventId))
+					{
+						List<string> newFileNames = new List<string>();
+						GetSourceFiles(stepEvent.EventData, newFileNames);
+						stepEvent.Fingerprint = new NewIssueFingerprint(Type, newFileNames, null, null);
+						hasMatches = true;
+					}
+					else if (hasMatches && IsMaskedEventId(eventId))
+					{
+						stepEvent.Ignored = true;
+					}
+				}
+			}
+		}
+
+		/// <inheritdoc/>
+		public override string GetSummary(IIssueFingerprint fingerprint, IssueSeverity severity)
+		{
+			string level = (severity == IssueSeverity.Warning) ? "warnings" : "errors";
+			string list = StringUtils.FormatList(fingerprint.Keys.Where(x => !x.StartsWith(NotePrefix, StringComparison.Ordinal)).ToArray(), 2);
+			return $"BuildGraph {level} in {list}";
+		}
+	}
+}
