Horde: Bump the font size on timeline panel
#rnx
#skipci

diff --git a/Engine/Source/Programs/Horde/HordeDashboard/src/components/jobDetailsV2/JobDetailTimeline.tsx b/Engine/Source/Programs/Horde/HordeDashboard/src/components/jobDetailsV2/JobDetailTimeline.tsx
--- a/Engine/Source/Programs/Horde/HordeDashboard/src/components/jobDetailsV2/JobDetailTimeline.tsx
+++ b/Engine/Source/Programs/Horde/HordeDashboard/src/components/jobDetailsV2/JobDetailTimeline.tsx
@@ -525,9 +525,9 @@ class TimelineRenderer {
 
       const xAxis = (g: SelectionType) => {
 
-         g.attr("transform", `translate(0,16)`)
-            .style("font-family", "Horde Open Sans Regular")
-            .style("font-size", "9px")
+         g.attr("transform", `translate(0,20)`)
+            .style("font-family", "Horde Open Sans SemiBold")
+            .style("font-size", "12px")
             .call(d3.axisTop(xScale)
                .tickFormat(d => {
                   return msecToElapsed(((d as number) * 1000) - dataView.anchorMS!, true, false);
