Horde: More trends panel tweaks
#rnx
#skipci

diff --git a/Engine/Source/Programs/Horde/HordeDashboard/src/components/jobDetailsV2/JobDetailStepTrendsV2.tsx b/Engine/Source/Programs/Horde/HordeDashboard/src/components/jobDetailsV2/JobDetailStepTrendsV2.tsx
--- a/Engine/Source/Programs/Horde/HordeDashboard/src/components/jobDetailsV2/JobDetailStepTrendsV2.tsx
+++ b/Engine/Source/Programs/Horde/HordeDashboard/src/components/jobDetailsV2/JobDetailStepTrendsV2.tsx
@@ -208,7 +208,7 @@ class StepTrendsDataView extends JobDataView {
    order = 8;
 
    durations = new Map<string, number>();
-   maxMinutes = 0;   
+   maxMinutes = 0;
 
    minTime?: Date;
    maxTime?: Date;
@@ -348,22 +348,6 @@ class StepTrendsRenderer {
          .attr("marker-end", "url(#cmarker)")
          .attr("stroke", dashboard.darktheme ? "#6D6C6B" : "#4D4C4B")
 
-      // left axis
-      svg.append("g")
-         .attr("transform", `translate(${this.margin.left},0)`)
-         .style("font-family", "Horde Open Sans SemiBold")
-         .style("font-size", "12px")
-         .call(d3.axisLeft(y)
-            .ticks(7)
-            .tickFormat((d) => {
-
-               if (!d) {
-                  return "";
-               }
-
-               return msecToElapsed((d as number) * 60000, true, false);
-
-            }))
 
       svg.append("line")
          .attr("class", "bline")
@@ -410,15 +394,36 @@ class StepTrendsRenderer {
                .attr("y2", height - this.margin.bottom))
       }
 
+      // left axis
+      const yAxis = (g: SelectionType) => {
+
+         g.attr("transform", `translate(${this.margin.left},0)`)
+            .style("font-family", "Horde Open Sans SemiBold")
+            .style("font-size", "12px")
+            .call(d3.axisLeft(this.scaleY!)
+               .ticks(7)
+               .tickFormat((d) => {
+
+                  if (!d) {
+                     return "";
+                  }
+
+                  return msecToElapsed((d as number) * 60000, true, false);
+
+               }))
+      }
+
       // top axis
       svg.append("g").attr("class", "x-axis").call(xAxis)
 
+      svg.append("g").attr("class", "y-axis").call(yAxis)
+
       const zoom = this.zoom = d3.zoom()
          .scaleExtent([1, 8])
          .extent([[margin.left, 0], [width - margin.right, height]])
          .translateExtent([[margin.left, -Infinity], [width - margin.right, Infinity]])
          .on("zoom", zoomed);
-      
+
       const renderer = this;
 
       function zoomed(event: any) {
@@ -471,6 +476,8 @@ class StepTrendsRenderer {
 
          svg!.selectAll(".linechart")
             .attr("d", scaledLine(lineI as any));
+         
+            svg!.selectAll(".y-axis").call(yAxis as any);
 
       }
 
@@ -512,7 +519,7 @@ class StepTrendsRenderer {
 
          if (ref) {
             const posx = x(new Date(ref.startTime!).getTime() / 1000);
-            const posy = y(dataView.durations.get(ref.jobId)!);
+            const posy = this.scaleY!(dataView.durations.get(ref.jobId)!);
 
             svg!.selectAll(".cline")
                .attr("x1", () => posx)
