Horde: Fix issue with trends view when using timeline to select steps
#skipci
#rnx

diff --git a/Engine/Source/Programs/Horde/HordeDashboard/src/components/jobDetailsV2/JobDetailStepTrendsV2.tsx b/Engine/Source/Programs/Horde/HordeDashboard/src/components/jobDetailsV2/JobDetailStepTrendsV2.tsx
--- a/Engine/Source/Programs/Horde/HordeDashboard/src/components/jobDetailsV2/JobDetailStepTrendsV2.tsx
+++ b/Engine/Source/Programs/Horde/HordeDashboard/src/components/jobDetailsV2/JobDetailStepTrendsV2.tsx
@@ -240,7 +240,7 @@ class StepTrendsRenderer {
    render(container: HTMLDivElement) {
 
       if (this.hasRendered && !this.forceRender) {
-         return;
+         //return;
       }
 
       const dataView = this.dataView;
