Horde: Preconfigure the command line tool to communicate with the correct server URL.

diff --git a/Engine/Source/Programs/Horde/Horde/Commands/ConfigCommand.cs b/Engine/Source/Programs/Horde/Horde/Commands/ConfigCommand.cs
new file mode 100644
--- /dev/null
+++ b/Engine/Source/Programs/Horde/Horde/Commands/ConfigCommand.cs
@@ -0,0 +1,26 @@
+// Copyright Epic Games, Inc. All Rights Reserved.
+
+using EpicGames.Core;
+using Microsoft.Extensions.Logging;
+
+namespace Horde.Commands
+{
+	[Command("config", "Updates the configuration for the Horde tool")]
+	class ConfigCommand : Command
+	{
+		[CommandLine("-Server=", Description = "Updates the server URL")]
+		public string? Server { get; set; }
+
+		/// <inheritdoc/>
+		public override async Task<int> ExecuteAsync(ILogger logger)
+		{
+			if (Server != null)
+			{
+				await Settings.SetServerAsync(Server);
+			}
+
+			logger.LogInformation("Server: {Server}", await Settings.GetServerAsync());
+			return 0;
+		}
+	}
+}
