Horde: Re-add functionality to preconfigure the agent for the correct server URL.

diff --git a/Engine/Source/Programs/Horde/BuildHorde.xml b/Engine/Source/Programs/Horde/BuildHorde.xml
--- a/Engine/Source/Programs/Horde/BuildHorde.xml
+++ b/Engine/Source/Programs/Horde/BuildHorde.xml
@@ -91,20 +91,8 @@
 			<DotNet Arguments="publish &quot;$(RootDir)/Engine/Source/Programs/Horde/Horde.Agent/Horde.Agent.csproj&quot; --output &quot;$(StagingDir)/Agent&quot; --runtime win7-x64 --self-contained $(VersionArguments)"/>
 			<DotNet Arguments="publish &quot;$(RootDir)/Engine/Source/Programs/Horde/Horde.Agent.TrayApp/Horde.Agent.TrayApp.csproj&quot; --output &quot;$(StagingDir)/Agent&quot; --runtime win7-x64 --self-contained $(VersionArguments)"/>
 
-			<!-- Create another build of the agent, this time for all platforms, and chunk it for download. -->
-			<DotNet Arguments="publish &quot;$(RootDir)/Engine/Source/Programs/Horde/Horde.Agent/Horde.Agent.csproj&quot; --output &quot;$(StagingDir)/Agent-CrossPlatform&quot; $(VersionArguments)"/>
-			<Spawn Exe="$(StagingDir)/Tools/Horde.exe" Arguments="bundle create -input=&quot;$(StagingDir)/Agent-CrossPlatform&quot; -file=&quot;$(StagingDir)/Server/Tools/horde-agent/default.ref&quot;"/>
-			
-			<Property Name="NewTool" Separator="">
-			  {
-				"Id": "horde-agent",
-				"Name": "Horde Agent",
-				"Description": "Cross-platform build of the Horde Agent.",
-				"RefName": "default",
-				"Version": "$(InformationalVersion)"
-			  }
-			</Property>
-			<WriteJsonValue File="$(StagingDir)/Server/appsettings.json" Key="$.Horde.BundledTools[]" Value="$(NewTool)"/>
+			<!-- Create another build of the agent, this time for all platforms. -->
+			<DotNet Arguments="publish &quot;$(RootDir)/Engine/Source/Programs/Horde/Horde.Agent/Horde.Agent.csproj&quot; --output &quot;$(StagingDir)/Server/Tools/horde-agent-loose&quot; $(VersionArguments)"/>
 
 			<!-- Copy the Dashboard -->
 			<Copy Files="#DashboardOutputFiles" From="$(StagingDir)/Dashboard/..." To="$(StagingDir)/Server/DashboardApp/..." />
diff --git a/Engine/Source/Programs/Horde/Horde.Server/Commands/Install/SetupCommand.cs b/Engine/Source/Programs/Horde/Horde.Server/Commands/Install/SetupCommand.cs
new file mode 100644
--- /dev/null
+++ b/Engine/Source/Programs/Horde/Horde.Server/Commands/Install/SetupCommand.cs
@@ -0,0 +1,159 @@
+// Copyright Epic Games, Inc. All Rights Reserved.
+
+using System;
+using System.Buffers;
+using System.Diagnostics;
+using System.IdentityModel.Tokens.Jwt;
+using System.IO;
+using System.Net;
+using System.Security.Cryptography;
+using System.Security.Cryptography.X509Certificates;
+using System.Text.Json;
+using System.Text.Json.Nodes;
+using System.Text.RegularExpressions;
+using System.Threading.Tasks;
+using EpicGames.Core;
+using EpicGames.Horde.Storage;
+using EpicGames.Horde.Storage.Backends;
+using EpicGames.Horde.Storage.Nodes;
+using Horde.Server.Agents;
+using Horde.Server.Utilities;
+using Microsoft.Extensions.Logging;
+using Microsoft.IdentityModel.Tokens;
+
+namespace Horde.Server.Commands.Install
+{
+	using JsonObject = System.Text.Json.Nodes.JsonObject;
+
+	/// <summary>
+	/// Performs custom actions required by an MSI installation on Windows; creates a unique certificate, configures the agent to use it, and packages the agent into a bundle.
+	/// </summary>
+	[Command("setup", "Runs post-install setup actions to configure the agent bundle, etc...")]
+	public class SetupCommand : Command
+	{
+		[CommandLine("-Url=")]
+		string? ServerUrl { get; set; }
+
+		[CommandLine("-BaseDir")]
+		DirectoryReference BaseDir { get; set; } = Program.AppDir.ParentDirectory!;
+
+		/// <inheritdoc/>
+		public override async Task<int> ExecuteAsync(ILogger logger)
+		{
+			DirectoryReference serverDir = DirectoryReference.Combine(BaseDir, "Server");
+
+			// Update the agent to recognize the server certificate, and give it a custom token for being able to connect
+			DirectoryReference looseAgentDir = DirectoryReference.Combine(serverDir, "Tools", "horde-agent-loose");
+			FileReference agentConfigFile = FileReference.Combine(looseAgentDir, "appsettings.json");
+			{
+				JsonObject agentConfig = await ReadConfigAsync(agentConfigFile);
+
+				JsonObject hordeConfig = FindOrAddNode(agentConfig, "Horde", () => new JsonObject());
+				hordeConfig["Server"] = "Default";
+
+				JsonArray serverProfiles = FindOrAddNode(hordeConfig, "ServerProfiles", () => new JsonArray());
+
+				JsonObject serverProfile = FindOrAddElementByKey(serverProfiles, "Name", "Default");
+				serverProfile["Environment"] = "Prod";
+				serverProfile["Url"] = ServerUrl ?? "http://localhost:5000";
+
+				await SaveConfigAsync(agentConfigFile, agentConfig);
+			}
+
+			// Create the local agent bundle
+			DirectoryReference bundleDir = DirectoryReference.Combine(serverDir, "Tools", "horde-agent");
+			DirectoryReference.CreateDirectory(bundleDir);
+
+			RefName refName = new RefName("latest");
+
+			FileStorageClient client = new FileStorageClient(bundleDir, logger);
+			using (TreeWriter writer = new TreeWriter(client, refName))
+			{
+				DirectoryNode dirNode = new DirectoryNode();
+				await dirNode.CopyFromDirectoryAsync(looseAgentDir.ToDirectoryInfo(), new ChunkingOptions(), writer, null);
+				await writer.WriteAsync(refName, dirNode);
+			}
+
+			// Update the server config to include the bundled tool
+			FileReference serverConfigFile = FileReference.Combine(serverDir, "appsettings.json");
+			{
+				JsonObject serverConfig = await ReadConfigAsync(serverConfigFile);
+
+				JsonObject hordeConfig = FindOrAddNode(serverConfig, "Horde", () => new JsonObject());
+				JsonArray bundledTools = FindOrAddNode(hordeConfig, nameof(ServerSettings.BundledTools), () => new JsonArray());
+
+				JsonObject bundledTool = FindOrAddElementByKey(bundledTools, nameof(BundledToolConfig.Id), AgentExtensions.DefaultAgentSoftwareToolId.ToString());
+				bundledTool[nameof(BundledToolConfig.Name)] = "Horde Agent";
+				bundledTool[nameof(BundledToolConfig.Description)] = "Cross-platform build of the Horde Agent";
+				bundledTool[nameof(BundledToolConfig.RefName)] = refName.ToString();
+
+				FileVersionInfo versionInfo = FileVersionInfo.GetVersionInfo(FileReference.Combine(looseAgentDir, "hordeagent.exe").FullName);
+				if (!String.IsNullOrEmpty(versionInfo.ProductVersion))
+				{
+					bundledTool[nameof(BundledToolConfig.Version)] = versionInfo.ProductVersion;
+				}
+
+				await SaveConfigAsync(serverConfigFile, serverConfig);
+			}
+
+			return 0;
+		}
+
+		static T FindOrAddNode<T>(JsonObject obj, string name, Func<T> factory) where T : JsonNode
+		{
+			JsonNode? node = obj[name];
+			if (node != null)
+			{
+				if (node is T existingTypedNode)
+				{
+					return existingTypedNode;
+				}
+				else
+				{
+					obj.Remove(name);
+				}
+			}
+
+			T newTypedNode = factory();
+			obj.Add(name, newTypedNode);
+			return newTypedNode;
+		}
+
+		static JsonObject FindOrAddElementByKey(JsonArray array, string key, string name)
+		{
+			foreach (JsonNode? element in array)
+			{
+				if (element is JsonObject obj)
+				{
+					JsonNode? node = obj[key];
+					if (node != null && (string?)node.AsValue() == name)
+					{
+						return obj;
+					}
+				}
+			}
+
+			JsonObject newObj = new JsonObject();
+			newObj[key] = name;
+			array.Add(newObj);
+			return newObj;
+		}
+
+		static async Task<JsonObject> ReadConfigAsync(FileReference file)
+		{
+			byte[] data = await FileReference.ReadAllBytesAsync(file);
+			JsonObject? obj = JsonNode.Parse(data, new JsonNodeOptions { PropertyNameCaseInsensitive = true }, new JsonDocumentOptions { AllowTrailingCommas = true, CommentHandling = JsonCommentHandling.Skip }) as JsonObject;
+			return obj ?? new JsonObject();
+		}
+
+		static async Task SaveConfigAsync(FileReference file, JsonNode node)
+		{
+			ArrayBufferWriter<byte> buffer = new ArrayBufferWriter<byte>();
+			using (Utf8JsonWriter writer = new Utf8JsonWriter(buffer, new JsonWriterOptions { Indented = true }))
+			{
+				node.WriteTo(writer);
+			}
+			await FileReference.WriteAllBytesAsync(file, buffer.WrittenMemory.ToArray());
+		}
+	}
+}
