Horde: Use correct workspace info when initializing AutoSDK in WorkspaceExecutor

diff --git a/Engine/Source/Programs/Horde/Horde.Agent/Execution/WorkspaceExecutor.cs b/Engine/Source/Programs/Horde/Horde.Agent/Execution/WorkspaceExecutor.cs
--- a/Engine/Source/Programs/Horde/Horde.Agent/Execution/WorkspaceExecutor.cs
+++ b/Engine/Source/Programs/Horde/Horde.Agent/Execution/WorkspaceExecutor.cs
@@ -148,7 +148,7 @@ namespace Horde.Agent.Execution
 			IWorkspaceMaterializer? autoSdkMaterializer = null;
 			if (autoSdkWorkspaceInfo != null)
 			{
-				autoSdkMaterializer = _materializerFactory.CreateMaterializer(type, workspaceInfo, options, forAutoSdk: true);
+				autoSdkMaterializer = _materializerFactory.CreateMaterializer(type, autoSdkWorkspaceInfo, options, forAutoSdk: true);
 			}
 			
 			return new WorkspaceExecutor(options, workspaceMaterializer, autoSdkMaterializer, _loggerFactory.CreateLogger<WorkspaceExecutor>());
