Horde: Use IMemoryCache for graph cache

diff --git a/Engine/Source/Programs/Horde/Horde.Server/Jobs/Graphs/GraphCollection.cs b/Engine/Source/Programs/Horde/Horde.Server/Jobs/Graphs/GraphCollection.cs
--- a/Engine/Source/Programs/Horde/Horde.Server/Jobs/Graphs/GraphCollection.cs
+++ b/Engine/Source/Programs/Horde/Horde.Server/Jobs/Graphs/GraphCollection.cs
@@ -1,17 +1,15 @@
 // Copyright Epic Games, Inc. All Rights Reserved.
 
 using System;
-using System.Collections.Concurrent;
 using System.Collections.Generic;
 using System.Diagnostics;
 using System.Linq;
-using System.Threading;
 using System.Threading.Tasks;
 using EpicGames.Core;
 using Horde.Server.Jobs.Templates;
 using Horde.Server.Server;
-using Horde.Server.Utilities;
 using HordeCommon;
+using Microsoft.Extensions.Caching.Memory;
 using MongoDB.Bson;
 using MongoDB.Bson.Serialization.Attributes;
 using MongoDB.Bson.Serialization.Options;
@@ -22,7 +20,7 @@ namespace Horde.Server.Jobs.Graphs
 	/// <summary>
 	/// Collection of graph documents
 	/// </summary>
-	public class GraphCollection : IGraphCollection
+	public sealed class GraphCollection : IGraphCollection, IDisposable
 	{
 		/// <summary>
 		/// Represents a node in the graph
@@ -354,53 +352,6 @@ namespace Horde.Server.Jobs.Graphs
 			}
 		}
 
-		/// <summary>
-		/// Stores information about a cached graph
-		/// </summary>
-		class CachedGraph
-		{
-			/// <summary>
-			/// Time at which the graph was last accessed
-			/// </summary>
-			public long LastAccessTime => _lastAccessTimePrivate;
-
-			/// <summary>
-			/// Backing value for <see cref="LastAccessTime"/>
-			/// </summary>
-			private long _lastAccessTimePrivate;
-
-			/// <summary>
-			/// The graph instance
-			/// </summary>
-			public IGraph Graph { get; }
-
-			/// <summary>
-			/// Constructor
-			/// </summary>
-			/// <param name="graph">The graph to store</param>
-			public CachedGraph(IGraph graph)
-			{
-				_lastAccessTimePrivate = Stopwatch.GetTimestamp();
-				Graph = graph;
-			}
-
-			/// <summary>
-			/// Update the last access time
-			/// </summary>
-			public void Touch()
-			{
-				for (; ; )
-				{
-					long time = Stopwatch.GetTimestamp();
-					long lastAccessTimeCopy = _lastAccessTimePrivate;
-					if (time < lastAccessTimeCopy || Interlocked.CompareExchange(ref _lastAccessTimePrivate, time, lastAccessTimeCopy) == lastAccessTimeCopy)
-					{
-						break;
-					}
-				}
-			}
-		}
-
 		/// <summary>
 		/// The jobs collection
 		/// </summary>
@@ -412,9 +363,10 @@ namespace Horde.Server.Jobs.Graphs
 		const int MaxGraphs = 1000;
 
 		/// <summary>
-		/// Cache of recently accessed graphs
+		/// Cache for graphs
+		/// Use a non-shared cache to ensure enough space for graphs
 		/// </summary>
-		readonly ConcurrentDictionary<ContentHash, CachedGraph> _cachedGraphs = new ConcurrentDictionary<ContentHash, CachedGraph>();
+		private readonly IMemoryCache _memoryCache = new MemoryCache(new MemoryCacheOptions() { SizeLimit = MaxGraphs });
 
 		/// <summary>
 		/// Constructor
@@ -424,6 +376,12 @@ namespace Horde.Server.Jobs.Graphs
 		{
 			_graphs = mongoService.GetCollection<GraphDocument>("Graphs");
 		}
+		
+		/// <inheritdoc/>
+		public void Dispose()
+		{
+			_memoryCache.Dispose();
+		}
 
 		/// <summary>
 		/// Adds a new graph document
@@ -476,31 +434,17 @@ namespace Horde.Server.Jobs.Graphs
 				return GraphDocument.Empty;
 			}
 
-			// Try to read the graph from the cache
-			CachedGraph? cachedGraph;
-			if (_cachedGraphs.TryGetValue(hash, out cachedGraph))
+			if (_memoryCache.TryGetValue(hash, out IGraph cachedGraph))
 			{
-				// Update the last access time
-				cachedGraph.Touch();
+				return cachedGraph;
 			}
-			else
-			{
-				// Trim the cache
-				while (_cachedGraphs.Count > MaxGraphs)
-				{
-					ContentHash? removeHash = _cachedGraphs.OrderBy(x => x.Value.LastAccessTime).Select(x => x.Key).FirstOrDefault();
-					if (removeHash == null || removeHash == ContentHash.Empty)
-					{
-						break;
-					}
-					_cachedGraphs.TryRemove(removeHash, out _);
-				}
 
-				// Create the new entry
-				cachedGraph = new CachedGraph(await _graphs.Find<GraphDocument>(x => x.Id == hash).FirstAsync());
-				_cachedGraphs.TryAdd(hash, cachedGraph);
-			}
-			return cachedGraph.Graph;
+			GraphDocument graphDoc = await _graphs.Find<GraphDocument>(x => x.Id == hash).FirstAsync();
+			MemoryCacheEntryOptions cacheEntryOptions = new MemoryCacheEntryOptions()
+				.SetSlidingExpiration(TimeSpan.FromHours(24))
+				.SetSize(10 * 1024); // Size is required when "SizeLimit" is set on MemoryCache. Just provide a rough estimate.
+			_memoryCache.Set(graphDoc.Id, graphDoc, cacheEntryOptions);
+			return graphDoc;
 		}
 
 		/// <inheritdoc/>
