Horde: Use correct size for graph cache entries

Previous size set was incorrect and led to instant eviction.

diff --git a/Engine/Source/Programs/Horde/Horde.Server/Jobs/Graphs/GraphCollection.cs b/Engine/Source/Programs/Horde/Horde.Server/Jobs/Graphs/GraphCollection.cs
--- a/Engine/Source/Programs/Horde/Horde.Server/Jobs/Graphs/GraphCollection.cs
+++ b/Engine/Source/Programs/Horde/Horde.Server/Jobs/Graphs/GraphCollection.cs
@@ -434,17 +434,12 @@ namespace Horde.Server.Jobs.Graphs
 				return GraphDocument.Empty;
 			}
 
-			if (_memoryCache.TryGetValue(hash, out IGraph cachedGraph))
+			return await _memoryCache.GetOrCreateAsync(hash, cacheEntry =>
 			{
-				return cachedGraph;
-			}
-
-			GraphDocument graphDoc = await _graphs.Find<GraphDocument>(x => x.Id == hash).FirstAsync();
-			MemoryCacheEntryOptions cacheEntryOptions = new MemoryCacheEntryOptions()
-				.SetSlidingExpiration(TimeSpan.FromHours(24))
-				.SetSize(10 * 1024); // Size is required when "SizeLimit" is set on MemoryCache. Just provide a rough estimate.
-			_memoryCache.Set(graphDoc.Id, graphDoc, cacheEntryOptions);
-			return graphDoc;
+				cacheEntry.SlidingExpiration = TimeSpan.FromHours(24);
+				cacheEntry.Size = 1;
+				return _graphs.Find<GraphDocument>(x => x.Id == hash).FirstAsync();
+			});
 		}
 
 		/// <inheritdoc/>
