Horde: Support for bisecting error ranges to identify the first change that failed.

- State of the bisection is stored in a "bisect task" document. Bisect tasks are initialized with a job and failing step name.
- Every minute, BisectService queries all the active tasks and - if there are no current jobs running - starts a new job half way between the last successful change and the first failing change.
- The bisect task tracks the initial job being used for the search, and re-evaluates the search range with each tick. This allows externally generated state to be factored in to updates.

diff --git a/Engine/Source/Programs/Horde/Horde.Server.Tests/AsyncCachedValueTests.cs b/Engine/Source/Programs/Horde/Horde.Server.Tests/AsyncCachedValueTests.cs
new file mode 100644
--- /dev/null
+++ b/Engine/Source/Programs/Horde/Horde.Server.Tests/AsyncCachedValueTests.cs
@@ -0,0 +1,39 @@
+// Copyright Epic Games, Inc. All Rights Reserved.
+
+using System;
+using System.Collections.Generic;
+using System.Linq;
+using System.Text;
+using System.Threading;
+using System.Threading.Tasks;
+using Horde.Server.Utilities;
+using Microsoft.VisualStudio.TestTools.UnitTesting;
+
+namespace Horde.Server.Tests
+{
+	[TestClass]
+	public class AsyncCachedValueTests
+	{
+		[TestMethod]
+		public async Task BasicTest()
+		{
+			using CancellationTokenSource cts = new CancellationTokenSource();
+			for (int idx = 0; idx < 20; idx++)
+			{
+				AsyncCachedValue<int> value = new AsyncCachedValue<int>(() => Task.FromResult(GenerateValue()), TimeSpan.FromSeconds(1.0));
+
+				await value.GetAsync(cts.Token);
+				await value.GetAsync(cts.Token);
+				await Task.Delay(100);
+				await value.GetAsync(cts.Token);
+				await Task.Delay(1000);
+				await value.GetAsync(cts.Token);
+			}
+		}
+
+		static int GenerateValue()
+		{
+			return 7;
+		}
+	}
+}
diff --git a/Engine/Source/Programs/Horde/Horde.Server.Tests/BisectTests.cs b/Engine/Source/Programs/Horde/Horde.Server.Tests/BisectTests.cs
new file mode 100644
--- /dev/null
+++ b/Engine/Source/Programs/Horde/Horde.Server.Tests/BisectTests.cs
@@ -0,0 +1,154 @@
+// Copyright Epic Games, Inc. All Rights Reserved.
+
+using System;
+using System.Collections.Generic;
+using System.Linq;
+using System.Threading;
+using System.Threading.Tasks;
+using EpicGames.Core;
+using Horde.Server.Jobs;
+using Horde.Server.Jobs.Graphs;
+using Horde.Server.Jobs.Bisect;
+using Horde.Server.Jobs.Templates;
+using Horde.Server.Logs;
+using Horde.Server.Projects;
+using Horde.Server.Streams;
+using Horde.Server.Users;
+using HordeCommon;
+using Microsoft.Extensions.DependencyInjection;
+using Microsoft.VisualStudio.TestTools.UnitTesting;
+using Moq;
+
+namespace Horde.Server.Tests
+{
+	[TestClass]
+	public class BisectTests : TestSetup
+	{
+		StreamId StreamId { get; } = new StreamId("ue4-main");
+		TemplateId TemplateId { get; } = new TemplateId("test-build");
+
+		[TestMethod]
+		public async Task TestStates()
+		{
+			IUser user = await UserCollection.FindOrAddUserByLoginAsync("Bob");
+
+			PerforceService.AddChange(StreamId, 10, user, "", new[] { "Foo.cpp" });
+
+			PerforceService.AddChange(StreamId, 11, user, "", new[] { "Bar.cpp" });
+			PerforceService.AddChange(StreamId, 12, user, "", new[] { "Bar.cpp" });
+			PerforceService.AddChange(StreamId, 13, user, "", new[] { "Bar.cpp" });
+			PerforceService.AddChange(StreamId, 14, user, "", new[] { "Bar.cpp" });
+			PerforceService.AddChange(StreamId, 15, user, "", new[] { "Bar.cpp" });
+			PerforceService.AddChange(StreamId, 16, user, "", new[] { "Bar.cpp" });
+			PerforceService.AddChange(StreamId, 17, user, "", new[] { "Bar.cpp" });
+			PerforceService.AddChange(StreamId, 20, user, "", new[] { "Bar.cpp" });
+
+			ProjectConfig project = new ProjectConfig();
+			project.Id = new ProjectId("ue4");
+			project.Streams.Add(new StreamConfig { Id = StreamId, Templates = new List<TemplateRefConfig> { new TemplateRefConfig { Id = TemplateId } } });
+
+			UpdateConfig(x => x.Projects.Add(project));
+
+			// Create the graph
+			List<NewNode> nodes = new List<NewNode>();
+			nodes.Add(new NewNode("CompileEditor"));
+
+			List<NewGroup> groups = new List<NewGroup>();
+			groups.Add(new NewGroup("Foo", nodes));
+
+			Mock<ITemplate> templateMock = new Mock<ITemplate>(MockBehavior.Strict);
+			templateMock.SetupGet(x => x.InitialAgentType).Returns((string?)null);
+
+			IGraph graph = await GraphCollection.AddAsync(templateMock.Object, null);
+			graph = await GraphCollection.AppendAsync(graph, groups);
+
+			// Create two job runs, one which fails, one that succeeds.
+			CreateJobOptions options = new CreateJobOptions();
+			options.Arguments.Add($"{IJob.TargetArgumentPrefix}CompileEditor");
+
+			IJob succeededJob = await CreateJob(10, graph, JobStepOutcome.Success, options);
+			IJob failedJob = await CreateJob(20, graph, JobStepOutcome.Failure, options);
+
+			IBisectTaskCollection bisectTaskCollection = ServiceProvider.GetRequiredService<IBisectTaskCollection>();
+
+			IBisectTask bisectTask = await bisectTaskCollection.CreateAsync(failedJob, "CompileEditor", JobStepOutcome.Failure, UserId.Anonymous);
+			Assert.AreEqual(failedJob.TemplateId, bisectTask.TemplateId);
+			Assert.AreEqual(failedJob.StreamId, bisectTask.StreamId);
+			Assert.AreEqual("CompileEditor", bisectTask.NodeName);
+			Assert.AreEqual(JobStepOutcome.Failure, bisectTask.Outcome);
+			Assert.AreEqual(UserId.Anonymous, bisectTask.OwnerId);
+			Assert.AreEqual(failedJob.Change, bisectTask.InitialChange);
+			Assert.AreEqual(failedJob.Id, bisectTask.InitialJobId);
+			Assert.AreEqual(failedJob.Change, bisectTask.CurrentChange);
+			Assert.AreEqual(failedJob.Id, bisectTask.CurrentJobId);
+
+			List<IJob> jobs = await JobCollection.FindBisectTaskJobsAsync(bisectTask.Id, null).ToListAsync();
+			Assert.AreEqual(0, jobs.Count);
+
+			BisectService bisectService = ServiceProvider.GetRequiredService<BisectService>();
+			await bisectService.StartAsync(CancellationToken.None);
+
+			await Clock.AdvanceAsync(TimeSpan.FromMinutes(30));
+
+			bisectTask = Deref(await bisectTaskCollection.GetAsync(bisectTask.Id));
+			Assert.AreEqual(BisectTaskState.Running, bisectTask.State);
+			jobs = await JobCollection.FindBisectTaskJobsAsync(bisectTask.Id, running: true).ToListAsync();
+			Assert.AreEqual(1, jobs.Count);
+			Assert.AreEqual(14, jobs[0].Change);
+			await SetJobOutcome(jobs[0], graph, JobStepOutcome.Failure);
+
+			await Clock.AdvanceAsync(TimeSpan.FromMinutes(30));
+
+			bisectTask = Deref(await bisectTaskCollection.GetAsync(bisectTask.Id));
+			Assert.AreEqual(BisectTaskState.Running, bisectTask.State);
+			jobs = await JobCollection.FindBisectTaskJobsAsync(bisectTask.Id, running: true).ToListAsync();
+			Assert.AreEqual(1, jobs.Count);
+			Assert.AreEqual(12, jobs[0].Change);
+			await SetJobOutcome(jobs[0], graph, JobStepOutcome.Success);
+
+			await Clock.AdvanceAsync(TimeSpan.FromMinutes(30));
+
+			bisectTask = Deref(await bisectTaskCollection.GetAsync(bisectTask.Id));
+			Assert.AreEqual(BisectTaskState.Running, bisectTask.State);
+			jobs = await JobCollection.FindBisectTaskJobsAsync(bisectTask.Id, running: true).ToListAsync();
+			Assert.AreEqual(1, jobs.Count);
+			Assert.AreEqual(13, jobs[0].Change);
+			await SetJobOutcome(jobs[0], graph, JobStepOutcome.Success);
+
+			await Clock.AdvanceAsync(TimeSpan.FromMinutes(30));
+
+			bisectTask = Deref(await bisectTaskCollection.GetAsync(bisectTask.Id));
+			Assert.AreEqual(BisectTaskState.Succeeded, bisectTask.State);
+			Assert.AreEqual(20, bisectTask.InitialChange);
+			Assert.AreEqual(14, bisectTask.CurrentChange);
+			jobs = await JobCollection.FindBisectTaskJobsAsync(bisectTask.Id, running: true).ToListAsync();
+			Assert.AreEqual(0, jobs.Count);
+		}
+
+		async Task<IJob> CreateJob(int change, IGraph graph, JobStepOutcome outcome, CreateJobOptions options)
+		{
+			IJob job = await JobCollection.AddAsync(JobId.GenerateNewId(), StreamId, TemplateId, ContentHash.SHA1("hello"), graph, "Test job", change, change, options);
+			return await SetJobOutcome(job, graph, outcome);
+		}
+
+		async Task<IJob> SetJobOutcome(IJob job, IGraph graph, JobStepOutcome outcome)
+		{
+			job = Deref(await JobCollection.TryUpdateGraphAsync(job, graph));
+
+			job = Deref(await JobCollection.TryUpdateStepAsync(job, graph, job.Batches[0].Id, job.Batches[0].Steps[0].Id, JobStepState.Completed, JobStepOutcome.Success));
+			job = Deref(await JobCollection.TryUpdateBatchAsync(job, graph, job.Batches[0].Id, null, JobStepBatchState.Complete, null));
+
+			IJobStepBatch batch = job.Batches[^1];
+			IJobStep step = batch.Steps[^1];
+
+			job = Deref(await JobCollection.TryUpdateStepAsync(job, graph, batch.Id, step.Id, JobStepState.Completed, outcome));
+			job = Deref(await JobCollection.TryUpdateBatchAsync(job, graph, batch.Id, null, JobStepBatchState.Complete, null));
+
+			INodeGroup group = graph.Groups[batch.GroupIdx];
+			INode node = group.Nodes[step.NodeIdx];
+			await JobStepRefCollection.InsertOrReplaceAsync(new JobStepRefId(job.Id, batch.Id, step.Id), job.Name, node.Name, job.StreamId, job.TemplateId, job.Change, LogId.GenerateNewId(), null, null, outcome, false, null, null, 0.0f, 0.0f, DateTime.MinValue, DateTime.MinValue, DateTime.MinValue);
+
+			return job;
+		}
+	}
+}
diff --git a/Engine/Source/Programs/Horde/Horde.Server.Tests/TestSetup.cs b/Engine/Source/Programs/Horde/Horde.Server.Tests/TestSetup.cs
--- a/Engine/Source/Programs/Horde/Horde.Server.Tests/TestSetup.cs
+++ b/Engine/Source/Programs/Horde/Horde.Server.Tests/TestSetup.cs
@@ -65,6 +65,7 @@ using Horde.Server.Telemetry;
 using Horde.Server.Artifacts;
 using Horde.Server.Dashboard;
 using OpenTelemetry.Trace;
+using Horde.Server.Jobs.Bisect;
 
 namespace Horde.Server.Tests
 {
@@ -213,6 +214,7 @@ namespace Horde.Server.Tests
 			services.AddSingleton<ILogFileCollection, LogFileCollection>();
 			services.AddSingleton<INotificationTriggerCollection, NotificationTriggerCollection>();
 			services.AddSingleton<IPoolCollection, PoolCollection>();
+			services.AddSingleton<IBisectTaskCollection, BisectTaskCollection>();
 			services.AddSingleton<ISessionCollection, SessionCollection>();
 			services.AddSingleton<ISubscriptionCollection, SubscriptionCollection>();
 			services.AddSingleton<IStreamCollection, StreamCollection>();
@@ -256,6 +258,7 @@ namespace Horde.Server.Tests
 			services.AddSingleton<IPerforceService, PerforceServiceStub>();
 			services.AddSingleton<PerforceLoadBalancer>();
 			services.AddSingleton<PoolService>();
+			services.AddSingleton<BisectService>();
 			services.AddSingleton<RpcService>();
 			services.AddSingleton<ScheduleService>();
 			services.AddSingleton<DeviceService>();
diff --git a/Engine/Source/Programs/Horde/Horde.Server/Jobs/Bisect/BisectController.cs b/Engine/Source/Programs/Horde/Horde.Server/Jobs/Bisect/BisectController.cs
new file mode 100644
--- /dev/null
+++ b/Engine/Source/Programs/Horde/Horde.Server/Jobs/Bisect/BisectController.cs
@@ -0,0 +1,306 @@
+// Copyright Epic Games, Inc. All Rights Reserved.
+
+using System;
+using System.Collections.Generic;
+using System.Threading;
+using System.Threading.Tasks;
+using Horde.Server.Acls;
+using Horde.Server.Jobs.Graphs;
+using Horde.Server.Perforce;
+using Horde.Server.Server;
+using Horde.Server.Streams;
+using Horde.Server.Users;
+using Horde.Server.Utilities;
+using HordeCommon;
+using Microsoft.AspNetCore.Authorization;
+using Microsoft.AspNetCore.Components.Web;
+using Microsoft.AspNetCore.Mvc;
+using Microsoft.Extensions.Options;
+
+namespace Horde.Server.Jobs.Bisect
+{
+	/// <summary>
+	/// Request to create a new bisect task
+	/// </summary>
+	public class CreateBisectTaskRequest
+	{
+		/// <summary>
+		/// Job containing the node to check
+		/// </summary>
+		public JobId JobId { get; set; }
+
+		/// <summary>
+		/// Name of the node to query
+		/// </summary>
+		public string NodeName { get; set; } = String.Empty;
+
+		/// <summary>
+		/// Commit tag to filter possible changes against
+		/// </summary>
+		public List<CommitTag>? CommitTags { get; set; }
+
+		/// <summary>
+		/// Set of changes to ignore. Can be modified later through <see cref="UpdateBisectTaskRequest"/>.
+		/// </summary>
+		public List<int>? IgnoreChanges { get; set; }
+
+		/// <summary>
+		/// Set of jobs to ignore. Can be modified later through <see cref="UpdateBisectTaskRequest"/>.
+		/// </summary>
+		public List<JobId>? IgnoreJobs { get; set; }
+	}
+
+	/// <summary>
+	/// Response from creating a bisect task
+	/// </summary>
+	public class CreateBisectTaskResponse
+	{
+		/// <summary>
+		/// Identifier for the new bisect task
+		/// </summary>
+		public BisectTaskId BisectTaskId { get; }
+
+		internal CreateBisectTaskResponse(IBisectTask bisectTask) => BisectTaskId = bisectTask.Id;
+	}
+
+	/// <summary>
+	/// Information about a bisect task
+	/// </summary>
+	public class GetBisectTaskResponse
+	{
+		readonly IBisectTask _bisectTask;
+
+		/// <inheritdoc cref="IBisectTask.Id"/>
+		public BisectTaskId Id => _bisectTask.Id;
+
+		/// <inheritdoc cref="IBisectTask.State"/>
+		public BisectTaskState State => _bisectTask.State;
+
+		/// <inheritdoc cref="IBisectTask.OwnerId"/>
+		public GetThinUserInfoResponse Owner { get; }
+
+		/// <inheritdoc cref="IBisectTask.StreamId"/>
+		public StreamId StreamId => _bisectTask.StreamId;
+
+		/// <inheritdoc cref="IBisectTask.TemplateId"/>
+		public TemplateId TemplateId => _bisectTask.TemplateId;
+
+		/// <inheritdoc cref="IBisectTask.NodeName"/>
+		public string NodeName => _bisectTask.NodeName;
+
+		/// <inheritdoc cref="IBisectTask.Outcome"/>
+		public JobStepOutcome Outcome => _bisectTask.Outcome;
+
+		/// <inheritdoc cref="IBisectTask.InitialJobId"/>
+		public JobId InitialJobId => _bisectTask.InitialJobId;
+
+		/// <inheritdoc cref="IBisectTask.InitialChange"/>
+		public int InitialChange => _bisectTask.InitialChange;
+
+		/// <inheritdoc cref="IBisectTask.CurrentJobId"/>
+		public JobId CurrentJobId => _bisectTask.CurrentJobId;
+
+		/// <inheritdoc cref="IBisectTask.CurrentChange"/>
+		public int CurrentChange => _bisectTask.CurrentChange;
+
+		internal GetBisectTaskResponse(IBisectTask bisectTask, GetThinUserInfoResponse owner)
+		{
+			_bisectTask = bisectTask;
+			Owner = owner;
+		}
+	}
+
+	/// <summary>
+	/// Updates the state of a bisect task
+	/// </summary>
+	public class UpdateBisectTaskRequest
+	{
+		/// <summary>
+		/// Cancels the current task
+		/// </summary>
+		public bool? Cancel { get; set; }
+
+		/// <summary>
+		/// List of change numbers to include in the search. 
+		/// </summary>
+		public List<int> IncludeChanges { get; set; } = new List<int>();
+
+		/// <summary>
+		/// List of change numbers to exclude from the search.
+		/// </summary>
+		public List<int> ExcludeChanges { get; set; } = new List<int>();
+
+		/// <summary>
+		/// List of jobs to include in the search.
+		/// </summary>
+		public List<JobId> IncludeJobs { get; set; } = new List<JobId>();
+
+		/// <summary>
+		/// List of jobs to exclude from the search.
+		/// </summary>
+		public List<JobId> ExcludeJobs { get; set; } = new List<JobId>();
+	}
+
+
+	/// <summary>
+	/// Controller for the /api/v1/bisect endpoint
+	/// </summary>
+	[ApiController]
+	[Authorize]
+	[Route("[controller]")]
+	public class BisectTasksController : HordeControllerBase
+	{
+		readonly IBisectTaskCollection _bisectTaskCollection;
+		readonly IJobCollection _jobCollection;
+		readonly IGraphCollection _graphCollection;
+		readonly IUserCollection _userCollection;
+		readonly IOptionsSnapshot<GlobalConfig> _globalConfig;
+
+		/// <summary>
+		/// Constructor
+		/// </summary>
+		public BisectTasksController(IBisectTaskCollection bisectTaskCollection, IJobCollection jobCollection, IGraphCollection graphCollection, IUserCollection userCollection, IOptionsSnapshot<GlobalConfig> globalConfig)
+		{
+			_bisectTaskCollection = bisectTaskCollection;
+			_jobCollection = jobCollection;
+			_graphCollection = graphCollection;
+			_userCollection = userCollection;
+			_globalConfig = globalConfig;
+		}
+
+		/// <summary>
+		/// Creates a new bisect task
+		/// </summary>
+		/// <param name="create">Properties of the new bisect task</param>
+		/// <param name="cancellationToken">Cancellation token for the request</param>
+		/// <returns>Id of the new bisect task</returns>
+		[HttpPost]
+		[Route("/api/v1/bisect")]
+		public async Task<ActionResult<CreateBisectTaskResponse>> CreateAsync([FromBody] CreateBisectTaskRequest create, CancellationToken cancellationToken)
+		{
+			IJob? job = await _jobCollection.GetAsync(create.JobId);
+			if (job == null)
+			{
+				return NotFound(create.JobId);
+			}
+
+			StreamConfig? streamConfig;
+			if (!_globalConfig.Value.TryGetStream(job.StreamId, out streamConfig))
+			{
+				return NotFound(job.StreamId);
+			}
+			if (!streamConfig.Authorize(JobAclAction.ViewJob, User))
+			{
+				return Forbid(JobAclAction.ViewJob, create.JobId);
+			}
+			if (!streamConfig.Authorize(BisectTaskAclAction.CreateBisectTask, User))
+			{
+				return Forbid(BisectTaskAclAction.CreateBisectTask, streamConfig.Id);
+			}
+
+			IGraph graph = await _graphCollection.GetAsync(job.GraphHash);
+
+			NodeRef nodeRef;
+			if (!graph.TryFindNode(create.NodeName, out nodeRef))
+			{
+				return BadRequest($"Graph does not contain '{create.NodeName}'");
+			}
+
+			IJobStep? jobStep;
+			if (!job.TryGetStepForNode(nodeRef, out jobStep))
+			{
+				return BadRequest($"Job does not contain '{create.NodeName}'");
+			}
+			if (jobStep.IsPending())
+			{
+				return BadRequest("Step has not finished executing yet");
+			}
+			if (jobStep.Outcome == JobStepOutcome.Success)
+			{
+				return BadRequest("Step has not failed");
+			}
+
+			CreateBisectTaskOptions options = new CreateBisectTaskOptions();
+			options.CommitTags = create.CommitTags;
+			options.IgnoreChanges = create.IgnoreChanges;
+			options.IgnoreJobs = create.IgnoreJobs;
+
+			IBisectTask bisectTask = await _bisectTaskCollection.CreateAsync(job, create.NodeName, jobStep.Outcome, User.GetUserId() ?? UserId.Empty, options, cancellationToken);
+			return new CreateBisectTaskResponse(bisectTask);
+		}
+
+		/// <summary>
+		/// Creates a new bisect task
+		/// </summary>
+		/// <param name="bisectTaskId">Id of the bisect task to retrieve</param>
+		/// <param name="cancellationToken">Cancellation token for the request</param>
+		/// <returns>Id of the new bisect task</returns>
+		[HttpPost]
+		[Route("/api/v1/bisect/{id}")]
+		public async Task<ActionResult<GetBisectTaskResponse>> GetAsync([FromRoute] BisectTaskId bisectTaskId, CancellationToken cancellationToken)
+		{
+			IBisectTask? bisectTask = await _bisectTaskCollection.GetAsync(bisectTaskId, cancellationToken);
+			if (bisectTask == null)
+			{
+				return NotFound(bisectTaskId);
+			}
+
+			StreamConfig? streamConfig;
+			if (!_globalConfig.Value.TryGetStream(bisectTask.StreamId, out streamConfig))
+			{
+				return NotFound(bisectTask.StreamId);
+			}
+			if (!streamConfig.Authorize(BisectTaskAclAction.ViewBisectTask, User))
+			{
+				return Forbid(BisectTaskAclAction.ViewBisectTask, streamConfig.Id);
+			}
+
+			IUser? user = await _userCollection.GetUserAsync(bisectTask.OwnerId);
+			return new GetBisectTaskResponse(bisectTask, new GetThinUserInfoResponse(user));
+		}
+
+		/// <summary>
+		/// Creates a new bisect task
+		/// </summary>
+		/// <param name="bisectTaskId">Id of the bisect task to retrieve</param>
+		/// <param name="request">Updates for the task</param>
+		/// <param name="cancellationToken">Cancellation token for the request</param>
+		/// <returns>Id of the new bisect task</returns>
+		[HttpPatch]
+		[Route("/api/v1/bisect/{id}")]
+		public async Task<ActionResult> UpdateAsync([FromRoute] BisectTaskId bisectTaskId, [FromBody] UpdateBisectTaskRequest request, CancellationToken cancellationToken)
+		{
+			for(; ;)
+			{
+				IBisectTask? bisectTask = await _bisectTaskCollection.GetAsync(bisectTaskId, cancellationToken);
+				if (bisectTask == null)
+				{
+					return NotFound(bisectTaskId);
+				}
+
+				StreamConfig? streamConfig;
+				if (!_globalConfig.Value.TryGetStream(bisectTask.StreamId, out streamConfig))
+				{
+					return NotFound(bisectTask.StreamId);
+				}
+				if (!streamConfig.Authorize(BisectTaskAclAction.UpdateBisectTask, User))
+				{
+					return Forbid(BisectTaskAclAction.UpdateBisectTask, streamConfig.Id);
+				}
+
+				UpdateBisectTaskOptions options = new UpdateBisectTaskOptions();
+				options.State = (!request.Cancel.HasValue) ? null : request.Cancel.Value ? BisectTaskState.Cancelled : BisectTaskState.Running;
+				options.IncludeChanges = request.IncludeChanges;
+				options.ExcludeChanges = request.ExcludeChanges;
+				options.IncludeJobs = request.IncludeJobs;
+				options.ExcludeJobs = request.ExcludeJobs;
+
+				IBisectTask? updatedTask = await _bisectTaskCollection.TryUpdateAsync(bisectTask, options, cancellationToken);
+				if (updatedTask != null)
+				{
+					return Ok();
+				}
+			}
+		}
+	}
+}
diff --git a/Engine/Source/Programs/Horde/Horde.Server/Jobs/Bisect/BisectService.cs b/Engine/Source/Programs/Horde/Horde.Server/Jobs/Bisect/BisectService.cs
new file mode 100644
--- /dev/null
+++ b/Engine/Source/Programs/Horde/Horde.Server/Jobs/Bisect/BisectService.cs
@@ -0,0 +1,202 @@
+// Copyright Epic Games, Inc. All Rights Reserved.
+
+using System;
+using System.Collections.Generic;
+using System.Linq;
+using System.Threading;
+using System.Threading.Tasks;
+using EpicGames.Core;
+using Horde.Server.Jobs.Graphs;
+using Horde.Server.Jobs.Templates;
+using Horde.Server.Perforce;
+using Horde.Server.Server;
+using Horde.Server.Streams;
+using Horde.Server.Utilities;
+using HordeCommon;
+using Microsoft.Extensions.Hosting;
+using Microsoft.Extensions.Logging;
+using Microsoft.Extensions.Options;
+using MongoDB.Driver;
+
+namespace Horde.Server.Jobs.Bisect
+{
+	class BisectService : IHostedService
+	{
+		readonly IBisectTaskCollection _bisectTaskCollection;
+		readonly IJobCollection _jobCollection;
+		readonly IJobStepRefCollection _jobStepRefCollection;
+		readonly ITemplateCollection _templateCollection;
+		readonly IGraphCollection _graphCollection;
+		readonly ICommitService _commitService;
+		readonly IOptionsMonitor<GlobalConfig> _globalConfig;
+		readonly ILogger _logger;
+		readonly ITicker _ticker;
+
+		public BisectService(IBisectTaskCollection bisectTaskCollection, IJobCollection jobCollection, IJobStepRefCollection jobStepRefCollection, ITemplateCollection templateCollection, IGraphCollection graphCollection, ICommitService commitService, IOptionsMonitor<GlobalConfig> globalConfig, IClock clock, ILogger<BisectService> logger)
+		{
+			_bisectTaskCollection = bisectTaskCollection;
+			_jobCollection = jobCollection;
+			_jobStepRefCollection = jobStepRefCollection;
+			_templateCollection = templateCollection;
+			_graphCollection = graphCollection;
+			_commitService = commitService;
+			_globalConfig = globalConfig;
+			_logger = logger;
+			_ticker = clock.AddSharedTicker<BisectService>(TimeSpan.FromMinutes(1.0), TickAsync, logger);
+		}
+
+		/// <inheritdoc/>
+		public async Task StartAsync(CancellationToken cancellationToken)
+		{
+			await _ticker.StartAsync();
+		}
+
+		/// <inheritdoc/>
+		public async Task StopAsync(CancellationToken cancellationToken)
+		{
+			await _ticker.StopAsync();
+		}
+
+		async ValueTask TickAsync(CancellationToken cancellationToken)
+		{
+			GlobalConfig globalConfig = _globalConfig.CurrentValue;
+			await foreach (IBisectTask bisectTask in _bisectTaskCollection.FindActiveAsync(cancellationToken))
+			{
+				IBisectTask? currentBisectTask = bisectTask;
+				while (currentBisectTask != null)
+				{
+					BisectTaskState? newState = await TryUpdateBisectTaskAsync(currentBisectTask, globalConfig, cancellationToken);
+					if (newState != null)
+					{
+						UpdateBisectTaskOptions options = new UpdateBisectTaskOptions();
+						options.State = newState;
+
+						if (newState == currentBisectTask.State || await _bisectTaskCollection.TryUpdateAsync(currentBisectTask, options, cancellationToken) != null)
+						{
+							break;
+						}
+					}
+					currentBisectTask = await _bisectTaskCollection.GetAsync(bisectTask.Id, cancellationToken);
+				}
+			}
+		}
+
+		async Task<BisectTaskState?> TryUpdateBisectTaskAsync(IBisectTask bisectTask, GlobalConfig globalConfig, CancellationToken cancellationToken)
+		{
+			// Check if there are any jobs triggered by this bisect task still running.
+			IJob? existingJob = await _jobCollection.FindBisectTaskJobsAsync(bisectTask.Id, true, cancellationToken).FirstOrDefaultAsync(cancellationToken);
+			if (existingJob != null)
+			{
+				_logger.LogInformation("Bisect task {BisectTaskId} ({StreamId}:{TemplateId}:{NodeName}): Waiting for result of job {JobId}", bisectTask.Id, bisectTask.StreamId, bisectTask.TemplateId, bisectTask.NodeName, existingJob.Id);
+				return BisectTaskState.Running;
+			}
+
+			// Find the step that ran before this one, ignoring any excluded jobs
+			(IJobStepRef? currentJobStepRef, IJobStepRef? previousJobStepRef) = await GetBisectRangeAsync(bisectTask, cancellationToken);
+
+			// Check if the task needs to be updated with the current bisect state
+			if (currentJobStepRef != null && bisectTask.CurrentJobId != currentJobStepRef.Id.JobId)
+			{
+				_logger.LogInformation("Bisect task {BisectTaskId} ({StreamId}:{TemplateId}:{NodeName}): Changing current job from {PrevJobId} ({PrevChange}) -> {NextJobId} ({NextChange}).", bisectTask.Id, bisectTask.StreamId, bisectTask.TemplateId, bisectTask.NodeName, bisectTask.CurrentJobId, bisectTask.CurrentChange, currentJobStepRef.Id.JobId, currentJobStepRef.Change);
+				await _bisectTaskCollection.TryUpdateAsync(bisectTask, new UpdateBisectTaskOptions { CurrentJob = (currentJobStepRef.Id.JobId, currentJobStepRef.Change) }, cancellationToken);
+				return null;
+			}
+
+			// Check that there's a lower bound for the search
+			if (previousJobStepRef == null)
+			{
+				// No job before the current one
+				_logger.LogInformation("Bisect task {BisectTaskId} ({StreamId}:{TemplateId}:{NodeName}): No previous node; marking search as complete.", bisectTask.Id, bisectTask.StreamId, bisectTask.TemplateId, bisectTask.NodeName);
+				return BisectTaskState.MissingHistory;
+			}
+			else if (previousJobStepRef.Outcome == null)
+			{
+				// Job is still running
+				_logger.LogDebug("Bisect task {BisectTaskId} ({StreamId}:{TemplateId}:{NodeName}): Waiting on {JobStepRefId}", bisectTask.Id, bisectTask.StreamId, bisectTask.TemplateId, bisectTask.NodeName, previousJobStepRef.Id);
+				return BisectTaskState.Running;
+			}
+
+			// Find the next commit to test
+			StreamConfig? streamConfig;
+			if (!globalConfig.TryGetStream(bisectTask.StreamId, out streamConfig))
+			{
+				_logger.LogInformation("Bisect task {BisectTaskId} ({StreamId}:{TemplateId}:{NodeName}): Stream no longer exists.", bisectTask.Id, bisectTask.StreamId, bisectTask.TemplateId, bisectTask.NodeName);
+				return BisectTaskState.MissingStream;
+			}
+
+			ICommitCollection commitCollection = _commitService.GetCollection(streamConfig);
+			List<ICommit> commits = await commitCollection.FindAsync(previousJobStepRef.Change + 1, bisectTask.CurrentChange - 1, 100, bisectTask.CommitTags, cancellationToken).ToListAsync(cancellationToken);
+			commits.RemoveAll(x => bisectTask.IgnoreChanges.Contains(x.Number));
+
+			if (commits.Count == 0)
+			{
+				_logger.LogInformation("Bisect task {BisectTaskId} ({StreamId}:{TemplateId}:{NodeName}): Success - first failure was CL {Change}.", bisectTask.Id, bisectTask.StreamId, bisectTask.TemplateId, bisectTask.NodeName, bisectTask.CurrentChange);
+				return BisectTaskState.Succeeded;
+			}
+
+			ICommit nextCommit = commits[commits.Count / 2];
+			ICommit nextCodeCommit = await commitCollection.GetLastCodeChangeAsync(nextCommit.Number, cancellationToken) ?? nextCommit;
+
+			// Get the initial job
+			IJob? job = await _jobCollection.GetAsync(bisectTask.InitialJobId);
+			if (job == null)
+			{
+				_logger.LogInformation("Bisect task {BisectTaskId} ({StreamId}:{TemplateId}:{NodeName}): Missing job {JobId}.", bisectTask.Id, bisectTask.StreamId, bisectTask.TemplateId, bisectTask.NodeName, bisectTask.InitialJobId);
+				return BisectTaskState.MissingJob;
+			}
+
+			// Get the template configuration
+			TemplateRefConfig? templateRefConfig;
+			if (!streamConfig.TryGetTemplate(bisectTask.TemplateId, out templateRefConfig))
+			{
+				_logger.LogInformation("Bisect task {BisectTaskId} ({StreamId}:{TemplateId}:{NodeName}): Template {TemplateId} not found.", bisectTask.Id, bisectTask.StreamId, bisectTask.TemplateId, bisectTask.NodeName, bisectTask.TemplateId);
+				return BisectTaskState.MissingTemplate;
+			}
+
+			ITemplate template = await _templateCollection.GetOrAddAsync(templateRefConfig);
+			IGraph graph = await _graphCollection.AddAsync(template, streamConfig.InitialAgentType);
+
+			CreateJobOptions options = new CreateJobOptions(templateRefConfig);
+			options.StartedByBisectTaskId = bisectTask.Id;
+			options.Priority = Priority.BelowNormal;
+			options.UpdateIssues = template.UpdateIssues;
+			options.Arguments.AddRange(job.Arguments.Where(x => !x.StartsWith(IJob.TargetArgumentPrefix, StringComparison.OrdinalIgnoreCase)));
+			options.Arguments.Add($"{IJob.TargetArgumentPrefix}{bisectTask.NodeName}");
+
+			IJob nextJob = await _jobCollection.AddAsync(JobId.GenerateNewId(), bisectTask.StreamId, bisectTask.TemplateId, template.Hash, graph, $"{template.Name} (Bisect)", nextCommit.Number, nextCodeCommit.Number, options);
+			_logger.LogInformation("Bisect task {BisectTaskId} ({StreamId}:{TemplateId}:{NodeName}): {NumCommits} possible commits ({MinChange}..{MaxChange}). Started new job {JobId} at CL {Change}.", bisectTask.Id, bisectTask.StreamId, bisectTask.TemplateId, bisectTask.NodeName, commits.Count, commits[^1].Number, commits[0].Number, nextJob.Id, nextCommit.Number);
+
+			return BisectTaskState.Running;
+		}
+
+		async Task<(IJobStepRef? Current, IJobStepRef? Next)> GetBisectRangeAsync(IBisectTask bisectTask, CancellationToken cancellationToken)
+		{
+			// Query jobs before the start position
+			for (int maxJobCount = 20; ; maxJobCount += 20)
+			{
+				List<IJobStepRef> jobStepRefs = await _jobStepRefCollection.GetStepsForNodeAsync(bisectTask.StreamId, bisectTask.TemplateId, bisectTask.NodeName, bisectTask.InitialChange, true, maxJobCount, cancellationToken);
+
+				IJobStepRef? currentJobStepRef = null;
+				foreach(IJobStepRef jobStepRef in jobStepRefs.OrderByDescending(x => x.Change).ThenBy(x => x.Id))
+				{
+					if(!bisectTask.IgnoreJobs.Contains(jobStepRef.Id.JobId))
+					{
+						if (jobStepRef.Outcome == bisectTask.Outcome)
+						{
+							currentJobStepRef = jobStepRef;
+						}
+						else
+						{
+							return (currentJobStepRef, jobStepRef);
+						}
+					}
+				}
+
+				if (jobStepRefs.Count < maxJobCount)
+				{
+					return (null, null);
+				}
+			}
+		}
+	}
+}
diff --git a/Engine/Source/Programs/Horde/Horde.Server/Jobs/Bisect/BisectTaskAclAction.cs b/Engine/Source/Programs/Horde/Horde.Server/Jobs/Bisect/BisectTaskAclAction.cs
new file mode 100644
--- /dev/null
+++ b/Engine/Source/Programs/Horde/Horde.Server/Jobs/Bisect/BisectTaskAclAction.cs
@@ -0,0 +1,27 @@
+// Copyright Epic Games, Inc. All Rights Reserved.
+
+using Horde.Server.Acls;
+
+namespace Horde.Server.Jobs.Bisect
+{
+	/// <summary>
+	/// ACL actions valid for manipulating jobs
+	/// </summary>
+	public static class BisectTaskAclAction
+	{
+		/// <summary>
+		/// Ability to start new bisect tasks
+		/// </summary>
+		public static readonly AclAction CreateBisectTask = new AclAction("CreateBisectTask");
+
+		/// <summary>
+		/// Ability to update a bisect task
+		/// </summary>
+		public static readonly AclAction UpdateBisectTask = new AclAction("UpdateBisectTask");
+
+		/// <summary>
+		/// Ability to view a bisect task
+		/// </summary>
+		public static readonly AclAction ViewBisectTask = new AclAction("ViewBisectTask");
+	}
+}
diff --git a/Engine/Source/Programs/Horde/Horde.Server/Jobs/Bisect/BisectTaskCollection.cs b/Engine/Source/Programs/Horde/Horde.Server/Jobs/Bisect/BisectTaskCollection.cs
new file mode 100644
--- /dev/null
+++ b/Engine/Source/Programs/Horde/Horde.Server/Jobs/Bisect/BisectTaskCollection.cs
@@ -0,0 +1,174 @@
+// Copyright Epic Games, Inc. All Rights Reserved.
+
+using System;
+using System.Collections.Generic;
+using System.Runtime.CompilerServices;
+using System.Threading;
+using System.Threading.Tasks;
+using Horde.Server.Perforce;
+using Horde.Server.Server;
+using Horde.Server.Streams;
+using Horde.Server.Users;
+using Horde.Server.Utilities;
+using HordeCommon;
+using MongoDB.Bson.Serialization.Attributes;
+using MongoDB.Driver;
+
+namespace Horde.Server.Jobs.Bisect
+{
+	class BisectTaskCollection : IBisectTaskCollection
+	{
+		class BisectTaskDoc : IBisectTask
+		{
+			public BisectTaskId Id { get; set; }
+
+			[BsonElement("state")]
+			public BisectTaskState State { get; set; }
+
+			[BsonElement("oid")]
+			public UserId OwnerId { get; set; }
+
+			[BsonElement("sid")]
+			public StreamId StreamId { get; set; }
+
+			[BsonElement("tid")]
+			public TemplateId TemplateId { get; set; }
+
+			[BsonElement("step")]
+			public string NodeName { get; set; } = String.Empty;
+
+			[BsonElement("out")]
+			public JobStepOutcome Outcome { get; set; }
+
+			[BsonElement("job")]
+			public JobId InitialJobId { get; set; }
+
+			[BsonElement("chg")]
+			public int InitialChange { get; set; }
+
+			[BsonElement("curJob")]
+			public JobId CurrentJobId { get; set; }
+
+			[BsonElement("curChg")]
+			public int CurrentChange { get; set; }
+
+			[BsonElement("idx")]
+			public int UpdateIdx { get; set; }
+
+			[BsonElement("tags"), BsonIgnoreIfNull]
+			public List<CommitTag>? CommitTags { get; set; }
+			IReadOnlyList<CommitTag>? IBisectTask.CommitTags => CommitTags;
+
+			[BsonElement("nochg")]
+			public HashSet<int> IgnoreChanges { get; set; } = new HashSet<int>();
+			IReadOnlySet<int> IBisectTask.IgnoreChanges => IgnoreChanges;
+
+			[BsonElement("nojob")]
+			public HashSet<JobId> IgnoreJobs { get; set; } = new HashSet<JobId>();
+			IReadOnlySet<JobId> IBisectTask.IgnoreJobs => IgnoreJobs;
+		}
+
+		readonly IMongoCollection<BisectTaskDoc> _bisectTasks;
+
+		public BisectTaskCollection(MongoService mongoService)
+		{
+			List<MongoIndex<BisectTaskDoc>> indexes = new List<MongoIndex<BisectTaskDoc>>();
+			indexes.Add(keys => keys.Ascending(x => x.Id), Builders<BisectTaskDoc>.Filter.Eq(x => x.State, BisectTaskState.Running));
+
+			_bisectTasks = mongoService.GetCollection<BisectTaskDoc>("BisectTasks", indexes);
+		}
+
+		/// <inheritdoc/>
+		public async Task<IBisectTask> CreateAsync(IJob job, string nodeName, JobStepOutcome outcome, UserId ownerId, CreateBisectTaskOptions? options, CancellationToken cancellationToken = default)
+		{
+			BisectTaskDoc bisectTaskDoc = new BisectTaskDoc();
+			bisectTaskDoc.Id = BisectTaskId.GenerateNewId();
+			bisectTaskDoc.State = BisectTaskState.Running;
+			bisectTaskDoc.OwnerId = ownerId;
+			bisectTaskDoc.StreamId = job.StreamId;
+			bisectTaskDoc.TemplateId = job.TemplateId;
+			bisectTaskDoc.NodeName = nodeName;
+			bisectTaskDoc.Outcome = outcome;
+			bisectTaskDoc.InitialJobId = job.Id;
+			bisectTaskDoc.InitialChange = job.Change;
+			bisectTaskDoc.CurrentJobId = job.Id;
+			bisectTaskDoc.CurrentChange = job.Change;
+
+			if (options != null)
+			{
+				if (options.CommitTags != null && options.CommitTags.Count > 0)
+				{
+					bisectTaskDoc.CommitTags = new List<CommitTag>(options.CommitTags);
+				}
+				if (options.IgnoreChanges != null)
+				{
+					bisectTaskDoc.IgnoreChanges.UnionWith(options.IgnoreChanges);
+				}
+				if (options.IgnoreJobs != null)
+				{
+					bisectTaskDoc.IgnoreJobs.UnionWith(options.IgnoreJobs);
+				}
+			}
+
+			await _bisectTasks.InsertOneAsync(bisectTaskDoc, cancellationToken: cancellationToken);
+
+			return bisectTaskDoc;
+		}
+
+		/// <inheritdoc/>
+		public async IAsyncEnumerable<IBisectTask> FindActiveAsync([EnumeratorCancellation] CancellationToken cancellationToken = default)
+		{
+			using (IAsyncCursor<BisectTaskDoc> cursor = await _bisectTasks.Find(x => x.State == BisectTaskState.Running).SortBy(x => x.Id).ToCursorAsync(cancellationToken))
+			{
+				while (await cursor.MoveNextAsync(cancellationToken))
+				{
+					foreach (BisectTaskDoc bisectTaskDoc in cursor.Current)
+					{
+						yield return bisectTaskDoc;
+					}
+				}
+			}
+		}
+
+		/// <inheritdoc/>
+		public async Task<IBisectTask?> GetAsync(BisectTaskId bisectTaskId, CancellationToken cancellationToken = default)
+		{
+			return await _bisectTasks.Find(x => x.Id == bisectTaskId).FirstOrDefaultAsync(cancellationToken);
+		}
+
+		/// <inheritdoc/>
+		public async Task<IBisectTask?> TryUpdateAsync(IBisectTask bisectTask, UpdateBisectTaskOptions options, CancellationToken cancellationToken = default)
+		{
+			BisectTaskDoc bisectTaskDoc = (BisectTaskDoc)bisectTask;
+
+			UpdateDefinition<BisectTaskDoc> update = Builders<BisectTaskDoc>.Update.Inc(x => x.UpdateIdx, 1);
+			if (options.CurrentJob != null)
+			{
+				update = update.Set(x => x.CurrentJobId, options.CurrentJob.Value.JobId).Set(x => x.CurrentChange, options.CurrentJob.Value.Change);
+			}
+			if (options.State != null)
+			{
+				update = update.Set(x => x.State, options.State.Value);
+			}
+			if (options.IncludeChanges != null)
+			{
+				update = update.PullAll(x => x.IgnoreChanges, options.IncludeChanges); 
+			}
+			if (options.ExcludeChanges != null)
+			{
+				update = update.AddToSetEach(x => x.IgnoreChanges, options.ExcludeChanges);
+			}
+			if (options.IncludeJobs != null)
+			{
+				update = update.PullAll(x => x.IgnoreJobs, options.IncludeJobs);
+			}
+			if (options.ExcludeJobs != null)
+			{
+				update = update.AddToSetEach(x => x.IgnoreJobs, options.ExcludeJobs);
+			}
+
+			FilterDefinition<BisectTaskDoc> filter = Builders<BisectTaskDoc>.Filter.Expr(x => x.Id == bisectTaskDoc.Id && x.UpdateIdx == bisectTaskDoc.UpdateIdx);
+			return await _bisectTasks.FindOneAndUpdateAsync(filter, update, new FindOneAndUpdateOptions<BisectTaskDoc> { ReturnDocument = ReturnDocument.After }, cancellationToken);
+		}
+	}
+}
diff --git a/Engine/Source/Programs/Horde/Horde.Server/Jobs/Bisect/BisectTaskId.cs b/Engine/Source/Programs/Horde/Horde.Server/Jobs/Bisect/BisectTaskId.cs
new file mode 100644
--- /dev/null
+++ b/Engine/Source/Programs/Horde/Horde.Server/Jobs/Bisect/BisectTaskId.cs
@@ -0,0 +1,42 @@
+// Copyright Epic Games, Inc. All Rights Reserved.
+
+using System.ComponentModel;
+using EpicGames.Core;
+using Horde.Server.Utilities;
+using MongoDB.Bson;
+
+namespace Horde.Server.Jobs.Bisect
+{
+	/// <summary>
+	/// Identifier for a job
+	/// </summary>
+	/// <param name="Id">Id to construct from</param>
+	[JsonSchemaString]
+	[TypeConverter(typeof(ObjectIdTypeConverter<BisectTaskId, BisectTaskIdConverter>))]
+	[ObjectIdConverter(typeof(BisectTaskIdConverter))]
+	public record struct BisectTaskId(ObjectId Id)
+	{
+		/// <summary>
+		/// Creates a new <see cref="BisectTaskId"/>
+		/// </summary>
+		public static BisectTaskId GenerateNewId() => new BisectTaskId(ObjectId.GenerateNewId());
+
+		/// <inheritdoc cref="ObjectId.Parse(System.String)"/>
+		public static BisectTaskId Parse(string text) => new BisectTaskId(ObjectId.Parse(text));
+
+		/// <inheritdoc/>
+		public override string ToString() => Id.ToString();
+	}
+
+	/// <summary>
+	/// Converter to and from <see cref="ObjectId"/> instances.
+	/// </summary>
+	class BisectTaskIdConverter : ObjectIdConverter<BisectTaskId>
+	{
+		/// <inheritdoc/>
+		public override BisectTaskId FromObjectId(ObjectId id) => new BisectTaskId(id);
+
+		/// <inheritdoc/>
+		public override ObjectId ToObjectId(BisectTaskId value) => value.Id;
+	}
+}
diff --git a/Engine/Source/Programs/Horde/Horde.Server/Jobs/Bisect/IBisectTask.cs b/Engine/Source/Programs/Horde/Horde.Server/Jobs/Bisect/IBisectTask.cs
new file mode 100644
--- /dev/null
+++ b/Engine/Source/Programs/Horde/Horde.Server/Jobs/Bisect/IBisectTask.cs
@@ -0,0 +1,127 @@
+// Copyright Epic Games, Inc. All Rights Reserved.
+
+using System.Collections.Generic;
+using Horde.Server.Perforce;
+using Horde.Server.Streams;
+using Horde.Server.Users;
+using HordeCommon;
+
+namespace Horde.Server.Jobs.Bisect
+{
+	/// <summary>
+	/// State of a bisect task
+	/// </summary>
+	public enum BisectTaskState
+	{
+		/// <summary>
+		/// Currently running
+		/// </summary>
+		Running,
+
+		/// <summary>
+		/// Cancelled by a user
+		/// </summary>
+		Cancelled,
+
+		/// <summary>
+		/// Finished running. The first job/change identifies the first failure.
+		/// </summary>
+		Succeeded,
+
+		/// <summary>
+		/// Task failed due to not having a job before the first failure.
+		/// </summary>
+		MissingHistory,
+
+		/// <summary>
+		/// Task failed due to the stream no longer existing.
+		/// </summary>
+		MissingStream,
+
+		/// <summary>
+		/// Task failed due to the first job no longer existing.
+		/// </summary>
+		MissingJob,
+
+		/// <summary>
+		/// Task failed due to template no longer existing.
+		/// </summary>
+		MissingTemplate,
+	}
+
+	/// <summary>
+	/// State of a search for a build breakage
+	/// </summary>
+	public interface IBisectTask
+	{
+		/// <summary>
+		/// Identifier for this task
+		/// </summary>
+		public BisectTaskId Id { get; }
+
+		/// <summary>
+		/// Current task state
+		/// </summary>
+		public BisectTaskState State { get; }
+
+		/// <summary>
+		/// User that initiated the search
+		/// </summary>
+		public UserId OwnerId { get; }
+
+		/// <summary>
+		/// Stream being searched
+		/// </summary>
+		public StreamId StreamId { get; }
+
+		/// <summary>
+		/// Template within the stream to execute
+		/// </summary>
+		public TemplateId TemplateId { get; }
+
+		/// <summary>
+		/// Name of the step to search for
+		/// </summary>
+		public string NodeName { get; }
+
+		/// <summary>
+		/// Outcome to search for
+		/// </summary>
+		public JobStepOutcome Outcome { get; }
+
+		/// <summary>
+		/// Starting job id for the bisection
+		/// </summary>
+		public JobId InitialJobId { get; }
+
+		/// <summary>
+		/// Starting change for the bisection
+		/// </summary>
+		public int InitialChange { get; }
+
+		/// <summary>
+		/// First known job id that is broken
+		/// </summary>
+		public JobId CurrentJobId { get; }
+
+		/// <summary>
+		/// Changelist number of the first broken job id
+		/// </summary>
+		public int CurrentChange { get; }
+
+		/// <summary>
+		/// Tags to filter changes to consider for this task
+		/// </summary>
+		public IReadOnlyList<CommitTag>? CommitTags { get; }
+
+		/// <summary>
+		/// Set of changes to ignore while doing the bisection
+		/// </summary>
+		public IReadOnlySet<int> IgnoreChanges { get; }
+
+		/// <summary>
+		/// Set of jobs to ignore while doing the bisection
+		/// </summary>
+		public IReadOnlySet<JobId> IgnoreJobs { get; }
+	}
+}
\ No newline at end of file
diff --git a/Engine/Source/Programs/Horde/Horde.Server/Jobs/Bisect/IBisectTaskCollection.cs b/Engine/Source/Programs/Horde/Horde.Server/Jobs/Bisect/IBisectTaskCollection.cs
new file mode 100644
--- /dev/null
+++ b/Engine/Source/Programs/Horde/Horde.Server/Jobs/Bisect/IBisectTaskCollection.cs
@@ -0,0 +1,108 @@
+// Copyright Epic Games, Inc. All Rights Reserved.
+
+using System.Collections.Generic;
+using System.Threading;
+using System.Threading.Tasks;
+using Horde.Server.Perforce;
+using Horde.Server.Users;
+using HordeCommon;
+
+namespace Horde.Server.Jobs.Bisect
+{
+	/// <summary>
+	/// Manages a collection of bisect tasks
+	/// </summary>
+	public interface IBisectTaskCollection
+	{
+		/// <summary>
+		/// Creates a new bisect task
+		/// </summary>
+		/// <param name="job">First job to search back from</param>
+		/// <param name="nodeName">Name of the node to search for</param>
+		/// <param name="outcome">Outcome of the step to search for</param>
+		/// <param name="ownerId">User that initiated the search</param>
+		/// <param name="options">Options for the bisection</param>
+		/// <param name="cancellationToken">Cancellation token for the operation</param>
+		/// <returns>The new bisect task instance</returns>
+		public Task<IBisectTask> CreateAsync(IJob job, string nodeName, JobStepOutcome outcome, UserId ownerId, CreateBisectTaskOptions? options = null, CancellationToken cancellationToken = default);
+
+		/// <summary>
+		/// Finds all the active bisect tasks
+		/// </summary>
+		/// <param name="cancellationToken">Cancellation token for the operation</param>
+		/// <returns>List of bisect tasks</returns>
+		public IAsyncEnumerable<IBisectTask> FindActiveAsync(CancellationToken cancellationToken = default);
+
+		/// <summary>
+		/// Gets a bisect task by id
+		/// </summary>
+		/// <param name="bisectTaskId">Id of the bisect task</param>
+		/// <param name="cancellationToken">Cancellation token for the operation</param>
+		public Task<IBisectTask?> GetAsync(BisectTaskId bisectTaskId, CancellationToken cancellationToken = default);
+
+		/// <summary>
+		/// Updates the state of a bisect task
+		/// </summary>
+		/// <param name="bisectTask">Task to update</param>
+		/// <param name="options">Settings for updating the task</param>
+		/// <param name="cancellationToken">Cancellation token for the operation</param>
+		public Task<IBisectTask?> TryUpdateAsync(IBisectTask bisectTask, UpdateBisectTaskOptions options, CancellationToken cancellationToken = default);
+	}
+
+	/// <summary>
+	/// Settings for creating a new bisect task
+	/// </summary>
+	public class CreateBisectTaskOptions
+	{
+		/// <summary>
+		/// Filter the changes considered for the bisection
+		/// </summary>
+		public IReadOnlyList<CommitTag>? CommitTags { get; set; }
+
+		/// <summary>
+		/// Changes to exclude from the bisection
+		/// </summary>
+		public IReadOnlyList<int>? IgnoreChanges { get; set; }
+
+		/// <summary>
+		/// Jobs to exclude from the bisection
+		/// </summary>
+		public IReadOnlyList<JobId>? IgnoreJobs { get; set; }
+	}
+
+	/// <summary>
+	/// Settings for updating a bisect task
+	/// </summary>
+	public class UpdateBisectTaskOptions
+	{
+		/// <summary>
+		/// First job to bisect from
+		/// </summary>
+		public (JobId JobId, int Change)? CurrentJob { get; set; }
+
+		/// <summary>
+		/// New state for the task
+		/// </summary>
+		public BisectTaskState? State { get; set; }
+
+		/// <summary>
+		/// Changes to include in the bisection
+		/// </summary>
+		public IReadOnlyList<int>? IncludeChanges { get; set; }
+
+		/// <summary>
+		/// Changes to exclude from the bisection
+		/// </summary>
+		public IReadOnlyList<int>? ExcludeChanges { get; set; }
+
+		/// <summary>
+		/// Jobs to include in the bisection
+		/// </summary>
+		public IReadOnlyList<JobId>? IncludeJobs { get; set; }
+
+		/// <summary>
+		/// Jobs to exclude from the bisection
+		/// </summary>
+		public IReadOnlyList<JobId>? ExcludeJobs { get; set; }
+	}
+}
diff --git a/Engine/Source/Programs/Horde/Horde.Server/Jobs/IJob.cs b/Engine/Source/Programs/Horde/Horde.Server/Jobs/IJob.cs
--- a/Engine/Source/Programs/Horde/Horde.Server/Jobs/IJob.cs
+++ b/Engine/Source/Programs/Horde/Horde.Server/Jobs/IJob.cs
@@ -24,6 +24,7 @@ using HordeCommon;
 using HordeCommon.Rpc.Tasks;
 using MongoDB.Bson;
 using MongoDB.Driver;
+using Horde.Server.Jobs.Bisect;
 
 namespace Horde.Server.Jobs
 {
@@ -129,7 +130,7 @@ namespace Horde.Server.Jobs
 		/// If a retry is requested, stores the name of the user that requested it
 		/// </summary>
 		public UserId? RetriedByUserId { get; }
-		
+
 		/// <summary>
 		/// Signal if a step should be aborted
 		/// </summary>
@@ -587,6 +588,11 @@ namespace Horde.Server.Jobs
 		/// </summary>
 		public UserId? AbortedByUserId { get; }
 
+		/// <summary>
+		/// Identifier of the bisect task that started this job
+		/// </summary>
+		public BisectTaskId? StartedByBisectTaskId { get; }
+
 		/// <summary>
 		/// Name of the job.
 		/// </summary>
diff --git a/Engine/Source/Programs/Horde/Horde.Server/Jobs/IJobCollection.cs b/Engine/Source/Programs/Horde/Horde.Server/Jobs/IJobCollection.cs
--- a/Engine/Source/Programs/Horde/Horde.Server/Jobs/IJobCollection.cs
+++ b/Engine/Source/Programs/Horde/Horde.Server/Jobs/IJobCollection.cs
@@ -12,11 +12,11 @@
 using Horde.Server.Logs;
 using Horde.Server.Streams;
 using Horde.Server.Users;
-using Horde.Server.Utilities;
-using Horde.Common;
 using HordeCommon;
 using HordeCommon.Rpc.Tasks;
 using MongoDB.Bson;
+using Horde.Server.Jobs.Bisect;
+using System.Threading;
 
 namespace Horde.Server.Jobs
 {
@@ -25,49 +25,34 @@
 	/// </summary>
 	public class CreateJobOptions
 	{
-		/// <summary>
-		/// Optional changelist to preflight
-		/// </summary>
+		/// <inheritdoc cref="IJob.PreflightChange"/>
 		public int? PreflightChange { get; set; }
 
-		/// <summary>
-		/// Optional cloned preflight changelist
-		/// </summary>
+		/// <inheritdoc cref="IJob.ClonedPreflightChange"/>
 		public int? ClonedPreflightChange { get; set; }
 
-		/// <summary>
-		/// Description for the change being preflighted
-		/// </summary>
+		/// <inheritdoc cref="IJob.PreflightDescription"/>
 		public string? PreflightDescription { get; set; }
 
-		/// <summary>
-		/// User that started the job
-		/// </summary>
+		/// <inheritdoc cref="IJob.StartedByUserId"/>
 		public UserId? StartedByUserId { get; set; }
 
-		/// <summary>
-		/// Priority of the job
-		/// </summary>
+		/// <inheritdoc cref="IJob.StartedByBisectTaskId"/>
+		public BisectTaskId? StartedByBisectTaskId { get; set; }
+
+		/// <inheritdoc cref="IJob.Priority"/>
 		public Priority? Priority { get; set; }
 
-		/// <summary>
-		/// Whether to automatically submit the preflighted change on completion
-		/// </summary>
+		/// <inheritdoc cref="IJob.AutoSubmit"/>
 		public bool? AutoSubmit { get; set; }
 
-		/// <summary>
-		/// Whether to update issues based on the outcome of this job
-		/// </summary>
+		/// <inheritdoc cref="IJob.UpdateIssues"/>
 		public bool? UpdateIssues { get; set; }
 
-		/// <summary>
-		/// Whether to promote issues based on the outcome of this job by default
-		/// </summary>
+		/// <inheritdoc cref="IJob.PromoteIssuesByDefault"/>
 		public bool? PromoteIssuesByDefault { get; set; }
 
-		/// <summary>
-		/// Settings for agents executing the job
-		/// </summary>
+		/// <inheritdoc cref="IJob.JobOptions"/>
 		public JobOptions? JobOptions { get; set; }
 
 		/// <summary>
@@ -75,34 +60,22 @@
 		/// </summary>
 		public List<ChainedJobTemplateConfig> JobTriggers { get; } = new List<ChainedJobTemplateConfig>();
 
-		/// <summary>
-		/// Whether to show badges in UGS for this job
-		/// </summary>
+		/// <inheritdoc cref="IJob.ShowUgsBadges"/>
 		public bool ShowUgsBadges { get; set; }
 
-		/// <summary>
-		/// Whether to show alerts in UGS for this job
-		/// </summary>
+		/// <inheritdoc cref="IJob.ShowUgsAlerts"/>
 		public bool ShowUgsAlerts { get; set; }
 
-		/// <summary>
-		/// Notification channel for this job
-		/// </summary>
+		/// <inheritdoc cref="IJob.NotificationChannel"/>
 		public string? NotificationChannel { get; set; }
 
-		/// <summary>
-		/// Notification channel filter for this job
-		/// </summary>
+		/// <inheritdoc cref="IJob.NotificationChannelFilter"/>
 		public string? NotificationChannelFilter { get; set; }
 
-		/// <summary>
-		/// Arguments for the job
-		/// </summary>
+		/// <inheritdoc cref="IJob.Arguments"/>
 		public List<string> Arguments { get; } = new List<string>();
 
-		/// <summary>
-		/// Environment variables for the job
-		/// </summary>
+		/// <inheritdoc cref="IJob.Environment"/>
 		public Dictionary<string, string> Environment { get; } = new Dictionary<string, string>(StringComparer.OrdinalIgnoreCase);
 
 		/// <summary>
@@ -196,5 +169,14 @@
 		Task<List<IJob>> FindAsync(JobId[]? jobIds = null, StreamId? streamId = null, string? name = null, TemplateId[]? templates = null, int? minChange = null, int? maxChange = null, int? preflightChange = null, bool? preflightOnly = null, UserId? preflightStartedByUser = null, UserId? startedByUser = null, DateTimeOffset? minCreateTime = null, DateTimeOffset? maxCreateTime = null, DateTimeOffset? modifiedBefore = null, DateTimeOffset? modifiedAfter = null, int? index = null, int? count = null, bool consistentRead = true, string? indexHint = null, bool? excludeUserJobs = null);
 
 		/// <summary>
+		/// Searches for jobs matching the given criteria
+		/// </summary>
+		/// <param name="bisectTaskId">The bisect task to find jobs for</param>
+		/// <param name="running">Whether to filter by running jobs</param>
+		/// <param name="cancellationToken">Cancellation token for the operation</param>
+		/// <returns>List of jobs matching the given criteria</returns>
+		IAsyncEnumerable<IJob> FindBisectTaskJobsAsync(BisectTaskId bisectTaskId, bool? running, CancellationToken cancellationToken = default);
+
+		/// <summary>
 		/// Searches for jobs in a specified stream and templates
 		/// </summary>
diff --git a/Engine/Source/Programs/Horde/Horde.Server/Jobs/IJobStepRefCollection.cs b/Engine/Source/Programs/Horde/Horde.Server/Jobs/IJobStepRefCollection.cs
--- a/Engine/Source/Programs/Horde/Horde.Server/Jobs/IJobStepRefCollection.cs
+++ b/Engine/Source/Programs/Horde/Horde.Server/Jobs/IJobStepRefCollection.cs
@@ -2,6 +2,7 @@
 
 using System;
 using System.Collections.Generic;
+using System.Threading;
 using System.Threading.Tasks;
 using Horde.Server.Agents;
 using Horde.Server.Agents.Pools;
@@ -68,9 +69,10 @@ namespace Horde.Server.Jobs
 		/// <param name="nodeName">Name of the node</param>
 		/// <param name="change">The current change</param>
 		/// <param name="includeFailed">Whether to include failed nodes</param>
-		/// <param name="count">Number of results to return</param>
+		/// <param name="maxCount">Number of results to return</param>
+		/// <param name="cancellationToken">Cancellation token for the operation</param>
 		/// <returns>List of step references</returns>
-		Task<List<IJobStepRef>> GetStepsForNodeAsync(StreamId streamId, TemplateId templateId, string nodeName, int? change, bool includeFailed, int count);
+		Task<List<IJobStepRef>> GetStepsForNodeAsync(StreamId streamId, TemplateId templateId, string nodeName, int? change, bool includeFailed, int maxCount, CancellationToken cancellationToken = default);
 
 		/// <summary>
 		/// Gets the previous job that ran a given step
@@ -81,8 +83,9 @@ namespace Horde.Server.Jobs
 		/// <param name="change">The current changelist number</param>
 		/// <param name="outcome">The outcome to filter by or include all outcomes if null</param>
 		/// <param name="updateIssues">If true, constrain to steps which update issues</param>		 
+		/// <param name="excludeJobIds">Jobs to exclude from the search</param>
 		/// <returns>The previous job, or null.</returns>
-		Task<IJobStepRef?> GetPrevStepForNodeAsync(StreamId streamId, TemplateId templateId, string nodeName, int change, JobStepOutcome? outcome = null, bool? updateIssues = null);
+		Task<IJobStepRef?> GetPrevStepForNodeAsync(StreamId streamId, TemplateId templateId, string nodeName, int change, JobStepOutcome? outcome = null, bool? updateIssues = null, IEnumerable<JobId>? excludeJobIds = null);
 
 		/// <summary>
 		/// Gets the next job that ran a given step
diff --git a/Engine/Source/Programs/Horde/Horde.Server/Jobs/JobCollection.cs b/Engine/Source/Programs/Horde/Horde.Server/Jobs/JobCollection.cs
--- a/Engine/Source/Programs/Horde/Horde.Server/Jobs/JobCollection.cs
+++ b/Engine/Source/Programs/Horde/Horde.Server/Jobs/JobCollection.cs
@@ -6,7 +6,9 @@ using System.Diagnostics;
 using System.Diagnostics.CodeAnalysis;
 using System.IO;
 using System.Linq;
+using System.Runtime.CompilerServices;
 using System.Text;
+using System.Threading;
 using System.Threading.Tasks;
 using EpicGames.Core;
 using Horde.Server.Acls;
@@ -15,6 +17,7 @@ using Horde.Server.Agents.Leases;
 using Horde.Server.Agents.Pools;
 using Horde.Server.Agents.Sessions;
 using Horde.Server.Jobs.Graphs;
+using Horde.Server.Jobs.Bisect;
 using Horde.Server.Logs;
 using Horde.Server.Server;
 using Horde.Server.Streams;
@@ -239,6 +242,9 @@ namespace Horde.Server.Jobs
 			[BsonIgnoreIfNull]
 			public UserId? StartedByUserId { get; set; }
 
+			[BsonIgnoreIfNull]
+			public BisectTaskId? StartedByBisectTaskId { get; set; }
+
 			[BsonIgnoreIfNull, BsonElement("StartedByUser")]
 			public string? StartedByUserDeprecated { get; set; }
 
@@ -341,6 +347,7 @@ namespace Horde.Server.Jobs
 				ClonedPreflightChange = options.ClonedPreflightChange ?? 0;
 				PreflightDescription = options.PreflightDescription;
 				StartedByUserId = options.StartedByUserId;
+				StartedByBisectTaskId = options.StartedByBisectTaskId;
 				Priority = options.Priority ?? HordeCommon.Priority.Normal;
 				AutoSubmit = options.AutoSubmit ?? false;
 				UpdateIssues = options.UpdateIssues ?? (options.StartedByUserId == null && ( options.PreflightChange == 0 || options.PreflightChange == null));
@@ -575,6 +582,29 @@ namespace Horde.Server.Jobs
 			return results.ConvertAll<JobDocument, IJob>(x => x);
 		}
 
+		/// <inheritdoc/>
+		public async IAsyncEnumerable<IJob> FindBisectTaskJobsAsync(BisectTaskId bisectTaskId, bool? running, [EnumeratorCancellation] CancellationToken cancellationToken)
+		{
+			FilterDefinition<JobDocument> filter = Builders<JobDocument>.Filter.Eq(x => x.StartedByBisectTaskId, bisectTaskId);
+			if (running.HasValue)
+			{
+				if (running.Value)
+				{
+					filter &= Builders<JobDocument>.Filter.Gt(x => x.SchedulePriority, 0);
+				}
+				else
+				{
+					filter &= Builders<JobDocument>.Filter.Lte(x => x.SchedulePriority, 0);
+				}
+			}
+
+			await foreach (JobDocument jobDoc in _jobs.Find(filter).ToAsyncEnumerable(cancellationToken))
+			{
+				await PostLoadAsync(jobDoc);
+				yield return jobDoc;
+			}
+		}
+
 		/// <inheritdoc/>
 		public async Task<List<IJob>> FindLatestByStreamWithTemplatesAsync(StreamId streamId, TemplateId[] templates, UserId? preflightStartedByUser, DateTimeOffset? maxCreateTime, DateTimeOffset? modifiedAfter, int? index, int? count, bool consistentRead)
 		{
diff --git a/Engine/Source/Programs/Horde/Horde.Server/Jobs/JobStepRefCollection.cs b/Engine/Source/Programs/Horde/Horde.Server/Jobs/JobStepRefCollection.cs
--- a/Engine/Source/Programs/Horde/Horde.Server/Jobs/JobStepRefCollection.cs
+++ b/Engine/Source/Programs/Horde/Horde.Server/Jobs/JobStepRefCollection.cs
@@ -3,6 +3,7 @@
 using System;
 using System.Collections.Generic;
 using System.Linq;
+using System.Threading;
 using System.Threading.Tasks;
 using Horde.Server.Agents;
 using Horde.Server.Agents.Pools;
@@ -174,7 +175,7 @@ namespace Horde.Server.Jobs
 		}
 
 		/// <inheritdoc/>
-		public async Task<List<IJobStepRef>> GetStepsForNodeAsync(StreamId streamId, TemplateId templateId, string nodeName, int? change, bool includeFailed, int count)
+		public async Task<List<IJobStepRef>> GetStepsForNodeAsync(StreamId streamId, TemplateId templateId, string nodeName, int? change, bool includeFailed, int maxCount, CancellationToken cancellationToken)
 		{
 			// Find all the steps matching the given criteria
 			FilterDefinitionBuilder<JobStepRef> filterBuilder = Builders<JobStepRef>.Filter;
@@ -192,12 +193,12 @@ namespace Horde.Server.Jobs
 				filter &= filterBuilder.Ne(x => x.Outcome, JobStepOutcome.Failure);
 			}
 
-			List<JobStepRef> steps = await _jobStepRefs.Find(filter).SortByDescending(x => x.Change).ThenByDescending(x => x.StartTimeUtc).Limit(count).ToListAsync();
+			List<JobStepRef> steps = await _jobStepRefs.Find(filter).SortByDescending(x => x.Change).ThenByDescending(x => x.StartTimeUtc).Limit(maxCount).ToListAsync(cancellationToken);
 			return steps.ConvertAll<IJobStepRef>(x => x);
 		}
 
 		/// <inheritdoc/>
-		public async Task<IJobStepRef?> GetPrevStepForNodeAsync(StreamId streamId, TemplateId templateId, string nodeName, int change, JobStepOutcome? outcome = null, bool? updateIssues = null)
+		public async Task<IJobStepRef?> GetPrevStepForNodeAsync(StreamId streamId, TemplateId templateId, string nodeName, int change, JobStepOutcome? outcome = null, bool? updateIssues = null, IEnumerable<JobId>? excludeJobIds = null)
 		{
 			FilterDefinitionBuilder<JobStepRef> filterBuilder = Builders<JobStepRef>.Filter;
 
@@ -221,6 +222,11 @@ namespace Horde.Server.Jobs
 				filter &= filterBuilder.Ne(x => x.UpdateIssues, false);
 			}
 
+			if (excludeJobIds != null && excludeJobIds.Any())
+			{
+				filter &= filterBuilder.Nin(x => x.Id.JobId, excludeJobIds);
+			}
+
 			return await _jobStepRefs.Find(filter).SortByDescending(x => x.Change).FirstOrDefaultAsync();
 		}
 
diff --git a/Engine/Source/Programs/Horde/Horde.Server/Startup.cs b/Engine/Source/Programs/Horde/Horde.Server/Startup.cs
--- a/Engine/Source/Programs/Horde/Horde.Server/Startup.cs
+++ b/Engine/Source/Programs/Horde/Horde.Server/Startup.cs
@@ -98,6 +98,7 @@ using Horde.Server.Compute;
 using OpenTelemetry.Trace;
 using Polly;
 using Polly.Extensions.Http;
+using Horde.Server.Jobs.Bisect;
 
 namespace Horde.Server
 {
@@ -397,6 +398,7 @@ namespace Horde.Server
 			services.AddSingleton<ILogFileCollection, LogFileCollection>();
 			services.AddSingleton<INotificationTriggerCollection, NotificationTriggerCollection>();
 			services.AddSingleton<IPoolCollection, PoolCollection>();
+			services.AddSingleton<IBisectTaskCollection, BisectTaskCollection>();
 			services.AddSingleton<ISessionCollection, SessionCollection>();
 			services.AddSingleton<IServiceAccountCollection, ServiceAccountCollection>();
 			services.AddSingleton<ISubscriptionCollection, SubscriptionCollection>();
diff --git a/Engine/Source/Programs/Horde/Horde.Server/Utilities/HordeControllerBase.cs b/Engine/Source/Programs/Horde/Horde.Server/Utilities/HordeControllerBase.cs
--- a/Engine/Source/Programs/Horde/Horde.Server/Utilities/HordeControllerBase.cs
+++ b/Engine/Source/Programs/Horde/Horde.Server/Utilities/HordeControllerBase.cs
@@ -8,6 +8,7 @@ using Horde.Server.Agents;
 using Horde.Server.Agents.Leases;
 using Horde.Server.Artifacts;
 using Horde.Server.Jobs;
+using Horde.Server.Jobs.Bisect;
 using Horde.Server.Projects;
 using Horde.Server.Streams;
 using Horde.Server.Tools;
@@ -252,6 +253,15 @@ namespace Horde.Server.Utilities
 			return NotFound("Ref {RefName} in namespace {NamespaceId} not found", refName, namespaceId);
 		}
 
+		/// <summary>
+		/// Returns a 404 response for the given object
+		/// </summary>
+		[NonAction]
+		protected ActionResult NotFound(BisectTaskId bisectTaskId)
+		{
+			return NotFound("Bisect task {BisectTaskId} not found", bisectTaskId);
+		}
+
 		/// <summary>
 		/// Returns a 404 response for the given object
 		/// </summary>
