Horde: Missing file.

diff --git a/Engine/Source/Programs/Horde/Horde.Server/Jobs/JobMessages.cs b/Engine/Source/Programs/Horde/Horde.Server/Jobs/JobMessages.cs
--- a/Engine/Source/Programs/Horde/Horde.Server/Jobs/JobMessages.cs
+++ b/Engine/Source/Programs/Horde/Horde.Server/Jobs/JobMessages.cs
@@ -7,6 +7,7 @@ using System.Linq;
 using System.Text.Json.Serialization;
 using EpicGames.Core;
 using Horde.Server.Acls;
+using Horde.Server.Jobs.Bisect;
 using Horde.Server.Jobs.Graphs;
 using Horde.Server.Streams;
 using Horde.Server.Users;
@@ -271,6 +272,11 @@ namespace Horde.Server.Jobs
 		/// </summary>
 		public GetThinUserInfoResponse? StartedByUserInfo { get; set; }
 
+		/// <summary>
+		/// Bisection task id that started this job
+		/// </summary>
+		public BisectTaskId? StartedByBisectTaskId { get; set; }
+
 		/// <summary>
 		/// The user that aborted this job [DEPRECATED]
 		/// </summary>
@@ -373,6 +379,7 @@ namespace Horde.Server.Jobs
 			StartedByUserId = job.StartedByUserId?.ToString();
 			StartedByUser = startedByUserInfo?.Login;
 			StartedByUserInfo = startedByUserInfo;
+			StartedByBisectTaskId = job.StartedByBisectTaskId;
 			AbortedByUser = abortedByUserInfo?.Login;
 			AbortedByUserInfo = abortedByUserInfo;
 			CreateTime = new DateTimeOffset(job.CreateTimeUtc);
