Horde: Changing URIs for controller tests; local errors not being reproduced on farm.

diff --git a/Engine/Source/Programs/Horde/Horde.Server.Tests/JobsControllerTest.cs b/Engine/Source/Programs/Horde/Horde.Server.Tests/JobsControllerTest.cs
--- a/Engine/Source/Programs/Horde/Horde.Server.Tests/JobsControllerTest.cs
+++ b/Engine/Source/Programs/Horde/Horde.Server.Tests/JobsControllerTest.cs
@@ -6,14 +6,16 @@ using System.Threading.Tasks;
 using Horde.Server.Jobs.Artifacts;
 using Microsoft.VisualStudio.TestTools.UnitTesting;
 
+#pragma warning disable CA2234 // Pass system uri objects instead of strings
+
 namespace Horde.Server.Tests
 {
 	[TestClass]
     public class JobsControllerTest : ControllerIntegrationTest
     {
-        private static Uri GetUri(string jobId, string stepId, string fileName)
+        private static string GetUri(string jobId, string stepId, string fileName)
         {
-            return new Uri($"http://localhost/api/v1/jobs/{jobId}/steps/{stepId}/artifacts/{fileName}/data");
+            return $"api/v1/jobs/{jobId}/steps/{stepId}/artifacts/{fileName}/data";
         }
 
         [TestMethod]
diff --git a/Engine/Source/Programs/Horde/Horde.Server.Tests/SchemaControllerTest.cs b/Engine/Source/Programs/Horde/Horde.Server.Tests/SchemaControllerTest.cs
--- a/Engine/Source/Programs/Horde/Horde.Server.Tests/SchemaControllerTest.cs
+++ b/Engine/Source/Programs/Horde/Horde.Server.Tests/SchemaControllerTest.cs
@@ -2,14 +2,13 @@
 
 using System;
 using System.Collections.Generic;
-using System.Net;
 using System.Net.Http;
 using System.Text.Json;
-using System.Text.Json.Serialization;
 using System.Threading.Tasks;
-using Horde.Server.Jobs.Artifacts;
 using Microsoft.VisualStudio.TestTools.UnitTesting;
 
+#pragma warning disable CA2234 // Pass system uri objects instead of strings
+
 namespace Horde.Server.Tests
 {
 	[TestClass]
@@ -33,7 +32,7 @@ namespace Horde.Server.Tests
 
 			// Fetch the catalog
 			SchemaCatalog catalog;
-			using (HttpResponseMessage response = await Client.GetAsync(new Uri("http://localhost/api/v1/schema/catalog.json")))
+			using (HttpResponseMessage response = await Client.GetAsync("api/v1/schema/catalog.json"))
 			{
 				response.EnsureSuccessStatusCode();
 				byte[] data = await response.Content.ReadAsByteArrayAsync();
diff --git a/Engine/Source/Programs/Horde/Horde.Server.Tests/SwaggerTest.cs b/Engine/Source/Programs/Horde/Horde.Server.Tests/SwaggerTest.cs
--- a/Engine/Source/Programs/Horde/Horde.Server.Tests/SwaggerTest.cs
+++ b/Engine/Source/Programs/Horde/Horde.Server.Tests/SwaggerTest.cs
@@ -14,7 +14,7 @@ namespace Horde.Server.Tests
         [TestMethod]
         public async Task ValidateSwagger()
         {
-	        HttpResponseMessage res = await Client.GetAsync(new Uri("https://localhost:5001/swagger/v1/swagger.json"));
+	        HttpResponseMessage res = await Client.GetAsync("swagger/v1/swagger.json");
 	        if (!res.IsSuccessStatusCode)
 	        {
 		        string rawJson = await res.Content.ReadAsStringAsync();
