Horde: Return schema URLs using the same scheme as the incoming request.

diff --git a/Engine/Source/Programs/Horde/Horde.Server.Tests/JobsControllerTest.cs b/Engine/Source/Programs/Horde/Horde.Server.Tests/JobsControllerTest.cs
--- a/Engine/Source/Programs/Horde/Horde.Server.Tests/JobsControllerTest.cs
+++ b/Engine/Source/Programs/Horde/Horde.Server.Tests/JobsControllerTest.cs
@@ -15,7 +15,7 @@ namespace Horde.Server.Tests
     {
         private static string GetUri(string jobId, string stepId, string fileName)
         {
-            return $"api/v1/jobs/{jobId}/steps/{stepId}/artifacts/{fileName}/data";
+            return $"/api/v1/jobs/{jobId}/steps/{stepId}/artifacts/{fileName}/data";
         }
 
         [TestMethod]
diff --git a/Engine/Source/Programs/Horde/Horde.Server.Tests/SchemaControllerTest.cs b/Engine/Source/Programs/Horde/Horde.Server.Tests/SchemaControllerTest.cs
--- a/Engine/Source/Programs/Horde/Horde.Server.Tests/SchemaControllerTest.cs
+++ b/Engine/Source/Programs/Horde/Horde.Server.Tests/SchemaControllerTest.cs
@@ -32,7 +32,7 @@ namespace Horde.Server.Tests
 
 			// Fetch the catalog
 			SchemaCatalog catalog;
-			using (HttpResponseMessage response = await Client.GetAsync("api/v1/schema/catalog.json"))
+			using (HttpResponseMessage response = await Client.GetAsync("/api/v1/schema/catalog.json"))
 			{
 				response.EnsureSuccessStatusCode();
 				byte[] data = await response.Content.ReadAsByteArrayAsync();
diff --git a/Engine/Source/Programs/Horde/Horde.Server.Tests/SwaggerTest.cs b/Engine/Source/Programs/Horde/Horde.Server.Tests/SwaggerTest.cs
--- a/Engine/Source/Programs/Horde/Horde.Server.Tests/SwaggerTest.cs
+++ b/Engine/Source/Programs/Horde/Horde.Server.Tests/SwaggerTest.cs
@@ -6,6 +6,8 @@ using System.Text.Json;
 using System.Threading.Tasks;
 using Microsoft.VisualStudio.TestTools.UnitTesting;
 
+#pragma warning disable CA2234 // Pass system uri objects instead of strings
+
 namespace Horde.Server.Tests
 {
 	[TestClass]
@@ -14,8 +16,8 @@ namespace Horde.Server.Tests
         [TestMethod]
         public async Task ValidateSwagger()
         {
-	        HttpResponseMessage res = await Client.GetAsync("swagger/v1/swagger.json");
-	        if (!res.IsSuccessStatusCode)
+			HttpResponseMessage res = await Client.GetAsync("/swagger/v1/swagger.json");
+			if (!res.IsSuccessStatusCode)
 	        {
 		        string rawJson = await res.Content.ReadAsStringAsync();
 		        JsonElement tempElement = JsonSerializer.Deserialize<JsonElement>(rawJson);
diff --git a/Engine/Source/Programs/Horde/Horde.Server/Server/SchemaController.cs b/Engine/Source/Programs/Horde/Horde.Server/Server/SchemaController.cs
--- a/Engine/Source/Programs/Horde/Horde.Server/Server/SchemaController.cs
+++ b/Engine/Source/Programs/Horde/Horde.Server/Server/SchemaController.cs
@@ -9,6 +9,7 @@ using System.Reflection;
 using System.Text.Json.Serialization;
 using EpicGames.Core;
 using Horde.Server.Utilities;
+using Microsoft.AspNetCore.Http;
 using Microsoft.AspNetCore.Mvc;
 using Microsoft.Extensions.Primitives;
 
@@ -63,16 +64,6 @@ namespace Horde.Server.Server
 		[Route("/api/v1/schema/catalog.json")]
 		public ActionResult GetCatalog()
 		{
-			string? host = null;
-
-			StringValues hosts;
-			if (Request.Headers.TryGetValue("Host", out hosts))
-			{
-				host = hosts.FirstOrDefault();
-			}
-
-			host ??= Dns.GetHostName();
-
 			CatalogRoot root = new CatalogRoot();
 			foreach (Type schemaType in ConfigSchemas)
 			{
@@ -82,7 +73,7 @@ namespace Horde.Server.Server
 					JsonSchemaCatalogAttribute? catalogAttribute = schemaType.GetCustomAttribute<JsonSchemaCatalogAttribute>();
 					if (catalogAttribute != null)
 					{
-						Uri url = new Uri($"https://{host}/api/v1/schema/types/{schemaType.Name}.json");
+						Uri url = new Uri($"{Request.Scheme}://{Request.Host}/api/v1/schema/types/{schemaType.Name}.json");
 						root.Schemas.Add(new CatalogItem { Name = catalogAttribute.Name, Description = catalogAttribute.Description, FileMatch = catalogAttribute.FileMatch, Url = url });
 					}
 				}
