Horde: Add logging for graph collection cache misses.

diff --git a/Engine/Source/Programs/Horde/Horde.Server.Tests/ControllerIntegrationTest.cs b/Engine/Source/Programs/Horde/Horde.Server.Tests/ControllerIntegrationTest.cs
--- a/Engine/Source/Programs/Horde/Horde.Server.Tests/ControllerIntegrationTest.cs
+++ b/Engine/Source/Programs/Horde/Horde.Server.Tests/ControllerIntegrationTest.cs
@@ -93,12 +93,11 @@ public class ControllerIntegrationTest : IDisposable
 	{
 		IServiceProvider services = Factory.Services;
 		ConfigService configService = services.GetRequiredService<ConfigService>();
-		MongoService mongoService = services.GetRequiredService<MongoService>();
 		ITemplateCollection templateService = services.GetRequiredService<ITemplateCollection>();
 		JobService jobService = services.GetRequiredService<JobService>();
 		IArtifactCollectionV1 artifactCollection = services.GetRequiredService<IArtifactCollectionV1>();
 		AgentService agentService = services.GetRequiredService<AgentService>();
-		GraphCollection graphCollection = new(mongoService);
+		IGraphCollection graphCollection = services.GetRequiredService<IGraphCollection>();
 		IOptions<ServerSettings> serverSettings = services.GetRequiredService<IOptions<ServerSettings>>();
 
 		return await Fixture.Create(configService, graphCollection, templateService, jobService, artifactCollection, agentService, serverSettings.Value);
diff --git a/Engine/Source/Programs/Horde/Horde.Server/Jobs/Graphs/GraphCollection.cs b/Engine/Source/Programs/Horde/Horde.Server/Jobs/Graphs/GraphCollection.cs
--- a/Engine/Source/Programs/Horde/Horde.Server/Jobs/Graphs/GraphCollection.cs
+++ b/Engine/Source/Programs/Horde/Horde.Server/Jobs/Graphs/GraphCollection.cs
@@ -10,6 +10,7 @@ using Horde.Server.Jobs.Templates;
 using Horde.Server.Server;
 using HordeCommon;
 using Microsoft.Extensions.Caching.Memory;
+using Microsoft.Extensions.Logging;
 using MongoDB.Bson;
 using MongoDB.Bson.Serialization.Attributes;
 using MongoDB.Bson.Serialization.Options;
@@ -366,15 +367,16 @@ namespace Horde.Server.Jobs.Graphs
 		/// Cache for graphs
 		/// Use a non-shared cache to ensure enough space for graphs
 		/// </summary>
-		private readonly IMemoryCache _memoryCache = new MemoryCache(new MemoryCacheOptions() { SizeLimit = MaxGraphs });
+		private readonly MemoryCache _memoryCache = new MemoryCache(new MemoryCacheOptions() { SizeLimit = MaxGraphs });
+		private readonly ILogger _logger;
 
 		/// <summary>
 		/// Constructor
 		/// </summary>
-		/// <param name="mongoService">The database service singleton</param>
-		public GraphCollection(MongoService mongoService)
+		public GraphCollection(MongoService mongoService, ILogger<GraphCollection> logger)
 		{
 			_graphs = mongoService.GetCollection<GraphDocument>("Graphs");
+			_logger = logger;
 		}
 		
 		/// <inheritdoc/>
@@ -434,12 +436,17 @@ namespace Horde.Server.Jobs.Graphs
 				return GraphDocument.Empty;
 			}
 
-			return await _memoryCache.GetOrCreateAsync(hash, cacheEntry =>
+			async Task<GraphDocument> CreateCacheEntry(ICacheEntry cacheEntry)
 			{
 				cacheEntry.SlidingExpiration = TimeSpan.FromHours(24);
 				cacheEntry.Size = 1;
-				return _graphs.Find<GraphDocument>(x => x.Id == hash).FirstAsync();
-			});
+
+				GraphDocument document = await _graphs.Find<GraphDocument>(x => x.Id == hash).FirstAsync();
+				_logger.LogDebug("Creating new cache entry for {Hash} (current size={Size})", hash, _memoryCache.Count);
+				return document;
+			}
+
+			return await _memoryCache.GetOrCreateAsync(hash, CreateCacheEntry);
 		}
 
 		/// <inheritdoc/>
