Horde: Fix schema generation due to ArtifactType not having known representation.

diff --git a/Engine/Source/Programs/Horde/Horde.Server/Artifacts/ArtifactType.cs b/Engine/Source/Programs/Horde/Horde.Server/Artifacts/ArtifactType.cs
--- a/Engine/Source/Programs/Horde/Horde.Server/Artifacts/ArtifactType.cs
+++ b/Engine/Source/Programs/Horde/Horde.Server/Artifacts/ArtifactType.cs
@@ -1,6 +1,7 @@
 // Copyright Epic Games, Inc. All Rights Reserved.
 
 using System.ComponentModel;
+using EpicGames.Core;
 using EpicGames.Horde;
 using Horde.Server.Utilities;
 
@@ -10,6 +11,7 @@ namespace Horde.Server.Artifacts
 	/// Type of an artifact
 	/// </summary>
 	/// <param name="Id">The artifact type</param>
+	[JsonSchemaString]
 	[StringIdConverter(typeof(ArtifactTypeConverter))]
 	[TypeConverter(typeof(StringIdTypeConverter<ArtifactType, ArtifactTypeConverter>))]
 	public record struct ArtifactType(StringId Id)
