Horde: Additional logging for index changes.

diff --git a/Engine/Source/Programs/Horde/Horde.Server/Server/MongoService.cs b/Engine/Source/Programs/Horde/Horde.Server/Server/MongoService.cs
--- a/Engine/Source/Programs/Horde/Horde.Server/Server/MongoService.cs
+++ b/Engine/Source/Programs/Horde/Horde.Server/Server/MongoService.cs
@@ -142,7 +142,7 @@ namespace Horde.Server.Server
 		public override int GetHashCode() => Name.GetHashCode(StringComparison.Ordinal);
 
 		/// <inheritdoc/>
-		public bool Equals(MongoIndex? other) => other != null && Name.Equals(other.Name, StringComparison.Ordinal) && KeysDocument.Equals(other.KeysDocument) && Unique == other.Unique;
+		public bool Equals(MongoIndex? other) => other != null && Name.Equals(other.Name, StringComparison.Ordinal) && KeysDocument.Equals(other.KeysDocument) && ((PartialFilterDocument == null)? (other.PartialFilterDocument == null) : PartialFilterDocument.Equals(other.PartialFilterDocument)) && Unique == other.Unique;
 
 		/// <summary>
 		/// Gets the default name for an index based on its keys
@@ -748,6 +748,10 @@ namespace Horde.Server.Server
 						removeIndexNames.Remove(newIndex.Name);
 						continue;
 					}
+					else
+					{
+						_logger.LogInformation("Index {Name} has changed. Old: '{Old}', New: '{New}'", existingIndex.ToJson(), newIndex.ToJson());
+					}
 				}
 				createIndexes.Add(newIndex);
 			}
