Horde: Missing argument to log message.

diff --git a/Engine/Source/Programs/Horde/Horde.Server/Server/MongoService.cs b/Engine/Source/Programs/Horde/Horde.Server/Server/MongoService.cs
--- a/Engine/Source/Programs/Horde/Horde.Server/Server/MongoService.cs
+++ b/Engine/Source/Programs/Horde/Horde.Server/Server/MongoService.cs
@@ -750,7 +750,7 @@ namespace Horde.Server.Server
 					}
 					else
 					{
-						_logger.LogInformation("Index {Name} has changed. Old: '{Old}', New: '{New}'", existingIndex.ToJson(), newIndex.ToJson());
+						_logger.LogInformation("Index {Name} has changed. Old: '{Old}', New: '{New}'", newIndex.Name, existingIndex.ToJson(), newIndex.ToJson());
 					}
 				}
 				createIndexes.Add(newIndex);
