Horde: Increase verbosity of graph collection logging.

diff --git a/Engine/Source/Programs/Horde/Horde.Server/Jobs/Graphs/GraphCollection.cs b/Engine/Source/Programs/Horde/Horde.Server/Jobs/Graphs/GraphCollection.cs
--- a/Engine/Source/Programs/Horde/Horde.Server/Jobs/Graphs/GraphCollection.cs
+++ b/Engine/Source/Programs/Horde/Horde.Server/Jobs/Graphs/GraphCollection.cs
@@ -442,7 +442,7 @@ namespace Horde.Server.Jobs.Graphs
 				cacheEntry.Size = 1;
 
 				GraphDocument document = await _graphs.Find<GraphDocument>(x => x.Id == hash).FirstAsync();
-				_logger.LogDebug("Creating new cache entry for {Hash} (current size={Size})", hash, _memoryCache.Count);
+				_logger.LogInformation("Creating new cache entry for {Hash} (current size={Size})", hash, _memoryCache.Count);
 				return document;
 			}
 
