Horde: Some additional logging for index creation.

diff --git a/Engine/Source/Programs/Horde/Horde.Server/Server/MongoService.cs b/Engine/Source/Programs/Horde/Horde.Server/Server/MongoService.cs
--- a/Engine/Source/Programs/Horde/Horde.Server/Server/MongoService.cs
+++ b/Engine/Source/Programs/Horde/Horde.Server/Server/MongoService.cs
@@ -756,38 +756,44 @@ namespace Horde.Server.Server
 				createIndexes.Add(newIndex);
 			}
 
-			// Drop any indexes that are no longer needed
-			foreach (string removeIndexName in removeIndexNames)
+			// List all the current indexes
+			if (createIndexes.Count > 0 || removeIndexNames.Count > 0)
 			{
-				if (ReadOnlyMode)
-				{
-					_logger.LogWarning("Would drop unused index {CollectionName}.{IndexName} - skipping due to read-only setting.", collectionName, removeIndexName);
-				}
-				else
-				{
-					_logger.LogInformation("Dropping unused index {IndexName}", removeIndexName);
-					await collection.Indexes.DropOneAsync(removeIndexName, cancellationToken);
-				}
-			}
+				_logger.LogInformation("Collection {Name} currently has {NumIndexes} indexes: {IndexList}", collectionName, nameToExistingIndex.Count, String.Join(", ", nameToExistingIndex.Keys));
 
-			// Create all the new indexes
-			foreach (MongoIndex<T> createIndex in createIndexes)
-			{
-				if (ReadOnlyMode)
+				// Drop any indexes that are no longer needed
+				foreach (string removeIndexName in removeIndexNames)
 				{
-					_logger.LogWarning("Would create index {CollectionName}.{IndexName} - skipping due to read-only setting.", collectionName, createIndex.Name);
+					if (ReadOnlyMode)
+					{
+						_logger.LogWarning("Would drop unused index {CollectionName}.{IndexName} - skipping due to read-only setting.", collectionName, removeIndexName);
+					}
+					else
+					{
+						_logger.LogInformation("Dropping unused index {IndexName}", removeIndexName);
+						await collection.Indexes.DropOneAsync(removeIndexName, cancellationToken);
+					}
 				}
-				else
+
+				// Create all the new indexes
+				foreach (MongoIndex<T> createIndex in createIndexes)
 				{
-					_logger.LogInformation("Creating index {CollectionName}: {IndexName}", collectionName, createIndex.Name);
+					if (ReadOnlyMode)
+					{
+						_logger.LogWarning("Would create index {CollectionName}.{IndexName} - skipping due to read-only setting.", collectionName, createIndex.Name);
+					}
+					else
+					{
+						_logger.LogInformation("Creating index {CollectionName}: {IndexName}", collectionName, createIndex.Name);
 
-					CreateIndexOptions<T> options = new CreateIndexOptions<T>();
-					options.Name = createIndex.Name;
-					options.PartialFilterExpression = createIndex.PartialFilter;
-					options.Unique = createIndex.Unique;
+						CreateIndexOptions<T> options = new CreateIndexOptions<T>();
+						options.Name = createIndex.Name;
+						options.PartialFilterExpression = createIndex.PartialFilter;
+						options.Unique = createIndex.Unique;
 
-					CreateIndexModel<T> model = new CreateIndexModel<T>(createIndex.Keys, options);
-					await collection.Indexes.CreateOneAsync(model, cancellationToken: cancellationToken);
+						CreateIndexModel<T> model = new CreateIndexModel<T>(createIndex.Keys, options);
+						await collection.Indexes.CreateOneAsync(model, cancellationToken: cancellationToken);
+					}
 				}
 			}
 
