Horde: Additional debug properties for job queue scaling.

diff --git a/Engine/Source/Programs/Horde/Horde.Server/Agents/Fleet/JobQueueStrategy.cs b/Engine/Source/Programs/Horde/Horde.Server/Agents/Fleet/JobQueueStrategy.cs
--- a/Engine/Source/Programs/Horde/Horde.Server/Agents/Fleet/JobQueueStrategy.cs
+++ b/Engine/Source/Programs/Horde/Horde.Server/Agents/Fleet/JobQueueStrategy.cs
@@ -172,9 +172,11 @@ namespace Horde.Server.Agents.Fleet
 		internal async Task<Dictionary<PoolId, int>> GetPoolQueueSizesAsync(DateTimeOffset jobsCreatedAfter)
 		{
 			using TelemetrySpan span = OpenTelemetryTracers.Horde.StartActiveSpan($"{nameof(JobQueueStrategy)}.{nameof(GetPoolQueueSizesAsync)}");
+			span.SetAttribute("after", jobsCreatedAfter);
 
 			Dictionary<StreamId, StreamConfig> streams = _globalConfig.CurrentValue.Streams.ToDictionary(x => x.Id, x => (StreamConfig)x);
 			List<IJob> recentJobs = await _jobs.FindAsync(minCreateTime: jobsCreatedAfter);
+			span.SetAttribute("numJobs", recentJobs.Count);
 
 			List<(IJob Job, IJobStepBatch Batch, PoolId PoolId)> jobBatches = new();
 			foreach (IJob job in recentJobs)
@@ -194,6 +196,7 @@ namespace Horde.Server.Agents.Fleet
 			using TelemetrySpan span = OpenTelemetryTracers.Horde.StartActiveSpan($"{nameof(JobQueueStrategy)}.{nameof(CalculatePoolSizeAsync)}");
 			span.SetAttribute(OpenTelemetryTracers.DatadogResourceAttribute, pool.Id.ToString());
 			span.SetAttribute("currentAgentCount", agents.Count);
+			span.SetAttribute("samplePeriodMin", Settings.SamplePeriodMin);
 			
 			DateTimeOffset minCreateTime = _clock.UtcNow - TimeSpan.FromMinutes(Settings.SamplePeriodMin);
 
