Horde: Include the number of unique graphs in autoscaling telemetry.

diff --git a/Engine/Source/Programs/Horde/Horde.Server/Agents/Fleet/JobQueueStrategy.cs b/Engine/Source/Programs/Horde/Horde.Server/Agents/Fleet/JobQueueStrategy.cs
--- a/Engine/Source/Programs/Horde/Horde.Server/Agents/Fleet/JobQueueStrategy.cs
+++ b/Engine/Source/Programs/Horde/Horde.Server/Agents/Fleet/JobQueueStrategy.cs
@@ -177,6 +177,7 @@ namespace Horde.Server.Agents.Fleet
 			Dictionary<StreamId, StreamConfig> streams = _globalConfig.CurrentValue.Streams.ToDictionary(x => x.Id, x => (StreamConfig)x);
 			List<IJob> recentJobs = await _jobs.FindAsync(minCreateTime: jobsCreatedAfter);
 			span.SetAttribute("numJobs", recentJobs.Count);
+			span.SetAttribute("numUniqueGraphs", recentJobs.Select(x => x.GraphHash).Distinct().Count());
 
 			List<(IJob Job, IJobStepBatch Batch, PoolId PoolId)> jobBatches = new();
 			foreach (IJob job in recentJobs)
