Horde: Improve logging for mongo DB upgrades.

diff --git a/Engine/Source/Programs/Horde/Horde.Server/Server/MongoService.cs b/Engine/Source/Programs/Horde/Horde.Server/Server/MongoService.cs
--- a/Engine/Source/Programs/Horde/Horde.Server/Server/MongoService.cs
+++ b/Engine/Source/Programs/Horde/Horde.Server/Server/MongoService.cs
@@ -748,19 +748,12 @@ namespace Horde.Server.Server
 						removeIndexNames.Remove(newIndex.Name);
 						continue;
 					}
-					else
-					{
-						_logger.LogInformation("Index {Name} has changed. Old: '{Old}', New: '{New}'", newIndex.Name, existingIndex.ToJson(), newIndex.ToJson());
-					}
 				}
 				createIndexes.Add(newIndex);
 			}
 
-			// List all the current indexes
-			if (createIndexes.Count > 0 || removeIndexNames.Count > 0)
+			if (removeIndexNames.Count > 0 || createIndexes.Count > 0)
 			{
-				_logger.LogInformation("Collection {Name} currently has {NumIndexes} indexes: {IndexList}", collectionName, nameToExistingIndex.Count, String.Join(", ", nameToExistingIndex.Keys));
-
 				// Drop any indexes that are no longer needed
 				foreach (string removeIndexName in removeIndexNames)
 				{
@@ -784,7 +777,7 @@ namespace Horde.Server.Server
 					}
 					else
 					{
-						_logger.LogInformation("Creating index {CollectionName}: {IndexName}", collectionName, createIndex.Name);
+						_logger.LogInformation("Creating index {IndexName} in {CollectionName}: {Definition}", createIndex.Name, collectionName, createIndex.ToJson());
 
 						CreateIndexOptions<T> options = new CreateIndexOptions<T>();
 						options.Name = createIndex.Name;
@@ -792,12 +785,14 @@ namespace Horde.Server.Server
 						options.Unique = createIndex.Unique;
 
 						CreateIndexModel<T> model = new CreateIndexModel<T>(createIndex.Keys, options);
-						await collection.Indexes.CreateOneAsync(model, cancellationToken: cancellationToken);
+
+						string result = await collection.Indexes.CreateOneAsync(model, cancellationToken: cancellationToken);
+						_logger.LogInformation("Created index {IndexName}", result);
 					}
 				}
-			}
 
-			_logger.LogDebug("Finished updating indexes for collection {CollectionName}", collectionName);
+				_logger.LogInformation("Finished updating indexes for collection {CollectionName}", collectionName);
+			}
 		}
 
 		class SingletonInfo<T> where T : SingletonBase
