Horde: Increase size of graph cache to 5000, reduce number of jobs queried to last two days.

diff --git a/Engine/Source/Programs/Horde/Horde.Server/Agents/Fleet/JobQueueStrategy.cs b/Engine/Source/Programs/Horde/Horde.Server/Agents/Fleet/JobQueueStrategy.cs
--- a/Engine/Source/Programs/Horde/Horde.Server/Agents/Fleet/JobQueueStrategy.cs
+++ b/Engine/Source/Programs/Horde/Horde.Server/Agents/Fleet/JobQueueStrategy.cs
@@ -41,7 +41,7 @@ namespace Horde.Server.Agents.Fleet
 		/// <summary>
 		/// How far back in time to look for job batches (that potentially are in the queue)
 		/// </summary>
-		public int SamplePeriodMin { get; set; } = 60 * 24 * 5; // 5 days
+		public int SamplePeriodMin { get; set; } = 60 * 24 * 2; // 2 days
 
 		/// <summary>
 		/// Time spent in ready state before considered truly waiting for an agent
diff --git a/Engine/Source/Programs/Horde/Horde.Server/Jobs/Graphs/GraphCollection.cs b/Engine/Source/Programs/Horde/Horde.Server/Jobs/Graphs/GraphCollection.cs
--- a/Engine/Source/Programs/Horde/Horde.Server/Jobs/Graphs/GraphCollection.cs
+++ b/Engine/Source/Programs/Horde/Horde.Server/Jobs/Graphs/GraphCollection.cs
@@ -361,7 +361,7 @@ namespace Horde.Server.Jobs.Graphs
 		/// <summary>
 		/// Maximum number of graphs to keep in the cache
 		/// </summary>
-		const int MaxGraphs = 1000;
+		const int MaxGraphs = 5000;
 
 		/// <summary>
 		/// Cache for graphs
@@ -442,7 +442,6 @@ namespace Horde.Server.Jobs.Graphs
 				cacheEntry.Size = 1;
 
 				GraphDocument document = await _graphs.Find<GraphDocument>(x => x.Id == hash).FirstAsync();
-				_logger.LogInformation("Creating new cache entry for {Hash} (current size={Size})", hash, _memoryCache.Count);
 				return document;
 			}
 
