Horde: Remove support for partial indexes (unsupported by AWS DocDB). Add support for sparse indexes instead.

diff --git a/Engine/Source/Programs/Horde/Horde.Server/Artifacts/ArtifactCollection.cs b/Engine/Source/Programs/Horde/Horde.Server/Artifacts/ArtifactCollection.cs
--- a/Engine/Source/Programs/Horde/Horde.Server/Artifacts/ArtifactCollection.cs
+++ b/Engine/Source/Programs/Horde/Horde.Server/Artifacts/ArtifactCollection.cs
@@ -77,7 +77,7 @@ namespace Horde.Server.Artifacts
 		{
 			List<MongoIndex<Artifact>> indexes = new List<MongoIndex<Artifact>>();
 			indexes.Add(keys => keys.Ascending(x => x.Keys));
-			indexes.Add(keys => keys.Ascending(x => x.ExpireAtUtc), Builders<Artifact>.Filter.Exists(x => x.ExpireAtUtc));
+			indexes.Add(keys => keys.Ascending(x => x.ExpireAtUtc), sparse: true);
 			_artifacts = mongoService.GetCollection<Artifact>("ArtifactsV2", indexes);
 		}
 
diff --git a/Engine/Source/Programs/Horde/Horde.Server/Jobs/Bisect/BisectTaskCollection.cs b/Engine/Source/Programs/Horde/Horde.Server/Jobs/Bisect/BisectTaskCollection.cs
--- a/Engine/Source/Programs/Horde/Horde.Server/Jobs/Bisect/BisectTaskCollection.cs
+++ b/Engine/Source/Programs/Horde/Horde.Server/Jobs/Bisect/BisectTaskCollection.cs
@@ -22,6 +22,13 @@ namespace Horde.Server.Jobs.Bisect
 		{
 			public BisectTaskId Id { get; set; }
 
+			[BsonElement("running"), BsonIgnoreIfNull]
+			public bool? Running // Used for sparse index
+			{
+				get => (State == BisectTaskState.Running)? true : null;
+				set { }
+			}
+
 			[BsonElement("state")]
 			public BisectTaskState State { get; set; }
 
@@ -73,7 +80,7 @@ namespace Horde.Server.Jobs.Bisect
 		public BisectTaskCollection(MongoService mongoService)
 		{
 			List<MongoIndex<BisectTaskDoc>> indexes = new List<MongoIndex<BisectTaskDoc>>();
-			indexes.Add(keys => keys.Ascending(x => x.Id), Builders<BisectTaskDoc>.Filter.Eq(x => x.State, BisectTaskState.Running));
+			indexes.Add(keys => keys.Ascending(x => x.Id).Ascending(x => x.Running), sparse: true);
 
 			_bisectTasks = mongoService.GetCollection<BisectTaskDoc>("BisectTasks", indexes);
 		}
diff --git a/Engine/Source/Programs/Horde/Horde.Server/Server/MongoService.cs b/Engine/Source/Programs/Horde/Horde.Server/Server/MongoService.cs
--- a/Engine/Source/Programs/Horde/Horde.Server/Server/MongoService.cs
+++ b/Engine/Source/Programs/Horde/Horde.Server/Server/MongoService.cs
@@ -53,18 +53,18 @@ namespace Horde.Server.Server
 		[BsonElement("key"), BsonRequired]
 		public BsonDocument KeysDocument { get; private set; }
 
-		/// <summary>
-		/// Optional partial filter for the index
-		/// </summary>
-		[BsonElement("partialFilterExpression")]
-		public BsonDocument? PartialFilterDocument { get; private set; }
-
 		/// <summary>
 		/// The index should be unique
 		/// </summary>
 		[BsonElement("unique")]
 		public bool Unique { get; private set; }
 
+		/// <summary>
+		/// Whether to create a sparse index
+		/// </summary>
+		[BsonElement("sparse")]
+		public bool Sparse { get; private set; }
+
 		/// <summary>
 		/// Constructor
 		/// </summary>
@@ -73,7 +73,6 @@ namespace Horde.Server.Server
 		{
 			Name = String.Empty;
 			KeysDocument = new BsonDocument();
-			PartialFilterDocument = null;
 		}
 
 		/// <summary>
@@ -81,14 +80,14 @@ namespace Horde.Server.Server
 		/// </summary>
 		/// <param name="name">Name of the index</param>
 		/// <param name="keysDocument">Keys for the index</param>
-		/// <param name="partialFilterDocument">Partial filter for the index</param>
 		/// <param name="unique">Whether the index should be unique</param>
-		public MongoIndex(string? name, BsonDocument keysDocument, BsonDocument? partialFilterDocument, bool unique)
+		/// <param name="sparse">Whether to create a sparse index</param>
+		public MongoIndex(string? name, BsonDocument keysDocument, bool unique, bool sparse)
 		{
 			Name = name ?? GetDefaultName(keysDocument);
 			KeysDocument = keysDocument;
-			PartialFilterDocument = partialFilterDocument;
 			Unique = unique;
+			Sparse = sparse;
 		}
 
 		/// <summary>
@@ -96,31 +95,10 @@ namespace Horde.Server.Server
 		/// </summary>
 		/// <param name="keysFunc">Callback to configure keys for this document</param>
 		/// <param name="unique">Whether the index should be unique</param>
-		public static MongoIndex<T> Create<T>(Func<IndexKeysDefinitionBuilder<T>, IndexKeysDefinition<T>> keysFunc, bool unique = false)
-		{
-			return new MongoIndex<T>(null, keysFunc(Builders<T>.IndexKeys), null, unique);
-		}
-
-		/// <summary>
-		/// Constructor
-		/// </summary>
-		/// <param name="keysFunc">Callback to configure keys for this document</param>
-		/// <param name="partialFilter">Partial filter for the index</param>
-		/// <param name="unique">Whether the index should be unique</param>
-		public static MongoIndex<T> Create<T>(Func<IndexKeysDefinitionBuilder<T>, IndexKeysDefinition<T>> keysFunc, FilterDefinition<T>? partialFilter, bool unique = false)
-		{
-			return new MongoIndex<T>(null, keysFunc(Builders<T>.IndexKeys), partialFilter, unique);
-		}
-
-		/// <summary>
-		/// Constructor
-		/// </summary>
-		/// <param name="name">Name of the index</param>
-		/// <param name="keysFunc">Callback to configure keys for this document</param>
-		/// <param name="unique">Whether the index should be unique</param>
-		public static MongoIndex<T> Create<T>(string name, Func<IndexKeysDefinitionBuilder<T>, IndexKeysDefinition<T>> keysFunc, bool unique = false)
+		/// <param name="sparse">Whether to create a sparse index</param>
+		public static MongoIndex<T> Create<T>(Func<IndexKeysDefinitionBuilder<T>, IndexKeysDefinition<T>> keysFunc, bool unique = false, bool sparse = false)
 		{
-			return new MongoIndex<T>(name, keysFunc(Builders<T>.IndexKeys), null, unique);
+			return new MongoIndex<T>(null, keysFunc(Builders<T>.IndexKeys), unique, sparse);
 		}
 
 		/// <summary>
@@ -128,11 +106,11 @@ namespace Horde.Server.Server
 		/// </summary>
 		/// <param name="name">Name of the index</param>
 		/// <param name="keysFunc">Callback to configure keys for this document</param>
-		/// <param name="partialFilter">Partial filter for the index</param>
 		/// <param name="unique">Whether the index should be unique</param>
-		public static MongoIndex<T> Create<T>(string name, Func<IndexKeysDefinitionBuilder<T>, IndexKeysDefinition<T>> keysFunc, FilterDefinition<T>? partialFilter, bool unique = false)
+		/// <param name="sparse">Whether to create a sparse index</param>
+		public static MongoIndex<T> Create<T>(string name, Func<IndexKeysDefinitionBuilder<T>, IndexKeysDefinition<T>> keysFunc, bool unique = false, bool sparse = false)
 		{
-			return new MongoIndex<T>(name, keysFunc(Builders<T>.IndexKeys), partialFilter, unique);
+			return new MongoIndex<T>(name, keysFunc(Builders<T>.IndexKeys), unique, sparse);
 		}
 
 		/// <inheritdoc/>
@@ -142,7 +120,7 @@ namespace Horde.Server.Server
 		public override int GetHashCode() => Name.GetHashCode(StringComparison.Ordinal);
 
 		/// <inheritdoc/>
-		public bool Equals(MongoIndex? other) => other != null && Name.Equals(other.Name, StringComparison.Ordinal) && KeysDocument.Equals(other.KeysDocument) && ((PartialFilterDocument == null)? (other.PartialFilterDocument == null) : PartialFilterDocument.Equals(other.PartialFilterDocument)) && Unique == other.Unique;
+		public bool Equals(MongoIndex? other) => other != null && Name.Equals(other.Name, StringComparison.Ordinal) && KeysDocument.Equals(other.KeysDocument) && Sparse == other.Sparse && Unique == other.Unique;
 
 		/// <summary>
 		/// Gets the default name for an index based on its keys
@@ -175,25 +153,20 @@ namespace Horde.Server.Server
 		/// <summary>
 		/// Keys for the index
 		/// </summary>
+		[BsonIgnore]
 		public IndexKeysDefinition<T> Keys { get; }
 
-		/// <summary>
-		/// Optional expression specifying documents to include in the index
-		/// </summary>
-		public FilterDefinition<T>? PartialFilter { get; }
-
 		/// <summary>
 		/// Constructor
 		/// </summary>
 		/// <param name="name">Name of the index</param>
 		/// <param name="keys">Keys for the index</param>
-		/// <param name="partialFilter">Partial filter for the index</param>
 		/// <param name="unique">Whether the index should be unique</param>
-		public MongoIndex(string? name, IndexKeysDefinition<T> keys, FilterDefinition<T>? partialFilter, bool unique)
-			: base(name, keys.Render(BsonSerializer.LookupSerializer<T>(), BsonSerializer.SerializerRegistry), partialFilter?.Render(BsonSerializer.LookupSerializer<T>(), BsonSerializer.SerializerRegistry), unique)
+		/// <param name="sparse">Whether to create a sparse index</param>
+		public MongoIndex(string? name, IndexKeysDefinition<T> keys, bool unique, bool sparse)
+			: base(name, keys.Render(BsonSerializer.LookupSerializer<T>(), BsonSerializer.SerializerRegistry), unique, sparse)
 		{
 			Keys = keys;
-			PartialFilter = partialFilter;
 		}
 	}
 
@@ -202,24 +175,14 @@ namespace Horde.Server.Server
 	/// </summary>
 	static class MongoIndexExtensions
 	{
-		public static void Add<T>(this List<MongoIndex<T>> list, Func<IndexKeysDefinitionBuilder<T>, IndexKeysDefinition<T>> keyFunc, bool unique = false)
-		{
-			list.Add(MongoIndex.Create(keyFunc, unique));
-		}
-
-		public static void Add<T>(this List<MongoIndex<T>> list, Func<IndexKeysDefinitionBuilder<T>, IndexKeysDefinition<T>> keyFunc, FilterDefinition<T>? partialFilter, bool unique = false)
-		{
-			list.Add(MongoIndex.Create(keyFunc, partialFilter, unique));
-		}
-
-		public static void Add<T>(this List<MongoIndex<T>> list, string name, Func<IndexKeysDefinitionBuilder<T>, IndexKeysDefinition<T>> keyFunc, bool unique = false)
+		public static void Add<T>(this List<MongoIndex<T>> list, Func<IndexKeysDefinitionBuilder<T>, IndexKeysDefinition<T>> keyFunc, bool unique = false, bool sparse = false)
 		{
-			list.Add(MongoIndex.Create(name, keyFunc, unique));
+			list.Add(MongoIndex.Create(keyFunc, unique, sparse));
 		}
 
-		public static void Add<T>(this List<MongoIndex<T>> list, string name, Func<IndexKeysDefinitionBuilder<T>, IndexKeysDefinition<T>> keyFunc, FilterDefinition<T> partialIndex, bool unique = false)
+		public static void Add<T>(this List<MongoIndex<T>> list, string name, Func<IndexKeysDefinitionBuilder<T>, IndexKeysDefinition<T>> keyFunc, bool unique = false, bool sparse = false)
 		{
-			list.Add(MongoIndex.Create(name, keyFunc, partialIndex, unique));
+			list.Add(MongoIndex.Create(name, keyFunc, unique, sparse));
 		}
 	}
 
@@ -781,8 +744,8 @@ namespace Horde.Server.Server
 
 						CreateIndexOptions<T> options = new CreateIndexOptions<T>();
 						options.Name = createIndex.Name;
-						options.PartialFilterExpression = createIndex.PartialFilter;
 						options.Unique = createIndex.Unique;
+						options.Sparse = createIndex.Sparse;
 
 						CreateIndexModel<T> model = new CreateIndexModel<T>(createIndex.Keys, options);
 
