Horde: Reorganizing project.

diff --git a/Engine/Source/Programs/Horde/Horde/Commands/Login/LoginCommand.cs b/Engine/Source/Programs/Horde/Horde/Commands/Login/LoginCommand.cs
deleted file mode 100644
diff --git a/Engine/Source/Programs/Horde/Horde/Commands/LoginCommand.cs b/Engine/Source/Programs/Horde/Horde/Commands/LoginCommand.cs
new file mode 100644
--- /dev/null
+++ b/Engine/Source/Programs/Horde/Horde/Commands/LoginCommand.cs
@@ -0,0 +1,45 @@
+// Copyright Epic Games, Inc. All Rights Reserved.
+
+using EpicGames.Core;
+using Microsoft.Extensions.Logging;
+
+namespace Horde.Commands
+{
+	[Command("login", "Logs in to a Horde server")]
+	class LoginCommand : Command
+	{
+		class GetAuthConfigResponse
+		{
+			public string? Method { get; set; }
+			public string? ServerUrl { get; set; }
+			public string? ClientId { get; set; }
+			public string[]? RedirectUrls { get; set; }
+		}
+
+		[CommandLine("-Server=")]
+		public string? Server { get; set; }
+
+		/// <inheritdoc/>
+		public override async Task<int> ExecuteAsync(ILogger logger)
+		{
+			if (Server != null)
+			{
+				await Settings.SetServerAsync(Server);
+			}
+
+			if(await Settings.GetServerAsync() == null)
+			{
+				logger.LogError("No Horde server is configured. Specify -Server=... to configure one.");
+				return 1;
+			}
+
+			if (await Settings.GetAccessTokenAsync(logger) == null)
+			{
+				logger.LogError("Unable to log in to server");
+				return 1;
+			}
+
+			return 0;
+		}
+	}
+}
