Horde: Optimize JobQueueStrategy by filtering on batch state

Reduces the number of documents returned. No index on batch state is used as combining with CreateTimeUtc index is tricky.

diff --git a/Engine/Source/Programs/Horde/Horde.Server/Agents/Fleet/JobQueueStrategy.cs b/Engine/Source/Programs/Horde/Horde.Server/Agents/Fleet/JobQueueStrategy.cs
--- a/Engine/Source/Programs/Horde/Horde.Server/Agents/Fleet/JobQueueStrategy.cs
+++ b/Engine/Source/Programs/Horde/Horde.Server/Agents/Fleet/JobQueueStrategy.cs
@@ -175,7 +175,7 @@ namespace Horde.Server.Agents.Fleet
 			span.SetAttribute("after", jobsCreatedAfter);
 
 			Dictionary<StreamId, StreamConfig> streams = _globalConfig.CurrentValue.Streams.ToDictionary(x => x.Id, x => (StreamConfig)x);
-			List<IJob> recentJobs = await _jobs.FindAsync(minCreateTime: jobsCreatedAfter);
+			List<IJob> recentJobs = await _jobs.FindAsync(minCreateTime: jobsCreatedAfter, batchState: JobStepBatchState.Ready);
 			span.SetAttribute("numJobs", recentJobs.Count);
 			span.SetAttribute("numUniqueGraphs", recentJobs.Select(x => x.GraphHash).Distinct().Count());
 
diff --git a/Engine/Source/Programs/Horde/Horde.Server/Jobs/IJobCollection.cs b/Engine/Source/Programs/Horde/Horde.Server/Jobs/IJobCollection.cs
--- a/Engine/Source/Programs/Horde/Horde.Server/Jobs/IJobCollection.cs
+++ b/Engine/Source/Programs/Horde/Horde.Server/Jobs/IJobCollection.cs
@@ -160,13 +160,14 @@ namespace Horde.Server.Jobs
 		/// <param name="maxCreateTime">The maximum creation time</param>
 		/// <param name="modifiedBefore">Filter the results by modified time</param>
 		/// <param name="modifiedAfter">Filter the results by modified time</param>
+		/// <param name="batchState">One or more batches matches this state</param>
 		/// <param name="index">Index of the first result to return</param>
 		/// <param name="count">Number of results to return</param>
 		/// <param name="consistentRead">If the database read should be made to the replica server</param>
 		/// <param name="indexHint">Name of index to be specified as a hint to the database query planner</param>
 		/// <param name="excludeUserJobs">Whether to exclude user jobs from the find</param>
 		/// <returns>List of jobs matching the given criteria</returns>
-		Task<List<IJob>> FindAsync(JobId[]? jobIds = null, StreamId? streamId = null, string? name = null, TemplateId[]? templates = null, int? minChange = null, int? maxChange = null, int? preflightChange = null, bool? preflightOnly = null, UserId? preflightStartedByUser = null, UserId? startedByUser = null, DateTimeOffset? minCreateTime = null, DateTimeOffset? maxCreateTime = null, DateTimeOffset? modifiedBefore = null, DateTimeOffset? modifiedAfter = null, int? index = null, int? count = null, bool consistentRead = true, string? indexHint = null, bool? excludeUserJobs = null);
+		Task<List<IJob>> FindAsync(JobId[]? jobIds = null, StreamId? streamId = null, string? name = null, TemplateId[]? templates = null, int? minChange = null, int? maxChange = null, int? preflightChange = null, bool? preflightOnly = null, UserId? preflightStartedByUser = null, UserId? startedByUser = null, DateTimeOffset? minCreateTime = null, DateTimeOffset? maxCreateTime = null, DateTimeOffset? modifiedBefore = null, DateTimeOffset? modifiedAfter = null, JobStepBatchState? batchState = null, int? index = null, int? count = null, bool consistentRead = true, string? indexHint = null, bool? excludeUserJobs = null);
 
 		/// <summary>
 		/// Searches for jobs matching the given criteria
diff --git a/Engine/Source/Programs/Horde/Horde.Server/Jobs/JobCollection.cs b/Engine/Source/Programs/Horde/Horde.Server/Jobs/JobCollection.cs
--- a/Engine/Source/Programs/Horde/Horde.Server/Jobs/JobCollection.cs
+++ b/Engine/Source/Programs/Horde/Horde.Server/Jobs/JobCollection.cs
@@ -492,7 +492,7 @@ namespace Horde.Server.Jobs
 		}
 
 		/// <inheritdoc/>
-		public async Task<List<IJob>> FindAsync(JobId[]? jobIds, StreamId? streamId, string? name, TemplateId[]? templates, int? minChange, int? maxChange, int? preflightChange, bool? preflightOnly, UserId ? preflightStartedByUser, UserId? startedByUser, DateTimeOffset? minCreateTime, DateTimeOffset? maxCreateTime, DateTimeOffset? modifiedBefore, DateTimeOffset? modifiedAfter, int? index, int? count, bool consistentRead, string? indexHint, bool? excludeUserJobs)
+		public async Task<List<IJob>> FindAsync(JobId[]? jobIds, StreamId? streamId, string? name, TemplateId[]? templates, int? minChange, int? maxChange, int? preflightChange, bool? preflightOnly, UserId ? preflightStartedByUser, UserId? startedByUser, DateTimeOffset? minCreateTime, DateTimeOffset? maxCreateTime, DateTimeOffset? modifiedBefore, DateTimeOffset? modifiedAfter, JobStepBatchState? batchState, int? index, int? count, bool consistentRead, string? indexHint, bool? excludeUserJobs)
 		{
 			FilterDefinitionBuilder<JobDocument> filterBuilder = Builders<JobDocument>.Filter;
 
@@ -568,6 +568,10 @@ namespace Horde.Server.Jobs
 			{
 				filter &= filterBuilder.Gte(x => x.UpdateTimeUtc!, modifiedAfter.Value.UtcDateTime);
 			}
+			if (batchState != null)
+			{
+				filter &= filterBuilder.ElemMatch(x => x.Batches, batch => batch.State == batchState);
+			}
 
 			List<JobDocument> results;
 			using (TelemetrySpan _ = _tracer.StartActiveSpan($"{nameof(JobCollection)}.{nameof(FindAsync)}"))
diff --git a/Engine/Source/Programs/Horde/Horde.Server/Jobs/JobService.cs b/Engine/Source/Programs/Horde/Horde.Server/Jobs/JobService.cs
--- a/Engine/Source/Programs/Horde/Horde.Server/Jobs/JobService.cs
+++ b/Engine/Source/Programs/Horde/Horde.Server/Jobs/JobService.cs
@@ -530,6 +530,7 @@ namespace Horde.Server.Jobs
 		/// <param name="minCreateTime">The minimum creation time</param>
 		/// <param name="maxCreateTime">The maximum creation time</param>
 		/// <param name="target">The target to query</param>
+		/// <param name="batchState">One or more batches matches this state</param>
 		/// <param name="state">State to query</param>
 		/// <param name="outcome">Outcomes to return</param>
 		/// <param name="modifiedBefore">Filter the results by last modified time</param>
@@ -539,7 +540,7 @@ namespace Horde.Server.Jobs
 		/// <param name="consistentRead">If the database read should be made to the replica server</param>
 		/// <param name="excludeUserJobs">Whether to exclude user jobs from the find</param>
 		/// <returns>List of jobs matching the given criteria</returns>
-		public async Task<List<IJob>> FindJobsAsync(JobId[]? jobIds = null, StreamId? streamId = null, string? name = null, TemplateId[]? templates = null, int? minChange = null, int? maxChange = null, int? preflightChange = null, bool? preflightOnly = null, UserId? preflightStartedByUser = null, UserId? startedByUser = null, DateTimeOffset ? minCreateTime = null, DateTimeOffset? maxCreateTime = null, string? target = null, JobStepState[]? state = null, JobStepOutcome[]? outcome = null, DateTimeOffset? modifiedBefore = null, DateTimeOffset? modifiedAfter = null, int? index = null, int? count = null, bool consistentRead = true, bool? excludeUserJobs = null)
+		public async Task<List<IJob>> FindJobsAsync(JobId[]? jobIds = null, StreamId? streamId = null, string? name = null, TemplateId[]? templates = null, int? minChange = null, int? maxChange = null, int? preflightChange = null, bool? preflightOnly = null, UserId? preflightStartedByUser = null, UserId? startedByUser = null, DateTimeOffset ? minCreateTime = null, DateTimeOffset? maxCreateTime = null, string? target = null, JobStepBatchState? batchState = null, JobStepState[]? state = null, JobStepOutcome[]? outcome = null, DateTimeOffset? modifiedBefore = null, DateTimeOffset? modifiedAfter = null, int? index = null, int? count = null, bool consistentRead = true, bool? excludeUserJobs = null)
 		{
 			using TelemetrySpan span = _tracer.StartActiveSpan($"{nameof(JobService)}.{nameof(FindJobsAsync)}");
 			span.SetAttribute("JobIds", (jobIds == null)? null : String.Join(',', jobIds));
@@ -563,7 +564,7 @@ namespace Horde.Server.Jobs
 			
 			if (target == null && (state == null || state.Length == 0) && (outcome == null || outcome.Length == 0))
 			{
-				return await _jobs.FindAsync(jobIds, streamId, name, templates, minChange, maxChange, preflightChange, preflightOnly, preflightStartedByUser, startedByUser, minCreateTime, maxCreateTime, modifiedBefore, modifiedAfter, index, count, consistentRead, null, excludeUserJobs);
+				return await _jobs.FindAsync(jobIds, streamId, name, templates, minChange, maxChange, preflightChange, preflightOnly, preflightStartedByUser, startedByUser, minCreateTime, maxCreateTime, modifiedBefore, modifiedAfter, batchState, index, count, consistentRead, null, excludeUserJobs);
 			}
 			else
 			{
@@ -573,7 +574,7 @@ namespace Horde.Server.Jobs
 				int maxCount = (count ?? 1);
 				while (results.Count < maxCount)
 				{
-					List<IJob> scanJobs = await _jobs.FindAsync(jobIds, streamId, name, templates, minChange, maxChange, preflightChange, preflightOnly, preflightStartedByUser, startedByUser, minCreateTime, maxCreateTime, modifiedBefore, modifiedAfter, 0, 5, consistentRead, null, excludeUserJobs);
+					List<IJob> scanJobs = await _jobs.FindAsync(jobIds, streamId, name, templates, minChange, maxChange, preflightChange, preflightOnly, preflightStartedByUser, startedByUser, minCreateTime, maxCreateTime, modifiedBefore, modifiedAfter, batchState, 0, 5, consistentRead, null, excludeUserJobs);
 					if (scanJobs.Count == 0)
 					{
 						break;
diff --git a/Engine/Source/Programs/Horde/Horde.Server/Jobs/JobsController.cs b/Engine/Source/Programs/Horde/Horde.Server/Jobs/JobsController.cs
--- a/Engine/Source/Programs/Horde/Horde.Server/Jobs/JobsController.cs
+++ b/Engine/Source/Programs/Horde/Horde.Server/Jobs/JobsController.cs
@@ -775,7 +775,7 @@ namespace Horde.Server.Jobs
 
 			List<IJob> jobs;
 			jobs = await _jobService.FindJobsAsync(jobIdValues, streamIdValue, name, templateRefIds, minChange,
-				maxChange, preflightChange, preflightOnly, preflightStartedByUserIdValue, startedByUserIdValue, minCreateTime?.UtcDateTime, maxCreateTime?.UtcDateTime, target, state, outcome,
+				maxChange, preflightChange, preflightOnly, preflightStartedByUserIdValue, startedByUserIdValue, minCreateTime?.UtcDateTime, maxCreateTime?.UtcDateTime, target, null, state, outcome,
 				modifiedBefore, modifiedAfter, index, count, false);
 
 			return await CreateAuthorizedJobResponsesAsync(jobs, filter);
