Horde: Fix GET method attribute on bisection controller
#rnx

diff --git a/Engine/Source/Programs/Horde/Horde.Server/Jobs/Bisect/BisectController.cs b/Engine/Source/Programs/Horde/Horde.Server/Jobs/Bisect/BisectController.cs
--- a/Engine/Source/Programs/Horde/Horde.Server/Jobs/Bisect/BisectController.cs
+++ b/Engine/Source/Programs/Horde/Horde.Server/Jobs/Bisect/BisectController.cs
@@ -230,12 +230,12 @@ namespace Horde.Server.Jobs.Bisect
 		}
 
 		/// <summary>
-		/// Creates a new bisect task
+		/// Gets information about a bisection task
 		/// </summary>
 		/// <param name="bisectTaskId">Id of the bisect task to retrieve</param>
 		/// <param name="cancellationToken">Cancellation token for the request</param>
 		/// <returns>Id of the new bisect task</returns>
-		[HttpPost]
+		[HttpGet]
 		[Route("/api/v1/bisect/{id}")]
 		public async Task<ActionResult<GetBisectTaskResponse>> GetAsync([FromRoute] BisectTaskId bisectTaskId, CancellationToken cancellationToken)
 		{
