Horde: Fix bisect controller route parameter name
#rnx

diff --git a/Engine/Source/Programs/Horde/Horde.Server/Jobs/Bisect/BisectController.cs b/Engine/Source/Programs/Horde/Horde.Server/Jobs/Bisect/BisectController.cs
--- a/Engine/Source/Programs/Horde/Horde.Server/Jobs/Bisect/BisectController.cs
+++ b/Engine/Source/Programs/Horde/Horde.Server/Jobs/Bisect/BisectController.cs
@@ -236,7 +236,7 @@ namespace Horde.Server.Jobs.Bisect
 		/// <param name="cancellationToken">Cancellation token for the request</param>
 		/// <returns>Id of the new bisect task</returns>
 		[HttpGet]
-		[Route("/api/v1/bisect/{id}")]
+		[Route("/api/v1/bisect/{bisectTaskId}")]
 		public async Task<ActionResult<GetBisectTaskResponse>> GetAsync([FromRoute] BisectTaskId bisectTaskId, CancellationToken cancellationToken)
 		{
 			IBisectTask? bisectTask = await _bisectTaskCollection.GetAsync(bisectTaskId, cancellationToken);
@@ -267,7 +267,7 @@ namespace Horde.Server.Jobs.Bisect
 		/// <param name="cancellationToken">Cancellation token for the request</param>
 		/// <returns>Id of the new bisect task</returns>
 		[HttpPatch]
-		[Route("/api/v1/bisect/{id}")]
+		[Route("/api/v1/bisect/{bisectTaskId}")]
 		public async Task<ActionResult> UpdateAsync([FromRoute] BisectTaskId bisectTaskId, [FromBody] UpdateBisectTaskRequest request, CancellationToken cancellationToken)
 		{
 			for(; ;)
