Horde: Add command line option for installing agent's Windows service without starting it

diff --git a/Engine/Source/Programs/Horde/Horde.Agent/Commands/Service/ServiceInstall.cs b/Engine/Source/Programs/Horde/Horde.Agent/Commands/Service/ServiceInstall.cs
--- a/Engine/Source/Programs/Horde/Horde.Agent/Commands/Service/ServiceInstall.cs
+++ b/Engine/Source/Programs/Horde/Horde.Agent/Commands/Service/ServiceInstall.cs
@@ -45,6 +45,12 @@ namespace Horde.Agent.Commands.Service
 		/// </summary>
 		[CommandLine("-DotNetExecutable=")]
 		public string DotNetExecutable { get; set; } = "dotnet";
+		
+		/// <summary>
+		/// Start the service after installation
+		/// </summary>
+		[CommandLine("-Start=")]
+		public string? Start { get; set; } = "true";
 
 		/// <summary>
 		/// Runs the service indefinitely
@@ -87,14 +93,17 @@ namespace Horde.Agent.Commands.Service
 				{
 					service.SetDescription("Allows this machine to participate in a Horde farm.");
 
-					logger.LogInformation("Starting...");
-					service.Start();
-
-					WindowsServiceStatus status = service.WaitForStatusChange(WindowsServiceStatus.Starting, TimeSpan.FromSeconds(30.0));
-					if (status != WindowsServiceStatus.Running)
+					if (Start != null && Start.Equals("true", StringComparison.OrdinalIgnoreCase))
 					{
-						logger.LogError("Unable to start service (status = {Status})", status);
-						return Task.FromResult(1);
+						logger.LogInformation("Starting...");
+						service.Start();
+						
+						WindowsServiceStatus status = service.WaitForStatusChange(WindowsServiceStatus.Starting, TimeSpan.FromSeconds(30.0));
+						if (status != WindowsServiceStatus.Running)
+						{
+							logger.LogError("Unable to start service (status = {Status})", status);
+							return Task.FromResult(1);
+						}
 					}
 
 					logger.LogInformation("Done.");
