Horde: Store previously assigned Perforce server for agents in Redis, and weight strongly for preferring the the same server if still available.

diff --git a/Engine/Source/Programs/Horde/Horde.Server/Perforce/PerforceLoadBalancer.cs b/Engine/Source/Programs/Horde/Horde.Server/Perforce/PerforceLoadBalancer.cs
--- a/Engine/Source/Programs/Horde/Horde.Server/Perforce/PerforceLoadBalancer.cs
+++ b/Engine/Source/Programs/Horde/Horde.Server/Perforce/PerforceLoadBalancer.cs
@@ -9,7 +9,6 @@ using System.Net.Http;
 using System.Text.Json;
 using System.Threading;
 using System.Threading.Tasks;
-using EpicGames.Horde.Common;
 using Google.Protobuf.WellKnownTypes;
 using Horde.Server.Agents;
 using Horde.Server.Agents.Leases;
@@ -22,9 +21,12 @@ using Microsoft.Extensions.Logging;
 using Microsoft.Extensions.Options;
 using MongoDB.Bson.Serialization.Attributes;
 using MongoDB.Driver;
+using StackExchange.Redis;
 
 namespace Horde.Server.Perforce
 {
+	using Condition = EpicGames.Horde.Common.Condition;
+
 	/// <summary>
 	/// Health of a particular server
 	/// </summary>
@@ -139,6 +141,7 @@ namespace Horde.Server.Perforce
 
 		readonly MongoService _mongoService;
 		readonly GlobalsService _globalsService;
+		readonly RedisService _redisService;
 		readonly ILeaseCollection _leaseCollection;
 		readonly SingletonDocument<PerforceServerList> _serverListSingleton;
 		readonly Random _random = new Random();
@@ -149,10 +152,11 @@ namespace Horde.Server.Perforce
 		/// <summary>
 		/// Constructor
 		/// </summary>
-		public PerforceLoadBalancer(MongoService mongoService, GlobalsService globalsService, ILeaseCollection leaseCollection, IClock clock, IOptionsMonitor<GlobalConfig> globalConfig, ILogger<PerforceLoadBalancer> logger)
+		public PerforceLoadBalancer(MongoService mongoService, GlobalsService globalsService, RedisService redisService, ILeaseCollection leaseCollection, IClock clock, IOptionsMonitor<GlobalConfig> globalConfig, ILogger<PerforceLoadBalancer> logger)
 		{
 			_mongoService = mongoService;
 			_globalsService = globalsService;
+			_redisService = redisService;
 			_leaseCollection = leaseCollection;
 			_serverListSingleton = new SingletonDocument<PerforceServerList>(mongoService);
 			_globalConfig = globalConfig;
@@ -232,7 +236,7 @@ namespace Horde.Server.Perforce
 		public Task<IPerforceServer?> SelectServerAsync(PerforceCluster cluster)
 		{
 			List<string> properties = new List<string>{ "HordeServer=1" };
-			return SelectServerAsync(cluster, properties);
+			return SelectServerAsync("server", cluster, properties);
 		}
 
 		/// <summary>
@@ -243,7 +247,7 @@ namespace Horde.Server.Perforce
 		/// <returns></returns>
 		public Task<IPerforceServer?> SelectServerAsync(PerforceCluster cluster, IAgent agent)
 		{
-			return SelectServerAsync(cluster, agent.Properties);
+			return SelectServerAsync($"agent:{agent.Id}", cluster, agent.Properties);
 		}
 
 		/// <summary>
@@ -283,10 +287,11 @@ namespace Horde.Server.Perforce
 		/// <summary>
 		/// Select a Perforce server to use
 		/// </summary>
+		/// <param name="key"></param>
 		/// <param name="cluster"></param>
 		/// <param name="properties"></param>
 		/// <returns></returns>
-		public async Task<IPerforceServer?> SelectServerAsync(PerforceCluster cluster, IReadOnlyList<string> properties)
+		public async Task<IPerforceServer?> SelectServerAsync(string key, PerforceCluster cluster, IReadOnlyList<string> properties)
 		{
 			// Find all the valid servers for this agent
 			List<PerforceServer> validServers = new List<PerforceServer>();
@@ -335,19 +340,22 @@ namespace Horde.Server.Perforce
 				return null;
 			}
 
+			// Get the previously selected server for this key
+			RedisKey redisKey = new RedisKey($"perforce-lb:{key}");
+			string? prevServerAndPort = await _redisService.GetDatabase().StringGetAsync(redisKey);
+			PerforceServerEntry? prevCandidate = candidates.FirstOrDefault(x => x.ServerAndPort.Equals(prevServerAndPort, StringComparison.Ordinal));
+
 			// Select which server to use with a weighted average of the number of active leases
 			int index = 0;
 			if (candidates.Count > 1)
 			{
-				const int BaseWeight = 20;
-
 				int totalLeases = candidates.Sum(x => x.NumLeases);
-				int totalWeight = candidates.Sum(x => (totalLeases + BaseWeight) - x.NumLeases);
+				int totalWeight = candidates.Sum(x => GetWeight(x, prevCandidate, totalLeases));
 
 				int weight = _random.Next(totalWeight);
 				for (; index + 1 < candidates.Count; index++)
 				{
-					weight -= (totalLeases + BaseWeight) - candidates[index].NumLeases;
+					weight -= GetWeight(candidates[index], prevCandidate, totalLeases);
 					if(weight < 0)
 					{
 						break;
@@ -355,7 +363,38 @@ namespace Horde.Server.Perforce
 				}
 			}
 
-			return candidates[index];
+			// Update redis with the chosen candidate
+			PerforceServerEntry nextCandidate = candidates[index];
+			await _redisService.GetDatabase().StringSetAsync(redisKey, candidates[index].ServerAndPort, expiry: TimeSpan.FromDays(3.0), keepTtl: false, flags: CommandFlags.FireAndForget);
+
+			if (prevCandidate == null)
+			{
+				_logger.LogDebug("Adding new preferred server for {Key}: {ServerAndPort}", key, nextCandidate.ServerAndPort);
+			}
+			else if (prevCandidate == nextCandidate)
+			{
+				_logger.LogDebug("Reusing preferred server for {Key}: {ServerAndPort}", key, nextCandidate.ServerAndPort);
+			}
+			else
+			{
+				_logger.LogDebug("Changing preferred server for {Key}: {PrevServerAndPort} -> {NextServerAndPort}", key, prevCandidate.ServerAndPort, nextCandidate.ServerAndPort);
+			}
+
+			return nextCandidate;
+		}
+
+		static int GetWeight(PerforceServerEntry candidate, PerforceServerEntry? prevCandidate, int totalLeases)
+		{
+			int baseWeight;
+			if (candidate == prevCandidate)
+			{
+				baseWeight = 200;
+			}
+			else
+			{
+				baseWeight = 20;
+			}
+			return baseWeight + (totalLeases - candidate.NumLeases);
 		}
 
 		/// <inheritdoc/>
