Horde: Always keep agents using the same Perforce server for a day if it's still healthy, and only rely on randomization/weighting after that. Frequently used agents were switching servers more than rarely used agents.

diff --git a/Engine/Source/Programs/Horde/Horde.Server/Perforce/PerforceLoadBalancer.cs b/Engine/Source/Programs/Horde/Horde.Server/Perforce/PerforceLoadBalancer.cs
--- a/Engine/Source/Programs/Horde/Horde.Server/Perforce/PerforceLoadBalancer.cs
+++ b/Engine/Source/Programs/Horde/Horde.Server/Perforce/PerforceLoadBalancer.cs
@@ -345,10 +345,12 @@ namespace Horde.Server.Perforce
 			string? prevServerAndPort = await _redisService.GetDatabase().StringGetAsync(redisKey);
 			PerforceServerEntry? prevCandidate = candidates.FirstOrDefault(x => x.ServerAndPort.Equals(prevServerAndPort, StringComparison.Ordinal));
 
-			// Select which server to use with a weighted average of the number of active leases
-			int index = 0;
-			if (candidates.Count > 1)
+			// Use the previous server by default, if it's still available. If we don't have one, select which to use with a weighted average of the number of active leases
+			PerforceServerEntry? nextCandidate = prevCandidate;
+			if (nextCandidate == null)
 			{
+				int index = 0;
+
 				int totalLeases = candidates.Sum(x => x.NumLeases);
 				int totalWeight = candidates.Sum(x => GetWeight(x, prevCandidate, totalLeases));
 
@@ -361,11 +363,12 @@ namespace Horde.Server.Perforce
 						break;
 					}
 				}
+
+				nextCandidate = candidates[index];
 			}
 
-			// Update redis with the chosen candidate
-			PerforceServerEntry nextCandidate = candidates[index];
-			await _redisService.GetDatabase().StringSetAsync(redisKey, candidates[index].ServerAndPort, expiry: TimeSpan.FromDays(3.0), keepTtl: false, flags: CommandFlags.FireAndForget);
+			// Update redis with the chosen candidate. Only allow selecting a new server every day.
+			await _redisService.GetDatabase().StringSetAsync(redisKey, nextCandidate.ServerAndPort, expiry: TimeSpan.FromDays(1.0), keepTtl: true, flags: CommandFlags.FireAndForget);
 
 			if (prevCandidate == null)
 			{
