Horde: Remove methods which assume the hash is stored with a node handle.

diff --git a/Engine/Source/Programs/Horde/Horde.Server/Artifacts/ArtifactsController.cs b/Engine/Source/Programs/Horde/Horde.Server/Artifacts/ArtifactsController.cs
--- a/Engine/Source/Programs/Horde/Horde.Server/Artifacts/ArtifactsController.cs
+++ b/Engine/Source/Programs/Horde/Horde.Server/Artifacts/ArtifactsController.cs
@@ -240,10 +240,10 @@ namespace Horde.Server.Artifacts
 
 			TreeReader reader = await _storageService.GetReaderAsync(artifact.NamespaceId, cancellationToken);
 
-			NodeRef<DirectoryNode> directoryRef;
+			DirectoryNode directoryNode;
 			try
 			{
-				directoryRef = await reader.ReadNodeRefAsync<DirectoryNode>(artifact.RefName, DateTime.UtcNow.AddHours(1.0), cancellationToken);
+				directoryNode = await reader.ReadNodeAsync<DirectoryNode>(artifact.RefName, DateTime.UtcNow.AddHours(1.0), cancellationToken);
 			}
 			catch (RefNameNotFoundException)
 			{
@@ -254,17 +254,15 @@ namespace Horde.Server.Artifacts
 			{
 				foreach (string fragment in path.Split('/'))
 				{
-					DirectoryNode nextDirectory = await directoryRef.ExpandAsync(reader, cancellationToken);
-					if (!nextDirectory.TryGetDirectoryEntry(fragment, out DirectoryEntry? nextDirectoryEntry))
+					DirectoryEntry? nextDirectoryEntry;
+					if (!directoryNode.TryGetDirectoryEntry(fragment, out nextDirectoryEntry))
 					{
 						return NotFound();
 					}
-					directoryRef = nextDirectoryEntry;
+					directoryNode = await nextDirectoryEntry.ExpandAsync(reader, cancellationToken);
 				}
 			}
 
-			DirectoryNode directoryNode = await directoryRef.ExpandAsync(reader, cancellationToken);
-
 			GetArtifactDirectoryResponse response = new GetArtifactDirectoryResponse();
 			await ExpandDirectoriesAsync(reader, directoryNode, 0, response, cancellationToken);
 			return PropertyFilter.Apply(response, filter);
@@ -334,8 +332,7 @@ namespace Horde.Server.Artifacts
 			}
 
 			TreeReader reader = await _storageService.GetReaderAsync(artifact.NamespaceId, cancellationToken);
-			NodeRef<DirectoryNode> directoryRef = await reader.ReadNodeRefAsync<DirectoryNode>(artifact.RefName, DateTime.UtcNow.AddHours(1.0), cancellationToken);
-			DirectoryNode directory = await directoryRef.ExpandAsync(reader, cancellationToken);
+			DirectoryNode directory = await reader.ReadNodeAsync<DirectoryNode>(artifact.RefName, DateTime.UtcNow.AddHours(1.0), cancellationToken);
 
 			FileEntry? fileEntry = await directory.GetFileEntryByPathAsync(reader, path, cancellationToken);
 			if (fileEntry == null)
@@ -407,8 +404,7 @@ namespace Horde.Server.Artifacts
 			}
 
 			TreeReader reader = await _storageService.GetReaderAsync(artifact.NamespaceId, cancellationToken);
-			NodeRef<DirectoryNode> directoryRef = await reader.ReadNodeRefAsync<DirectoryNode>(artifact.RefName, DateTime.UtcNow.AddHours(1.0), cancellationToken);
-			DirectoryNode directory = await directoryRef.ExpandAsync(reader, cancellationToken);
+			DirectoryNode directory = await reader.ReadNodeAsync<DirectoryNode>(artifact.RefName, DateTime.UtcNow.AddHours(1.0), cancellationToken);
 
 			Stream stream = directory.AsZipStream(reader, filter);
 			return new FileStreamResult(stream, "application/zip") { FileDownloadName = $"{artifact.RefName}.zip" };
diff --git a/Engine/Source/Programs/Horde/Horde.Server/Perforce/PerforceReplicator.cs b/Engine/Source/Programs/Horde/Horde.Server/Perforce/PerforceReplicator.cs
--- a/Engine/Source/Programs/Horde/Horde.Server/Perforce/PerforceReplicator.cs
+++ b/Engine/Source/Programs/Horde/Horde.Server/Perforce/PerforceReplicator.cs
@@ -518,17 +518,15 @@ namespace Horde.Server.Perforce
 
 		static async Task<CommitNode?> FindParentAsync(TreeReader reader, RefName refName, int change, CancellationToken cancellationToken)
 		{
-			NodeRef<CommitNode>? parentRef = await reader.TryReadNodeRefAsync<CommitNode>(refName, cancellationToken: cancellationToken);
-			while (parentRef != null)
+			CommitNode? commitNode = await reader.TryReadNodeAsync<CommitNode>(refName, cancellationToken: cancellationToken);
+			if (commitNode != null)
 			{
-				CommitNode parent = await parentRef.ExpandCopyAsync(reader, cancellationToken);
-				if (parent.Number < change)
+				while (commitNode.Parent != null && commitNode.Number >= change)
 				{
-					return parent;
+					commitNode = await commitNode.Parent.ExpandAsync(reader, cancellationToken);
 				}
-				parentRef = parent.Parent;
 			}
-			return null;
+			return commitNode;
 		}
 
 		static int GetFileOffset(Utf8String path)
diff --git a/Engine/Source/Programs/Shared/EpicGames.Horde/Storage/Backends/HttpStorageClient.cs b/Engine/Source/Programs/Shared/EpicGames.Horde/Storage/Backends/HttpStorageClient.cs
--- a/Engine/Source/Programs/Shared/EpicGames.Horde/Storage/Backends/HttpStorageClient.cs
+++ b/Engine/Source/Programs/Shared/EpicGames.Horde/Storage/Backends/HttpStorageClient.cs
@@ -282,7 +282,7 @@ namespace EpicGames.Horde.Storage.Backends
 			_logger.LogDebug("Writing ref {RefName} -> {RefTarget}", name, target);
 			using (HttpClient httpClient = _createClient())
 			{
-				using (HttpResponseMessage response = await httpClient.PutAsync($"refs/{name}", new { hash = target.Hash, blob = target.Locator.Blob, exportIdx = target.Locator.ExportIdx, options }, cancellationToken))
+				using (HttpResponseMessage response = await httpClient.PutAsync($"refs/{name}", new { blob = target.Locator.Blob, exportIdx = target.Locator.ExportIdx, options }, cancellationToken))
 				{
 					response.EnsureSuccessStatusCode();
 				}
diff --git a/Engine/Source/Programs/Shared/EpicGames.Horde/Storage/TreeReader.cs b/Engine/Source/Programs/Shared/EpicGames.Horde/Storage/TreeReader.cs
--- a/Engine/Source/Programs/Shared/EpicGames.Horde/Storage/TreeReader.cs
+++ b/Engine/Source/Programs/Shared/EpicGames.Horde/Storage/TreeReader.cs
@@ -665,23 +665,6 @@ namespace EpicGames.Horde.Storage
 			return await ReadNodeAsync<TNode>(refTarget.Locator, cancellationToken);
 		}
 
-		/// <summary>
-		/// Reads data for a ref from the store, along with the node's contents.
-		/// </summary>
-		/// <param name="name">The ref name</param>
-		/// <param name="cacheTime">Minimum coherency for any cached value to be returned</param>
-		/// <param name="cancellationToken">Cancellation token for the operation</param>
-		/// <returns>Node for the given ref, or null if it does not exist</returns>
-		public async Task<NodeRef<TNode>?> TryReadNodeRefAsync<TNode>(RefName name, DateTime cacheTime = default, CancellationToken cancellationToken = default) where TNode : Node
-		{
-			NodeHandle? refTarget = await _store.TryReadRefTargetAsync(name, cacheTime, cancellationToken);
-			if (refTarget == null)
-			{
-				return null;
-			}
-			return new NodeRef<TNode>(refTarget);
-		}
-
 		/// <summary>
 		/// Reads a ref from the store, throwing an exception if it does not exist
 		/// </summary>
@@ -698,22 +681,5 @@ namespace EpicGames.Horde.Storage
 			}
 			return refValue;
 		}
-
-		/// <summary>
-		/// Reads a ref from the store, throwing an exception if it does not exist
-		/// </summary>
-		/// <param name="name">Id for the ref</param>
-		/// <param name="cacheTime">Minimum coherency of any cached result</param>
-		/// <param name="cancellationToken">Cancellation token for the operation</param>
-		/// <returns>The blob instance</returns>
-		public async Task<NodeRef<TNode>> ReadNodeRefAsync<TNode>(RefName name, DateTime cacheTime = default, CancellationToken cancellationToken = default) where TNode : Node
-		{
-			NodeRef<TNode>? refValue = await TryReadNodeRefAsync<TNode>(name, cacheTime, cancellationToken);
-			if (refValue == null)
-			{
-				throw new RefNameNotFoundException(name);
-			}
-			return refValue;
-		}
 	}
 }
