Horde: Relax update polling for projects
#skipci
#rnx

diff --git a/Engine/Source/Programs/Horde/HordeDashboard/src/backend/Backend.ts b/Engine/Source/Programs/Horde/HordeDashboard/src/backend/Backend.ts
--- a/Engine/Source/Programs/Horde/HordeDashboard/src/backend/Backend.ts
+++ b/Engine/Source/Programs/Horde/HordeDashboard/src/backend/Backend.ts
@@ -9,7 +9,7 @@ import { projectStore } from './ProjectStore';
 
 
 // Update interval for relatively static global data such as projects, schedules, templates
-const updateInterval = 120 * 1000;
+const updateInterval = 900 * 1000;
 
 export class Backend {
 
