Horde: Include exception message when unable to submit new changes.

diff --git a/Engine/Source/Programs/Horde/Horde.Server/Perforce/PerforceService.cs b/Engine/Source/Programs/Horde/Horde.Server/Perforce/PerforceService.cs
--- a/Engine/Source/Programs/Horde/Horde.Server/Perforce/PerforceService.cs
+++ b/Engine/Source/Programs/Horde/Horde.Server/Perforce/PerforceService.cs
@@ -1010,7 +1010,7 @@ namespace Horde.Server.Perforce
 						}
 						else
 						{
-							throw new PerforceException($"Unable to submit new changelist for {filePath} after {MaxAttempts} attempts");
+							throw new PerforceException($"Unable to submit new changelist for {filePath} after {MaxAttempts} attempts: {ex.Message}");
 						}
 					}
 					finally
