Horde: Remove log line throwing an exception.

diff --git a/Engine/Source/Programs/Horde/Horde.Server/Server/MongoService.cs b/Engine/Source/Programs/Horde/Horde.Server/Server/MongoService.cs
--- a/Engine/Source/Programs/Horde/Horde.Server/Server/MongoService.cs
+++ b/Engine/Source/Programs/Horde/Horde.Server/Server/MongoService.cs
@@ -740,7 +740,7 @@ namespace Horde.Server.Server
 					}
 					else
 					{
-						_logger.LogInformation("Creating index {IndexName} in {CollectionName}: {Definition}", createIndex.Name, collectionName, createIndex.ToJson());
+						_logger.LogInformation("Creating index {IndexName} in {CollectionName}", createIndex.Name, collectionName);
 
 						CreateIndexOptions<T> options = new CreateIndexOptions<T>();
 						options.Name = createIndex.Name;
