Horde: Add a generalized error handler path that formats exceptions as log events so they can be rendered on the dashboard. Should provide much more context for internal server errors.

diff --git a/Engine/Source/Programs/Horde/Horde.Server/Startup.cs b/Engine/Source/Programs/Horde/Horde.Server/Startup.cs
--- a/Engine/Source/Programs/Horde/Horde.Server/Startup.cs
+++ b/Engine/Source/Programs/Horde/Horde.Server/Startup.cs
@@ -99,6 +99,7 @@ using OpenTelemetry.Trace;
 using Polly;
 using Polly.Extensions.Http;
 using Horde.Server.Jobs.Bisect;
+using Microsoft.AspNetCore.Diagnostics;
 
 namespace Horde.Server
 {
@@ -1066,6 +1067,8 @@ namespace Horde.Server
 				app.UseMiddleware<RequestTrackerMiddleware>();	
 			}
 
+			app.UseExceptionHandler("/api/v1/exception");
+
 			app.UseDefaultFiles();
 			app.UseStaticFiles();
 
diff --git a/Engine/Source/Programs/Horde/Horde.Server/Utilities/StructuredException.cs b/Engine/Source/Programs/Horde/Horde.Server/Utilities/StructuredException.cs
new file mode 100644
--- /dev/null
+++ b/Engine/Source/Programs/Horde/Horde.Server/Utilities/StructuredException.cs
@@ -0,0 +1,67 @@
+// Copyright Epic Games, Inc. All Rights Reserved.
+
+using System;
+using System.Collections.Generic;
+using EpicGames.Core;
+using Microsoft.Extensions.Logging;
+
+namespace Horde.Server.Utilities
+{
+	/// <summary>
+	/// Exception class designed to allow logging structured log messages
+	/// </summary>
+	public class StructuredException : Exception
+	{
+		/// <summary>
+		/// Format string for the message
+		/// </summary>
+		public string Format { get; }
+
+		/// <summary>
+		/// Arguments for holes in the format string
+		/// </summary>
+		public IReadOnlyList<object> Args { get; }
+
+		/// <summary>
+		/// Constructor
+		/// </summary>
+		/// <param name="format">Message template for the log message</param>
+		/// <param name="args">Arguments to render</param>
+		public StructuredException(string format, params object[] args)
+		{
+			Format = format;
+			Args = args;
+		}
+
+		/// <summary>
+		/// Creates a log event from this exception
+		/// </summary>
+		public LogEvent ToLogEvent()
+		{
+			return LogEvent.Create(LogLevel.Error, default, this, Format, Args);
+		}
+	}
+
+	/// <summary>
+	/// Exception class designed to allow logging structured log messages
+	/// </summary>
+	public class StructuredHttpException : StructuredException
+	{
+		/// <summary>
+		/// Status code
+		/// </summary>
+		public int StatusCode { get; }
+
+		/// <summary>
+		/// Constructor
+		/// </summary>
+		/// <param name="statusCode">Status code for the response</param>
+		/// <param name="format">Message template for the log message</param>
+		/// <param name="args">Arguments to render</param>
+		public StructuredHttpException(int statusCode, string format, params object[] args)
+			: base(format, args)
+		{
+			StatusCode = statusCode;
+		}
+	}
+}
