Horde: Fix parsing of legacy node handles with hashes included.

diff --git a/Engine/Source/Programs/Horde/Horde.Server/Logs/LogRpcService.cs b/Engine/Source/Programs/Horde/Horde.Server/Logs/LogRpcService.cs
--- a/Engine/Source/Programs/Horde/Horde.Server/Logs/LogRpcService.cs
+++ b/Engine/Source/Programs/Horde/Horde.Server/Logs/LogRpcService.cs
@@ -52,7 +52,11 @@ namespace Horde.Server.Logs
 
 			IStorageClientImpl store = await _storageService.GetClientAsync(Namespace.Logs, context.CancellationToken);
 			_logger.LogInformation("Updating {LogId} to node {RefTarget} (lines: {LineCount}, complete: {Complete})", request.LogId, request.Target, request.LineCount, request.Complete);
-			await store.WriteRefTargetAsync(new RefName(request.LogId), NodeLocator.Parse(request.Target));
+
+			// Strip off the 'Hash@' prefix inserted by older clients
+			int atIdx = request.Target.IndexOf('@', StringComparison.Ordinal);
+			string target = (atIdx == -1) ? request.Target : request.Target.Substring(atIdx + 1);
+			await store.WriteRefTargetAsync(new RefName(request.LogId), NodeLocator.Parse(target));
 
 			await _logFileCollection.UpdateLineCountAsync(logFile, request.LineCount, request.Complete, CancellationToken.None);
 
