Horde: Make agents view embeddable and embed it on the pool view
#skipci
#rnx

diff --git a/Engine/Source/Programs/Horde/HordeDashboard/src/components/AgentView.tsx b/Engine/Source/Programs/Horde/HordeDashboard/src/components/AgentView.tsx
--- a/Engine/Source/Programs/Horde/HordeDashboard/src/components/AgentView.tsx
+++ b/Engine/Source/Programs/Horde/HordeDashboard/src/components/AgentView.tsx
@@ -229,8 +229,6 @@ class LocalState {
    @observable agentContextMenuOpen = false;
    contextMenuRef: any = null;
    contextMenuTargetRef: any = createRef<Element>();
-   agentContextMenuRef: any = null;
-   agentContextMenuTargetRef: any = createRef<Element>();
    @observable mouseX = 0;
    @observable mouseY = 0;
 
@@ -275,7 +273,7 @@ class LocalState {
    @action
    setAgentId(id?: string) {
       this.searchState.agentId = id;
-      this.updateSearch();    
+      this.updateSearch();
    }
 
 
@@ -1294,7 +1292,7 @@ const agentContextMenuProps: IContextualMenuItem[] = [
 const agentStatus = ["Active", "Ready", "Disabled", "Pending Conform", "Pending Shutdown", "Offline", "Offline (Autoscaler)", "Offline (Manual)", "Offline (Unexpected)"];
 
 
-export const AgentMenuBar: React.FC = observer(() => {
+export const AgentMenuBar: React.FC<{ agentView?: boolean }> = observer(({ agentView }) => {
 
    let selectedButton: any = <div></div>;
    if (localState.agentsSelectedCount !== 0) {
@@ -1317,7 +1315,7 @@ export const AgentMenuBar: React.FC = observer(() => {
 
 
    return (
-      <Stack horizontal horizontalAlign="space-between">
+      <Stack horizontal horizontalAlign="space-between" grow={!!agentView}>
          <Stack.Item styles={{ root: { paddingLeft: '20px' } }}>
             <Stack horizontal tokens={{ childrenGap: 12 }}>
                <Stack verticalFill={true} verticalAlign="center">
@@ -1362,8 +1360,9 @@ export const AgentMenuBar: React.FC = observer(() => {
             </Stack>
 
          </Stack.Item>
-         <Stack.Item>
-            <Stack horizontal tokens={{ childrenGap: 12 }}>
+         {!!agentView && 
+            <Stack horizontal tokens={{ childrenGap: 12 }} grow>
+               <Stack grow/>
                <PrimaryButton styles={{ root: { fontFamily: "Horde Open Sans SemiBold !important" } }} text="Download Agent" onClick={() => { backend.downloadAgentZip() }} />
                <CommandButton
                   onClick={() => { editPoolsModalState.setOpen(); }}
@@ -1375,7 +1374,7 @@ export const AgentMenuBar: React.FC = observer(() => {
                <PoolEditorModal></PoolEditorModal>
                <PoolSelectionModal></PoolSelectionModal>
             </Stack>
-         </Stack.Item>
+         }
       </Stack>
    );
 });
@@ -1719,21 +1718,15 @@ export const SearchUpdate: React.FC = observer(() => {
    const csearch = localState.search.toString();
 
    if (state.search !== csearch) {
-      setSearchParams(csearch, {replace: true});
+      setSearchParams(csearch, { replace: true });
       setState({ search: csearch });
    }
 
    return null;
 });
 
-export const AgentView: React.FC = observer(() => {   
+export const AgentViewInner: React.FC<{ agentId?: string, poolId?: string, searchParams?: URLSearchParams, agentView?: boolean }> = observer(({ agentId, poolId, searchParams, agentView }) => {
    const [initAgentUpdater, setInitAgentUpdater] = useState(false);
-   const [searchParams] = useSearchParams();
-
-   const agentId = searchParams.get("agentId") ? searchParams.get("agentId") : undefined;   
-
-   // adjust automatically to viewport changes
-   useWindowSize();
 
    useEffect(() => {
       const interval = setInterval(() => {
@@ -1745,25 +1738,17 @@ export const AgentView: React.FC = observer(() => {
    if (!initAgentUpdater) {
       agentStore.update().then(() => {
          localState.resetState();
-         localState.stateFromSearch(searchParams);
+         if (searchParams) {
+            localState.stateFromSearch(searchParams);
+         }
          setInitAgentUpdater(true);
       });
-      return <Stack className={hordeClasses.horde}>
-         <TopNav />
-         <Breadcrumbs items={[{ text: 'Admin' }, { text: 'Agents' }]} />
-         <Stack horizontal>
-            <Stack grow styles={{ root: { backgroundColor: 'rgb(250, 249, 249)' } }} />
-            <Stack tokens={{ maxWidth: 1440, childrenGap: 4 }} styles={{ root: { width: 1440, height: '100vh', backgroundColor: 'rgb(250, 249, 249)', paddingTop: 18, paddingLeft: 12 } }}>
-               <Stack className={hordeClasses.raised} styles={{ root: { paddingRight: '40px' } }}>
-                  <Stack style={{ position: "relative", height: "calc(100vh - 240px)" }} horizontalAlign="center">
-                     <Stack horizontal tokens={{ childrenGap: 24 }}>
-                        <Text variant="mediumPlus">Loading Agents</Text>
-                        <Spinner size={SpinnerSize.large} />
-                     </Stack>
-                  </Stack>
-               </Stack>
+      return <Stack>
+         <Stack style={{ position: "relative", height: "calc(100vh - 240px)" }} horizontalAlign="center">
+            <Stack horizontal tokens={{ childrenGap: 24 }}>
+               <Text variant="mediumPlus">Loading Agents</Text>
+               <Spinner size={SpinnerSize.large} />
             </Stack>
-            <Stack grow styles={{ root: { backgroundColor: 'rgb(250, 249, 249)' } }} />
          </Stack>
       </Stack>
    }
@@ -1926,7 +1911,7 @@ export const AgentView: React.FC = observer(() => {
                   let diskCapacity: string | number | null = getAgentCapability(agent, "DiskTotalSize");
                   toCheck = "Unknown";
                   if (diskFree !== null && diskCapacity !== null) {
-                     diskFree = Number(diskFree) /  1073741824;
+                     diskFree = Number(diskFree) / 1073741824;
                      diskCapacity = Number(diskCapacity) / 1073741824;
                      const percentage = (diskCapacity - diskFree) / diskCapacity;
                      toCheck = (diskCapacity - (percentage * diskCapacity)).toFixed(0);
@@ -2031,6 +2016,12 @@ export const AgentView: React.FC = observer(() => {
 
    let agentItems = agentStore.agents.filter(filterAgents, localState.agentFilter).sort(sortAgents);
 
+   if (poolId) {
+      agentItems = agentItems.filter((item => {
+         return !!item.pools && item.pools.indexOf(poolId) !== -1;
+      }))
+   }
+
    agentItems = agentItems.filter(item => {
       const filter = localState.agentStatusFilter;
       if (!filter.size) {
@@ -2098,116 +2089,109 @@ export const AgentView: React.FC = observer(() => {
       return !filtered;
    });
 
-   return (
-      <Stack className={hordeClasses.horde}>
-         <TopNav />
-         <Breadcrumbs items={[{ text: 'Admin' }, { text: 'Agents' }]} />
-         <Stack horizontal>
-            <Stack grow styles={{ root: { backgroundColor: 'rgb(250, 249, 249)' } }} />
-            <Stack tokens={{ maxWidth: 1440, childrenGap: 4 }} styles={{ root: { width: 1440, height: '100vh', backgroundColor: 'rgb(250, 249, 249)', paddingTop: 18, paddingLeft: 12 } }}>
-               <Stack className={hordeClasses.raised} styles={{ root: { paddingRight: '40px' } }}>
-                  <SearchUpdate />
-                  <Stack.Item>
-                     <AgentMenuBar />
-                  </Stack.Item>
-                  <Stack style={{ position: "relative", height: "calc(100vh - 240px)" }}>
-                     <ScrollablePane scrollbarVisibility={ScrollbarVisibility.always}>
-                        <DetailsList
-                           className={agentStyles.detailsList}
-                           onItemContextMenu={onContextMenu}
-                           checkboxCellClassName={agentStyles.checkboxCell}
-                           compact={true}
-                           setKey="set"
-                           items={agentItems}
-                           columns={localState.columnsState.filter(colState => { return colState.isChecked; }).map(colState => { return colState.columnDef!; })}
-                           selection={localState.selection}
-                           onRenderDetailsHeader={onRenderDetailsHeader}
-                           onColumnHeaderContextMenu={onColumnHeaderContextMenu}
-                           onRenderItemColumn={onRenderAgentListItem}
-                           layoutMode={DetailsListLayoutMode.justified}
-                           constrainMode={ConstrainMode.horizontalConstrained}
-                           selectionMode={SelectionMode.multiple}
-                        />
-                        <ContextualMenu
-                           items={localState.columnMenuProps}
-                           onItemClick={() => localState.setHeaderContextMenuOpen(false)}
-                           onDismiss={() => localState.setHeaderContextMenuOpen(false)}
-                           isBeakVisible={true}
-                           hidden={!localState.headerContextMenuOpen}
-                           target={localState.contextMenuTargetRef}
-                           directionalHint={DirectionalHint.bottomLeftEdge}
-                           directionalHintFixed={true}
-                        />
-                        <ContextualMenu
-                           items={agentContextMenuProps}
-                           onItemClick={() => localState.setAgentContextMenuOpen(false)}
-                           onDismiss={() => localState.setAgentContextMenuOpen(false)}
-                           isBeakVisible={true}
-                           target={localState.agentContextMenuTargetRef}
-                           hidden={!localState.agentContextMenuOpen}
-                           directionalHint={DirectionalHint.bottomLeftEdge}
-                           directionalHintFixed={true}
-                        />
-                     </ScrollablePane>
-                  </Stack>
-               </Stack>
-            </Stack>
-            <Stack grow styles={{ root: { backgroundColor: 'rgb(250, 249, 249)' } }} />
-         </Stack>
-         <div hidden={!localState.agentContextMenuOpen} ref={localState.agentContextMenuTargetRef} style={{ position: 'absolute', width: 1, height: 1, left: localState.mouseX, top: localState.mouseY }}></div>
-         <ConfirmationDialog
-            title={`Delete Agent${localState.selection.getSelectedCount() > 1 ? "s" : ""}`}
-            subText={`Really delete selected agent${localState.selection.getSelectedCount() > 1 ? "s" : ""}?`}
-            isOpen={localState.deleteAgentDialogIsOpen}
-            confirmText={"Delete"}
-            cancelText={"Cancel"}
-            onConfirm={() => { localState.deleteBuilder() }}
-            onCancel={() => { localState.setDeleteBuilderDialogOpen(false) }}
-         />
-         <ConfirmationDialog
-            title={`Restart Agent${localState.selection.getSelectedCount() > 1 ? "s" : ""}`}
-            subText={`Really restart selected agent${localState.selection.getSelectedCount() > 1 ? "s" : ""}?`}
-            isOpen={localState.restartAgentDialogIsOpen}
-            confirmText={"Restart"}
-            cancelText={"Cancel"}
-            textBoxLabel={"Type Confirm to confirm"}
-            isTextBoxSpawned={true}
-            onConfirm={() => { localState.requestBuilderUpdate(false, true) }}
-            onCancel={() => { localState.setRestartBuilderDialogOpen(false) }}
-         />
-         <ConfirmationDialog
-            title={`Shutdown Agent${localState.selection.getSelectedCount() > 1 ? "s" : ""}`}
-            subText={`Really shutdown selected agent${localState.selection.getSelectedCount() > 1 ? "s" : ""}?`}
-            isOpen={localState.shutdownAgentDialogIsOpen}
-            confirmText={"Shutdown"}
-            cancelText={"Cancel"}
-            textBoxLabel={"Type Confirm to confirm"}
-            isTextBoxSpawned={true}
-            onConfirm={() => { localState.requestBuilderUpdate(false, false, false, true) }}
-            onCancel={() => { localState.setShutdownBuilderDialogOpen(false) }}
-         />
+   // Dynamically changes based on filtered items, so keep static so page doesn't dump
+   const height = 400;// Math.min((agentItems.length+ 2) * 48, 480)
 
-         <ConfirmationDialog
-            title={`Cancel Leases`}
-            subText={`Really cancel selected agent leases?`}
-            isOpen={localState.cancelLeasesDialogIsOpen}
-            confirmText={"Yes"}
-            cancelText={"No"}
-            onConfirm={() => { localState.setCancelLeasesDialogOpen(false); localState.cancelLeases() }}
-            onCancel={() => { localState.setCancelLeasesDialogOpen(false) }}
-         />
-         <ConfirmationDialog
-            title={`Disable Agent${localState.selection.getSelectedCount() > 1 ? "s" : ""}`}
-            isOpen={localState.disableAgentDialogIsOpen}
-            confirmText={"Disable"}
-            cancelText={"Cancel"}
-            textBoxLabel={"Enter Disable Reason"}
-            isTextBoxSpawned={true}
-            onConfirm={(textFieldText: string) => { localState.changeBuilderEnabled("disable", textFieldText) }}
-            onCancel={() => { localState.setDisableBuilderDialogOpen(false) }}
-         />
-         <HistoryModal agentId={activeAgent?.id} onDismiss={onHistoryModalDismiss}></HistoryModal>
+   return (<Stack>
+      {!!agentView && <SearchUpdate />}
+      <Stack horizontal style={{ paddingBottom: !agentView ? 18 : 0 }}>
+         {!agentView && <Stack grow />}
+         <AgentMenuBar agentView={agentView} />
       </Stack>
+      <Stack style={{ position: "relative", height: agentView ? "calc(100vh - 240px)" : height }}>
+         <ScrollablePane scrollbarVisibility={ScrollbarVisibility.always}>
+            <DetailsList
+               className={agentStyles.detailsList}
+               onItemContextMenu={onContextMenu}
+               checkboxCellClassName={agentStyles.checkboxCell}
+               compact={true}
+               setKey="set"
+               items={agentItems}
+               columns={localState.columnsState.filter(colState => { return colState.isChecked; }).map(colState => { return colState.columnDef!; })}
+               selection={localState.selection}
+               onRenderDetailsHeader={onRenderDetailsHeader}
+               onColumnHeaderContextMenu={onColumnHeaderContextMenu}
+               onRenderItemColumn={onRenderAgentListItem}
+               layoutMode={DetailsListLayoutMode.justified}
+               constrainMode={ConstrainMode.horizontalConstrained}
+               selectionMode={SelectionMode.multiple}
+            />
+            <ContextualMenu
+               items={localState.columnMenuProps}
+               onItemClick={() => localState.setHeaderContextMenuOpen(false)}
+               onDismiss={() => localState.setHeaderContextMenuOpen(false)}
+               isBeakVisible={true}
+               hidden={!localState.headerContextMenuOpen}
+               target={localState.contextMenuTargetRef}
+               directionalHint={DirectionalHint.bottomLeftEdge}
+               directionalHintFixed={true}
+            />
+            <ContextualMenu
+               items={agentContextMenuProps}
+               onItemClick={() => localState.setAgentContextMenuOpen(false)}
+               onDismiss={() => localState.setAgentContextMenuOpen(false)}
+               isBeakVisible={true}
+               target={{x: localState.mouseX, y: localState.mouseY }}
+               hidden={!localState.agentContextMenuOpen}
+               directionalHint={DirectionalHint.bottomLeftEdge}
+               directionalHintFixed={true}
+            />
+         </ScrollablePane>
+         <Stack grow styles={{ root: { backgroundColor: 'rgb(250, 249, 249)' } }} />
+      </Stack>
+      <ConfirmationDialog
+         title={`Delete Agent${localState.selection.getSelectedCount() > 1 ? "s" : ""}`}
+         subText={`Really delete selected agent${localState.selection.getSelectedCount() > 1 ? "s" : ""}?`}
+         isOpen={localState.deleteAgentDialogIsOpen}
+         confirmText={"Delete"}
+         cancelText={"Cancel"}
+         onConfirm={() => { localState.deleteBuilder() }}
+         onCancel={() => { localState.setDeleteBuilderDialogOpen(false) }}
+      />
+      <ConfirmationDialog
+         title={`Restart Agent${localState.selection.getSelectedCount() > 1 ? "s" : ""}`}
+         subText={`Really restart selected agent${localState.selection.getSelectedCount() > 1 ? "s" : ""}?`}
+         isOpen={localState.restartAgentDialogIsOpen}
+         confirmText={"Restart"}
+         cancelText={"Cancel"}
+         textBoxLabel={"Type Confirm to confirm"}
+         isTextBoxSpawned={true}
+         onConfirm={() => { localState.requestBuilderUpdate(false, true) }}
+         onCancel={() => { localState.setRestartBuilderDialogOpen(false) }}
+      />
+      <ConfirmationDialog
+         title={`Shutdown Agent${localState.selection.getSelectedCount() > 1 ? "s" : ""}`}
+         subText={`Really shutdown selected agent${localState.selection.getSelectedCount() > 1 ? "s" : ""}?`}
+         isOpen={localState.shutdownAgentDialogIsOpen}
+         confirmText={"Shutdown"}
+         cancelText={"Cancel"}
+         textBoxLabel={"Type Confirm to confirm"}
+         isTextBoxSpawned={true}
+         onConfirm={() => { localState.requestBuilderUpdate(false, false, false, true) }}
+         onCancel={() => { localState.setShutdownBuilderDialogOpen(false) }}
+      />
+
+      <ConfirmationDialog
+         title={`Cancel Leases`}
+         subText={`Really cancel selected agent leases?`}
+         isOpen={localState.cancelLeasesDialogIsOpen}
+         confirmText={"Yes"}
+         cancelText={"No"}
+         onConfirm={() => { localState.setCancelLeasesDialogOpen(false); localState.cancelLeases() }}
+         onCancel={() => { localState.setCancelLeasesDialogOpen(false) }}
+      />
+      <ConfirmationDialog
+         title={`Disable Agent${localState.selection.getSelectedCount() > 1 ? "s" : ""}`}
+         isOpen={localState.disableAgentDialogIsOpen}
+         confirmText={"Disable"}
+         cancelText={"Cancel"}
+         textBoxLabel={"Enter Disable Reason"}
+         isTextBoxSpawned={true}
+         onConfirm={(textFieldText: string) => { localState.changeBuilderEnabled("disable", textFieldText) }}
+         onCancel={() => { localState.setDisableBuilderDialogOpen(false) }}
+      />
+      <HistoryModal agentId={activeAgent?.id} onDismiss={onHistoryModalDismiss}></HistoryModal>
+   </Stack>
    );
 
 
@@ -2262,7 +2246,7 @@ export const AgentView: React.FC = observer(() => {
                      {getAgentStatusIcon(agent)}
                   </Stack.Item>
                   <Stack.Item align={"center"}>
-                     <ReactLink styles={{ root: { paddingTop: '1px', fontSize: "12px"} }} title="Lease and Session History" onClick={() => localState.setAgentId(agent.id) }>{agent.name}</ReactLink>
+                     <ReactLink styles={{ root: { paddingTop: '1px', fontSize: "12px" } }} title="Lease and Session History" onClick={() => localState.setAgentId(agent.id)}>{agent.name}</ReactLink>
                   </Stack.Item>
                </Stack>
             );
@@ -2534,4 +2518,40 @@ export const AgentView: React.FC = observer(() => {
             return <span>{agent[column!.fieldName as keyof AgentData] as string}</span>;
       }
    }
-});
\ No newline at end of file
+});
+
+export const AgentPanel: React.FC<{ agentId?: string, poolId?: string, agentView?: boolean }> = ({ agentId, agentView, poolId }) => {
+
+   // adjust automatically to viewport changes
+   useWindowSize();
+
+   return <Stack>
+      <AgentViewInner agentId={agentId} agentView={agentView} poolId={poolId} />
+   </Stack>
+
+}
+
+export const AgentView: React.FC = () => {
+
+   const [searchParams] = useSearchParams();
+
+   const agentId = searchParams.get("agentId") ? searchParams.get("agentId") : undefined;
+
+   // adjust automatically to viewport changes
+   useWindowSize();
+
+   return <Stack className={hordeClasses.horde}>
+      <TopNav />
+      <Breadcrumbs items={[{ text: 'Admin' }, { text: 'Agents' }]} />
+      <Stack horizontal>
+         <Stack grow styles={{ root: { backgroundColor: 'rgb(250, 249, 249)' } }} />
+         <Stack tokens={{ maxWidth: 1440, childrenGap: 4 }} styles={{ root: { width: 1440, height: '100vh', backgroundColor: 'rgb(250, 249, 249)', paddingTop: 18, paddingLeft: 12 } }}>
+            <Stack className={hordeClasses.raised} styles={{ root: { paddingRight: '40px' } }}>
+               <AgentViewInner agentView={true} agentId={agentId ? agentId : undefined} searchParams={searchParams} />
+            </Stack>
+         </Stack>
+         <Stack grow styles={{ root: { backgroundColor: 'rgb(250, 249, 249)' } }} />
+      </Stack>
+   </Stack>
+
+}
\ No newline at end of file
diff --git a/Engine/Source/Programs/Horde/HordeDashboard/src/components/PoolView.tsx b/Engine/Source/Programs/Horde/HordeDashboard/src/components/PoolView.tsx
--- a/Engine/Source/Programs/Horde/HordeDashboard/src/components/PoolView.tsx
+++ b/Engine/Source/Programs/Horde/HordeDashboard/src/components/PoolView.tsx
@@ -17,6 +17,7 @@ import { Breadcrumbs } from "./Breadcrumbs";
 import { HistoryModal } from "./HistoryModal";
 import { LeaseStatusIcon, StepStatusIcon } from "./StatusIcon";
 import { TopNav } from "./TopNav";
+import { AgentPanel } from "./AgentView";
 
 
 type PendingBatch = {
@@ -341,6 +342,18 @@ enum StepState {
    Previous = "Previous"
 }
 
+const PoolAgentPanel: React.FC<{poolId:string}> = ({poolId}) => {
+
+   return <Stack styles={{ root: { paddingTop: 18, paddingLeft: 12, paddingRight: 12 } }} >
+      <Stack tokens={{ childrenGap: 12 }}>
+         <Stack>
+            <Text variant="mediumPlus" styles={{ root: { fontFamily: "Horde Open Sans SemiBold" } }}>Agents</Text>
+         </Stack>
+         <AgentPanel poolId={poolId} />
+      </Stack>
+   </Stack>
+}
+
 const StepPanel: React.FC<{ stepState: StepState }> = ({ stepState }) => {
 
    const [lastSelectedAgent, setLastSelectedAgent] = useState<string | undefined>(undefined);
@@ -1069,9 +1082,9 @@ const PoolPanel: React.FC = () => {
 
       if (column.name === "Value") {
          if (item.name === "Agents") {
-            return <Stack horizontal tokens={{childrenGap: 4}}>
+            return <Stack horizontal tokens={{ childrenGap: 4 }}>
                <Stack>{item.value}</Stack>
-               <Stack horizontal tokens={{childrenGap: 4}}>
+               <Stack horizontal tokens={{ childrenGap: 4 }}>
                   <Stack> - </Stack>
                   <Stack><Link to={`/reports/utilization?pools=${pool.id}`}>Utilization</Link></Stack>
                </Stack>
@@ -1273,6 +1286,7 @@ export const PoolView: React.FC = observer(() => {
                         <Stack className={hordeClasses.raised} style={{ width: 1443, height: "fit-content" }}>
                            <Stack tokens={{ childrenGap: 18 }}>
                               <PoolPanel />
+                              <PoolAgentPanel poolId={poolId} />
                               <StepPanel stepState={StepState.Active} />
                               <ConformPanel />
                               <StepPanel stepState={StepState.Pending} />
