Horde: Dynamically size agents panel height based on agents in pool
#skipci
#rnx

diff --git a/Engine/Source/Programs/Horde/HordeDashboard/src/components/AgentView.tsx b/Engine/Source/Programs/Horde/HordeDashboard/src/components/AgentView.tsx
--- a/Engine/Source/Programs/Horde/HordeDashboard/src/components/AgentView.tsx
+++ b/Engine/Source/Programs/Horde/HordeDashboard/src/components/AgentView.tsx
@@ -2014,6 +2014,13 @@ export const AgentViewInner: React.FC<{ agentId?: string, poolId?: string, searc
       localState.setAgentId(undefined);
    }
 
+   const totalPoolSize = agentStore.agents.filter((a) => {
+      if (!a.pools) {
+         return false;
+      }
+      return !poolId || a.pools.indexOf(poolId) !== -1
+   }).length;
+
    let agentItems = agentStore.agents.filter(filterAgents, localState.agentFilter).sort(sortAgents);
 
    if (poolId) {
@@ -2089,8 +2096,7 @@ export const AgentViewInner: React.FC<{ agentId?: string, poolId?: string, searc
       return !filtered;
    });
 
-   // Dynamically changes based on filtered items, so keep static so page doesn't dump
-   const height = 400;// Math.min((agentItems.length+ 2) * 48, 480)
+   const height = Math.min(((totalPoolSize < 5 ? 5 : totalPoolSize) * 48) + 48, 480)
 
    return (<Stack>
       {!!agentView && <SearchUpdate />}
