Horde: Add exception debugging endpoint
#skpci
#rnx

diff --git a/Engine/Source/Programs/Horde/Horde.Server/Server/DebugController.cs b/Engine/Source/Programs/Horde/Horde.Server/Server/DebugController.cs
--- a/Engine/Source/Programs/Horde/Horde.Server/Server/DebugController.cs
+++ b/Engine/Source/Programs/Horde/Horde.Server/Server/DebugController.cs
@@ -524,6 +524,19 @@ namespace Horde.Server.Server
 			
 			return PhysicalFile(snapshotZipFile, "application/zip", Path.GetFileName(snapshotZipFile));
 		}
+
+		/// <summary>
+		/// Throws an exception to debug error handling
+		/// </summary>
+		/// <returns></returns>
+		[HttpGet]
+		[Route("/api/v1/debug/exception")]
+		public ActionResult ThrowException()
+		{
+			int numberArg = 42;
+			string stringArg = "hello";
+			throw new Exception($"Message: numberArg:{numberArg}, stringArg:{stringArg}");
+		}
 	}
 }
 
