Horde: Add supported for structured logging errors from internal server errors
#rnx
#preflght none
#skipci

diff --git a/Engine/Source/Programs/Horde/HordeDashboard/src/backend/Backend.ts b/Engine/Source/Programs/Horde/HordeDashboard/src/backend/Backend.ts
--- a/Engine/Source/Programs/Horde/HordeDashboard/src/backend/Backend.ts
+++ b/Engine/Source/Programs/Horde/HordeDashboard/src/backend/Backend.ts
@@ -692,6 +692,19 @@ export class Backend {
 
     }
 
+    testException(): Promise<void> {
+
+        return new Promise<void>((resolve, reject) => {
+
+            this.backend.get("/api/v1/debug/exception").then(() => {
+                resolve();
+            }).catch(reason => {
+                reject(reason);
+            });
+        });
+    }
+
+
     createJob(request: CreateJobRequest): Promise<CreateJobResponse> {
 
         return new Promise<CreateJobResponse>((resolve, reject) => {
diff --git a/Engine/Source/Programs/Horde/HordeDashboard/src/backend/Fetch.ts b/Engine/Source/Programs/Horde/HordeDashboard/src/backend/Fetch.ts
--- a/Engine/Source/Programs/Horde/HordeDashboard/src/backend/Fetch.ts
+++ b/Engine/Source/Programs/Horde/HordeDashboard/src/backend/Fetch.ts
@@ -39,9 +39,28 @@ export class Fetch {
         return new Promise<FetchResponse>(async (resolve, reject) => {
 
             try {
+                
+                const response = await fetch(url, this.buildRequest("GET")).then(response => {
+                    if (!response.ok && response.status === 500) {
+                        response.text().then(text => {
+                            handleError({
+                                response: response,
+                                reason: "Internal Server Error",
+                                mode: "GET",
+                                url: url,
+                                message: text
+                            }, true);            
+                        });
+                        return null;
+                    }
+                    return response;
+                });
 
-                const response = await fetch(url, this.buildRequest("GET"));
-
+                if (!response) {
+                    reject('');
+                    return;
+                }
+		
                 if (response.ok) {
                     this.handleResponse(response, url, "GET", resolve, reject, config);
                     return;
@@ -172,6 +191,19 @@ export class Fetch {
 
     private async handleResponse(response: Response, url: string, mode: string, resolve: (value: FetchResponse | PromiseLike<FetchResponse>) => void, reject: (reason?: any) => void, config?: FetchRequestConfig) {
 
+        if (!response.ok && response.status === 500) {
+
+            response.text().then(text => {
+                handleError({
+                    mode: mode,
+                    response: response,
+                    url: url,
+                    title: "Internal Server Error",
+                    message: text                
+                }, true);
+            });
+            return reject("Internal Server Error");
+        }
 
         if (response.status === 401) {
 
@@ -185,6 +217,7 @@ export class Fetch {
             if (response.url?.indexOf("AccessDenied") !== -1) {
 
                 handleError({
+                    mode: mode,
                     response: response,
                     url: url,
                     title: "Access Denied"
@@ -195,6 +228,7 @@ export class Fetch {
                 if (response.status !== 404 || !config?.suppress404) {
 
                     let errorInfo: ErrorInfo = {
+                        mode: mode,
                         response: response,
                         url: url
                     }
@@ -234,6 +268,7 @@ export class Fetch {
             }).catch(reason => {
 
                 handleError({
+                    mode: mode,
                     response: response,
                     url: url,
                     reason: reason
@@ -257,6 +292,7 @@ export class Fetch {
                 }).catch(reason => {
 
                     handleError({
+                        mode: mode,
                         response: response,
                         url: url,
                         reason: reason
diff --git a/Engine/Source/Programs/Horde/HordeDashboard/src/components/ErrorHandler.tsx b/Engine/Source/Programs/Horde/HordeDashboard/src/components/ErrorHandler.tsx
--- a/Engine/Source/Programs/Horde/HordeDashboard/src/components/ErrorHandler.tsx
+++ b/Engine/Source/Programs/Horde/HordeDashboard/src/components/ErrorHandler.tsx
@@ -30,7 +30,7 @@ export type ErrorInfo = {
    url?: string;
 
    // REST mode
-   mode?: "GET" | "PUT" | "DELETE" | "POST";
+   mode?: "GET" | "PUT" | "DELETE" | "POST" | string;
 
    // custom title override
    title?: string;
@@ -231,7 +231,7 @@ export const ErrorDialog: React.FC = observer(() => {
          dialogContentProps={dialogContentProps}>
 
          <Stack tokens={{ childrenGap: 12 }}>
-            <Text styles={{ root: { whiteSpace: "pre-wrap", fontSize: 16, color: "#EC4C47" } }}>{message}</Text>
+            <Text styles={{ root: { whiteSpace: "pre-wrap", fontSize: 13, color: "#EC4C47" } }}>{message}</Text>
             {!!error.url && <Text variant="medium">{`URL: ${error.url}`}</Text>}
             <Stack>
                <Stack tokens={{ childrenGap: 18 }} style={{ paddingTop: 12 }}>
