Horde: Do not weight towards code changes being likely suspects for buildgraph errors.

diff --git a/Engine/Source/Programs/Horde/Horde.Server/Issues/Handlers/BuildGraphIssueHandler.cs b/Engine/Source/Programs/Horde/Horde.Server/Issues/Handlers/BuildGraphIssueHandler.cs
--- a/Engine/Source/Programs/Horde/Horde.Server/Issues/Handlers/BuildGraphIssueHandler.cs
+++ b/Engine/Source/Programs/Horde/Horde.Server/Issues/Handlers/BuildGraphIssueHandler.cs
@@ -7,7 +7,6 @@ using EpicGames.Core;
 using Horde.Server.Jobs;
 using Horde.Server.Jobs.Graphs;
 using Microsoft.Extensions.Logging;
-using MongoDB.Driver;
 
 namespace Horde.Server.Issues.Handlers
 {
@@ -58,6 +57,12 @@ namespace Horde.Server.Issues.Handlers
 			}
 		}
 
+		/// <inheritdoc/>
+		public override void RankSuspects(IIssueFingerprint fingerprint, List<SuspectChange> suspects)
+		{
+			RankSuspects(fingerprint, suspects, preferCodeChanges: false);
+		}
+
 		/// <inheritdoc/>
 		public override string GetSummary(IIssueFingerprint fingerprint, IssueSeverity severity)
 		{
diff --git a/Engine/Source/Programs/Horde/Horde.Server/Issues/Handlers/ContentIssueHandler.cs b/Engine/Source/Programs/Horde/Horde.Server/Issues/Handlers/ContentIssueHandler.cs
--- a/Engine/Source/Programs/Horde/Horde.Server/Issues/Handlers/ContentIssueHandler.cs
+++ b/Engine/Source/Programs/Horde/Horde.Server/Issues/Handlers/ContentIssueHandler.cs
@@ -100,16 +100,13 @@ namespace Horde.Server.Issues.Handlers
 		{
 			foreach (SuspectChange suspect in suspects)
 			{
-				if (suspect.ContainsContent)
+				if (suspect.Files.Any(x => fingerprint.Keys.Any(y => x.Contains(y, StringComparison.OrdinalIgnoreCase))))
 				{
-					if (suspect.Files.Any(x => fingerprint.Keys.Any(y => x.Contains(y, StringComparison.OrdinalIgnoreCase))))
-					{
-						suspect.Rank += 20;
-					}
-					else
-					{
-						suspect.Rank += 10;
-					}
+					suspect.Rank += 20;
+				}
+				else if (suspect.ContainsContent)
+				{
+					suspect.Rank += 10;
 				}
 			}
 		}
diff --git a/Engine/Source/Programs/Horde/Horde.Server/Issues/Handlers/SourceFileIssueHandler.cs b/Engine/Source/Programs/Horde/Horde.Server/Issues/Handlers/SourceFileIssueHandler.cs
--- a/Engine/Source/Programs/Horde/Horde.Server/Issues/Handlers/SourceFileIssueHandler.cs
+++ b/Engine/Source/Programs/Horde/Horde.Server/Issues/Handlers/SourceFileIssueHandler.cs
@@ -79,6 +79,12 @@ namespace Horde.Server.Issues.Handlers
 
 		/// <inheritdoc/>
 		public override void RankSuspects(IIssueFingerprint fingerprint, List<SuspectChange> suspects)
+		{
+			RankSuspects(fingerprint, suspects, preferCodeChanges: true);
+		}
+
+		/// <inheritdoc cref="RankSuspects(IIssueFingerprint, List{SuspectChange})"/>
+		protected static void RankSuspects(IIssueFingerprint fingerprint, List<SuspectChange> suspects, bool preferCodeChanges)
 		{
 			List<string> fileNames = new List<string>();
 			foreach (string key in fingerprint.Keys)
@@ -95,16 +101,13 @@ namespace Horde.Server.Issues.Handlers
 
 			foreach (SuspectChange change in suspects)
 			{
-				if (change.ContainsCode)
+				if (fileNames.Any(x => change.ModifiesFile(x)))
 				{
-					if (fileNames.Any(x => change.ModifiesFile(x)))
-					{
-						change.Rank += 20;
-					}
-					else
-					{
-						change.Rank += 10;
-					}
+					change.Rank += 20;
+				}
+				else if(preferCodeChanges && change.ContainsCode)
+				{
+					change.Rank += 10;
 				}
 			}
 		}
