Horde: Fix issue with agent modal not popping up upon selection in embedded agents widget
#skipci
#rnx

diff --git a/Engine/Source/Programs/Horde/HordeDashboard/src/components/AgentView.tsx b/Engine/Source/Programs/Horde/HordeDashboard/src/components/AgentView.tsx
--- a/Engine/Source/Programs/Horde/HordeDashboard/src/components/AgentView.tsx
+++ b/Engine/Source/Programs/Horde/HordeDashboard/src/components/AgentView.tsx
@@ -2526,16 +2526,19 @@ export const AgentViewInner: React.FC<{ agentId?: string, poolId?: string, searc
    }
 });
 
-export const AgentPanel: React.FC<{ agentId?: string, poolId?: string, agentView?: boolean }> = ({ agentId, agentView, poolId }) => {
+export const AgentPanel: React.FC<{ agentId?: string, poolId?: string, agentView?: boolean }> = observer(({ agentId, agentView, poolId }) => {
 
    // adjust automatically to viewport changes
    useWindowSize();
 
+   // subscribe
+   if (localState.searchUpdated) { }
+
    return <Stack>
-      <AgentViewInner agentId={agentId} agentView={agentView} poolId={poolId} />
+      <AgentViewInner agentId={localState.searchState?.agentId} agentView={agentView} poolId={poolId} />
    </Stack>
 
-}
+})
 
 export const AgentView: React.FC = () => {
 
