Horde: Change Slack CL links to point to Swarm rather than opening UGS.

diff --git a/Engine/Source/Programs/Horde/Horde.Server/Notifications/Sinks/SlackNotificationSink.cs b/Engine/Source/Programs/Horde/Horde.Server/Notifications/Sinks/SlackNotificationSink.cs
--- a/Engine/Source/Programs/Horde/Horde.Server/Notifications/Sinks/SlackNotificationSink.cs
+++ b/Engine/Source/Programs/Horde/Horde.Server/Notifications/Sinks/SlackNotificationSink.cs
@@ -1489,7 +1489,17 @@ namespace Horde.Server.Notifications.Sinks
 			return user.Name;
 		}
 
-		static string FormatChange(int change) => $"<ugs://change?number={change}|CL {change}>";
+		string FormatChange(int change)
+		{
+			if (_settings.P4SwarmUrl != null)
+			{
+				return $"<{new Uri(_settings.P4SwarmUrl, $"changes/{change}")}|CL {change}>";
+			}
+			else
+			{
+				return $"CL {change}";
+			}
+		}
 
 		string FormatJobStep(IIssueStep step, string name) => $"<{GetJobUrl(step.JobId)}|{step.JobName}> / <{GetStepUrl(step.JobId, step.StepId)}|{name}>";
 
