Horde: Make job step timeline more dynamic for running batches
#skipci
#rnx

diff --git a/Engine/Source/Programs/Horde/HordeDashboard/src/components/jobDetailsV2/JobDetailTimeline.tsx b/Engine/Source/Programs/Horde/HordeDashboard/src/components/jobDetailsV2/JobDetailTimeline.tsx
--- a/Engine/Source/Programs/Horde/HordeDashboard/src/components/jobDetailsV2/JobDetailTimeline.tsx
+++ b/Engine/Source/Programs/Horde/HordeDashboard/src/components/jobDetailsV2/JobDetailTimeline.tsx
@@ -7,7 +7,7 @@ import { observer } from "mobx-react-lite";
 import moment from "moment";
 import { useEffect, useState } from "react";
 import { Link, NavigateFunction, useNavigate } from "react-router-dom";
-import { GetBatchResponse, GetStepResponse, JobStepOutcome, StepData } from "../../backend/Api";
+import { GetBatchResponse, GetStepResponse, JobStepBatchState, JobStepOutcome, StepData } from "../../backend/Api";
 import dashboard, { StatusColor } from "../../backend/Dashboard";
 import { ISideRailLink } from "../../base/components/SideRail";
 import { getShortNiceTime, msecToElapsed } from "../../base/utilities/timeUtils";
@@ -137,7 +137,35 @@ class TimelineDataView extends JobDataView {
 
       const job = jobDetails.jobData!;
 
-      let bresponses = job.batches?.filter(b => !!b.agentId && (!!b.startTime && !!b.readyTime && !!b.finishTime) && b.steps?.length)!.map(b => { return { ...b } as GetBatchResponse })!;
+      // patch batches with latest finished step time, this renders the timeline more dynamically than waiting for 
+      // complete batches, note: should present batch stopping time, which can be significant
+      
+      const nbatches = job.batches?.map(b => {
+
+         const newBatch = {
+            ...b,
+         } as GetBatchResponse
+
+         if (newBatch.steps.length) {
+
+            let lastStep: GetStepResponse | undefined;
+            for (let i = newBatch.steps.length - 1; i >= 0; i--) {
+               if (newBatch.steps[i].finishTime) {
+                  lastStep = newBatch.steps[i];
+                  break;
+               }
+            }
+            
+            if (lastStep) {
+               newBatch.state = JobStepBatchState.Complete
+               newBatch.finishTime = lastStep.finishTime;
+            }
+         }
+
+         return newBatch;
+      })
+
+      let bresponses = nbatches?.filter(b => !!b.agentId && (!!b.startTime && !!b.readyTime && !!b.finishTime) && b.steps?.length)!.map(b => { return { ...b } as GetBatchResponse })!;
 
       if (!this.filterStepId && jobDetails.filter) {
          const label = jobDetails.filter.label;
@@ -153,8 +181,8 @@ class TimelineDataView extends JobDataView {
             });
 
             bresponses = bresponses.filter(b => !!b.steps?.length);
-         }         
-      }      
+         }
+      }
 
       let filterStep: StepData | undefined;
 
