Horde: Add previous step link and who started the retry to step summary
#rnx
#skipci

diff --git a/Engine/Source/Programs/Horde/HordeDashboard/src/components/jobDetailsV2/StepDetailSummary.tsx b/Engine/Source/Programs/Horde/HordeDashboard/src/components/jobDetailsV2/StepDetailSummary.tsx
--- a/Engine/Source/Programs/Horde/HordeDashboard/src/components/jobDetailsV2/StepDetailSummary.tsx
+++ b/Engine/Source/Programs/Horde/HordeDashboard/src/components/jobDetailsV2/StepDetailSummary.tsx
@@ -81,6 +81,21 @@ const getStepSummaryMarkdown = (jobDetails: JobDetailsV2, stepId: string): strin
 
    };
 
+   const retries = jobDetails.getStepRetries(step.id);
+   const idx = retries.findIndex(s => s.id === step.id);
+   if (idx > 0) {
+
+      const pstep = retries[idx - 1];      
+
+      let msg = `This is a retry of a [previous step](/job/${jobDetails.jobId!}?step=${pstep.id})`;
+      if (pstep.retriedByUserInfo) {
+         msg += ` started by ${pstep.retriedByUserInfo?.name}`;
+      } 
+
+      text.push(msg);
+
+   }
+
    if (step.retriedByUserInfo) {
       const retryId = jobDetails.getRetryStepId(step.id);
       if (retryId) {
