Horde: Add bisection to step history, add method to query bisections run on a specific job, use the bisection controller for tests, split out common bisection setup for future tests
#rnx

diff --git a/Engine/Source/Programs/Horde/Horde.Server.Tests/BisectTests.cs b/Engine/Source/Programs/Horde/Horde.Server.Tests/BisectTests.cs
--- a/Engine/Source/Programs/Horde/Horde.Server.Tests/BisectTests.cs
+++ b/Engine/Source/Programs/Horde/Horde.Server.Tests/BisectTests.cs
@@ -18,17 +18,93 @@
 using Microsoft.Extensions.DependencyInjection;
 using Microsoft.VisualStudio.TestTools.UnitTesting;
 using Moq;
+using Horde.Server.Utilities;
+using Microsoft.AspNetCore.Http;
+using Microsoft.AspNetCore.Mvc;
+using System.Security.Claims;
 
 namespace Horde.Server.Tests
 {
 	[TestClass]
 	public class BisectTests : TestSetup
 	{
-		StreamId StreamId { get; } = new StreamId("ue4-main");
-		TemplateId TemplateId { get; } = new TemplateId("test-build");
-
 		[TestMethod]
 		public async Task TestStates()
 		{
+			(IJob failedJob, IGraph graph) = await SetupBisectionTest();
+
+			// Create the bisection task
+			CreateBisectTaskResponse task = Deref(await BisectTasksController!.CreateAsync(new CreateBisectTaskRequest() {JobId = failedJob.Id, NodeName = "CompileEditor" }));
+
+			GetBisectTaskResponse response = Deref(await BisectTasksController!.GetAsync(task!.BisectTaskId));
+			Assert.AreEqual(failedJob.TemplateId, response.TemplateId);
+			Assert.AreEqual(failedJob.StreamId, response.StreamId);
+			Assert.AreEqual("CompileEditor", response.NodeName);
+			Assert.AreEqual(JobStepOutcome.Failure, response.Outcome);
+			Assert.AreEqual("TestUser", response.Owner.Name);
+			Assert.AreEqual(failedJob.Change, response.InitialChange);
+			Assert.AreEqual(failedJob.Id, response.InitialJobId);
+			Assert.AreEqual(failedJob.Change, response.CurrentChange);
+			Assert.AreEqual(failedJob.Id, response.CurrentJobId);
+
+			Assert.AreEqual(task.BisectTaskId, response.Id);
+			List<IJob> jobs = await JobCollection.FindBisectTaskJobsAsync(response!.Id, null).ToListAsync();
+			Assert.AreEqual(0, jobs.Count);
+
+			// Start bisection service
+			BisectService bisectService = ServiceProvider.GetRequiredService<BisectService>();
+			await bisectService.StartAsync(CancellationToken.None);
+
+			await Clock.AdvanceAsync(TimeSpan.FromMinutes(30));
+
+			GetBisectTaskResponse bisectTask = Deref(await BisectTasksController!.GetAsync(task!.BisectTaskId));
+			Assert.AreEqual(BisectTaskState.Running, bisectTask.State);
+			jobs = await JobCollection.FindBisectTaskJobsAsync(bisectTask.Id, running: true).ToListAsync();
+			Assert.AreEqual(1, jobs.Count);
+			Assert.AreEqual(14, jobs[0].Change);
+			await SetJobOutcome(jobs[0], graph, JobStepOutcome.Failure);
+
+			await Clock.AdvanceAsync(TimeSpan.FromMinutes(30));
+
+			bisectTask = Deref(await BisectTasksController!.GetAsync(task!.BisectTaskId));
+			Assert.AreEqual(BisectTaskState.Running, bisectTask.State);
+			jobs = await JobCollection.FindBisectTaskJobsAsync(bisectTask.Id, running: true).ToListAsync();
+			Assert.AreEqual(1, jobs.Count);
+			Assert.AreEqual(12, jobs[0].Change);
+			await SetJobOutcome(jobs[0], graph, JobStepOutcome.Success);
+
+			await Clock.AdvanceAsync(TimeSpan.FromMinutes(30));
+
+			bisectTask = Deref(await BisectTasksController!.GetAsync(task!.BisectTaskId));
+			Assert.AreEqual(BisectTaskState.Running, bisectTask.State);
+			jobs = await JobCollection.FindBisectTaskJobsAsync(bisectTask.Id, running: true).ToListAsync();
+			Assert.AreEqual(1, jobs.Count);
+			Assert.AreEqual(13, jobs[0].Change);
+			await SetJobOutcome(jobs[0], graph, JobStepOutcome.Success);
+
+			await Clock.AdvanceAsync(TimeSpan.FromMinutes(30));
+
+			bisectTask = Deref(await BisectTasksController!.GetAsync(task!.BisectTaskId));
+			Assert.AreEqual(BisectTaskState.Succeeded, bisectTask.State);
+			Assert.AreEqual(20, bisectTask.InitialChange);
+			Assert.AreEqual(14, bisectTask.CurrentChange);
+			jobs = await JobCollection.FindBisectTaskJobsAsync(bisectTask.Id, running: true).ToListAsync();
+			Assert.AreEqual(0, jobs.Count);
+
+			await Clock.AdvanceAsync(TimeSpan.FromMinutes(30));
+
+			List<GetBisectTaskResponse> jobTasks = Deref(await BisectTasksController!.GetJobBisectTasksAsync(failedJob.Id));
+			Assert.AreEqual(1, jobTasks.Count);
+			Assert.AreEqual(failedJob.Id, jobTasks[0].InitialJobId);
+			Assert.AreEqual(3, jobTasks[0].Steps.Count);
+			Assert.AreEqual(BisectTaskState.Succeeded, jobTasks[0].State);
+
+		}
+
+		StreamId StreamId { get; } = new StreamId("ue4-main");
+		TemplateId TemplateId { get; } = new TemplateId("test-build");
+
+		async Task<(IJob failedJob, IGraph graph)> SetupBisectionTest()
+		{
 			IUser user = await UserCollection.FindOrAddUserByLoginAsync("Bob");
 
@@ -69,60 +145,7 @@
 			IJob succeededJob = await CreateJob(10, graph, JobStepOutcome.Success, options);
 			IJob failedJob = await CreateJob(20, graph, JobStepOutcome.Failure, options);
 
-			IBisectTaskCollection bisectTaskCollection = ServiceProvider.GetRequiredService<IBisectTaskCollection>();
-
-			IBisectTask bisectTask = await bisectTaskCollection.CreateAsync(failedJob, "CompileEditor", JobStepOutcome.Failure, UserId.Anonymous);
-			Assert.AreEqual(failedJob.TemplateId, bisectTask.TemplateId);
-			Assert.AreEqual(failedJob.StreamId, bisectTask.StreamId);
-			Assert.AreEqual("CompileEditor", bisectTask.NodeName);
-			Assert.AreEqual(JobStepOutcome.Failure, bisectTask.Outcome);
-			Assert.AreEqual(UserId.Anonymous, bisectTask.OwnerId);
-			Assert.AreEqual(failedJob.Change, bisectTask.InitialChange);
-			Assert.AreEqual(failedJob.Id, bisectTask.InitialJobId);
-			Assert.AreEqual(failedJob.Change, bisectTask.CurrentChange);
-			Assert.AreEqual(failedJob.Id, bisectTask.CurrentJobId);
-
-			List<IJob> jobs = await JobCollection.FindBisectTaskJobsAsync(bisectTask.Id, null).ToListAsync();
-			Assert.AreEqual(0, jobs.Count);
-
-			BisectService bisectService = ServiceProvider.GetRequiredService<BisectService>();
-			await bisectService.StartAsync(CancellationToken.None);
-
-			await Clock.AdvanceAsync(TimeSpan.FromMinutes(30));
-
-			bisectTask = Deref(await bisectTaskCollection.GetAsync(bisectTask.Id));
-			Assert.AreEqual(BisectTaskState.Running, bisectTask.State);
-			jobs = await JobCollection.FindBisectTaskJobsAsync(bisectTask.Id, running: true).ToListAsync();
-			Assert.AreEqual(1, jobs.Count);
-			Assert.AreEqual(14, jobs[0].Change);
-			await SetJobOutcome(jobs[0], graph, JobStepOutcome.Failure);
-
-			await Clock.AdvanceAsync(TimeSpan.FromMinutes(30));
-
-			bisectTask = Deref(await bisectTaskCollection.GetAsync(bisectTask.Id));
-			Assert.AreEqual(BisectTaskState.Running, bisectTask.State);
-			jobs = await JobCollection.FindBisectTaskJobsAsync(bisectTask.Id, running: true).ToListAsync();
-			Assert.AreEqual(1, jobs.Count);
-			Assert.AreEqual(12, jobs[0].Change);
-			await SetJobOutcome(jobs[0], graph, JobStepOutcome.Success);
-
-			await Clock.AdvanceAsync(TimeSpan.FromMinutes(30));
-
-			bisectTask = Deref(await bisectTaskCollection.GetAsync(bisectTask.Id));
-			Assert.AreEqual(BisectTaskState.Running, bisectTask.State);
-			jobs = await JobCollection.FindBisectTaskJobsAsync(bisectTask.Id, running: true).ToListAsync();
-			Assert.AreEqual(1, jobs.Count);
-			Assert.AreEqual(13, jobs[0].Change);
-			await SetJobOutcome(jobs[0], graph, JobStepOutcome.Success);
-
-			await Clock.AdvanceAsync(TimeSpan.FromMinutes(30));
-
-			bisectTask = Deref(await bisectTaskCollection.GetAsync(bisectTask.Id));
-			Assert.AreEqual(BisectTaskState.Succeeded, bisectTask.State);
-			Assert.AreEqual(20, bisectTask.InitialChange);
-			Assert.AreEqual(14, bisectTask.CurrentChange);
-			jobs = await JobCollection.FindBisectTaskJobsAsync(bisectTask.Id, running: true).ToListAsync();
-			Assert.AreEqual(0, jobs.Count);
+			return (failedJob, graph);
 		}
 
 		async Task<IJob> CreateJob(int change, IGraph graph, JobStepOutcome outcome, CreateJobOptions options)
@@ -146,9 +169,34 @@
 
 			INodeGroup group = graph.Groups[batch.GroupIdx];
 			INode node = group.Nodes[step.NodeIdx];
-			await JobStepRefCollection.InsertOrReplaceAsync(new JobStepRefId(job.Id, batch.Id, step.Id), job.Name, node.Name, job.StreamId, job.TemplateId, job.Change, LogId.GenerateNewId(), null, null, outcome, false, null, null, 0.0f, 0.0f, DateTime.MinValue, DateTime.MinValue, DateTime.MinValue);
+			await JobStepRefCollection.InsertOrReplaceAsync(new JobStepRefId(job.Id, batch.Id, step.Id), job.Name, node.Name, job.StreamId, job.TemplateId, job.Change, LogId.GenerateNewId(), null, null, outcome, false, null, null, 0.0f, 0.0f, DateTime.MinValue, DateTime.MinValue, DateTime.MinValue, job.StartedByBisectTaskId);
 
 			return job;
 		}
+
+		BisectTasksController? _bisectTasksController;
+
+		// setup BisectTaskController with a test user
+		private new BisectTasksController BisectTasksController
+		{
+			get
+			{
+				if (_bisectTasksController == null)
+				{
+					IUser user = UserCollection.FindOrAddUserByLoginAsync("TestUser").Result;
+					_bisectTasksController = base.BisectTasksController;
+					ControllerContext controllerContext = new ControllerContext();
+					controllerContext.HttpContext = new DefaultHttpContext();
+					controllerContext.HttpContext.User = new ClaimsPrincipal(new ClaimsIdentity(
+						new List<Claim> { HordeClaims.AdminClaim.ToClaim(),
+						new Claim(ClaimTypes.Name, "TestUser"),
+						new Claim(HordeClaimTypes.UserId, user.Id.ToString()) }
+						, "TestAuthType"));
+					_bisectTasksController.ControllerContext = controllerContext;
+
+				}
+				return _bisectTasksController;
+			}
+		}
 	}
 }
diff --git a/Engine/Source/Programs/Horde/Horde.Server.Tests/DatabaseIntegrationTest.cs b/Engine/Source/Programs/Horde/Horde.Server.Tests/DatabaseIntegrationTest.cs
--- a/Engine/Source/Programs/Horde/Horde.Server.Tests/DatabaseIntegrationTest.cs
+++ b/Engine/Source/Programs/Horde/Horde.Server.Tests/DatabaseIntegrationTest.cs
@@ -5,6 +5,7 @@ using System.Collections.Generic;
 using System.Net;
 using Horde.Server.Server;
 using Horde.Server.Utilities;
+using Microsoft.AspNetCore.Mvc;
 using Microsoft.Extensions.DependencyInjection;
 using Microsoft.Extensions.Logging;
 using Microsoft.Extensions.Options;
@@ -224,5 +225,11 @@ namespace Horde.Server.Tests
 			Assert.IsNotNull(item);
 			return item!;
 		}
+
+		public static T Deref<T>(ActionResult<T>? item) where T : class
+		{
+			Assert.IsNotNull(item?.Value);
+			return item!.Value!;
+		}
 	}
 }
\ No newline at end of file
diff --git a/Engine/Source/Programs/Horde/Horde.Server.Tests/IssueServiceTests.cs b/Engine/Source/Programs/Horde/Horde.Server.Tests/IssueServiceTests.cs
--- a/Engine/Source/Programs/Horde/Horde.Server.Tests/IssueServiceTests.cs
+++ b/Engine/Source/Programs/Horde/Horde.Server.Tests/IssueServiceTests.cs
@@ -255,7 +255,7 @@ namespace Horde.Server.Tests
 
 			JobStepRefId jobStepRefId = new JobStepRefId(job.Id, batch.Id, step.Id);
 			string nodeName = _graph.Groups[batch.GroupIdx].Nodes[step.NodeIdx].Name;
-			await JobStepRefCollection.InsertOrReplaceAsync(jobStepRefId, "TestJob", nodeName, job.StreamId, job.TemplateId, job.Change, step.LogId, null, null, outcome, job.UpdateIssues, null, null, 0.0f, 0.0f, job.CreateTimeUtc, step.StartTimeUtc!.Value, step.StartTimeUtc);
+			await JobStepRefCollection.InsertOrReplaceAsync(jobStepRefId, "TestJob", nodeName, job.StreamId, job.TemplateId, job.Change, step.LogId, null, null, outcome, job.UpdateIssues, null, null, 0.0f, 0.0f, job.CreateTimeUtc, step.StartTimeUtc!.Value, step.StartTimeUtc, null);
 
 			if (job.UpdateIssues)
 			{
diff --git a/Engine/Source/Programs/Horde/Horde.Server.Tests/TestSetup.cs b/Engine/Source/Programs/Horde/Horde.Server.Tests/TestSetup.cs
--- a/Engine/Source/Programs/Horde/Horde.Server.Tests/TestSetup.cs
+++ b/Engine/Source/Programs/Horde/Horde.Server.Tests/TestSetup.cs
@@ -95,6 +95,7 @@ namespace Horde.Server.Tests
 		public IUserCollection UserCollection => ServiceProvider.GetRequiredService<IUserCollection>();
 		public IDeviceCollection DeviceCollection => ServiceProvider.GetRequiredService<IDeviceCollection>();
 		public IDashboardPreviewCollection DashboardPreviewCollection => ServiceProvider.GetRequiredService<IDashboardPreviewCollection>();
+		public IBisectTaskCollection BisectTaskCollection => ServiceProvider.GetRequiredService<IBisectTaskCollection>();
 
 		public AclService AclService => ServiceProvider.GetRequiredService<AclService>();
 		public FleetService FleetService => ServiceProvider.GetRequiredService<FleetService>();
@@ -131,6 +132,7 @@ namespace Horde.Server.Tests
 		public DevicesController DevicesController => GetDevicesController();
 		public DashboardController DashboardController => GetDashboardController();
 		public TestDataController TestDataController => GetTestDataController();
+		public BisectTasksController BisectTasksController => GetBisectTasksController();
 
 		public OpenTelemetry.Trace.Tracer Tracer => ServiceProvider.GetRequiredService<OpenTelemetry.Trace.Tracer>();
 		public Meter Meter => ServiceProvider.GetRequiredService<Meter>();
@@ -314,6 +316,13 @@ namespace Horde.Server.Tests
 			return dataCtrl;
 		}
 
+		private BisectTasksController GetBisectTasksController()
+		{
+			BisectTasksController bisectCtrl = new BisectTasksController(BisectTaskCollection, JobCollection, JobStepRefCollection, GraphCollection, UserCollection, GlobalConfigSnapshot);
+			bisectCtrl.ControllerContext = GetControllerContext();
+			return bisectCtrl;
+		}
+
 		private AgentsController GetAgentsController()
 		{
 			AgentsController agentCtrl = new AgentsController(AgentService, GlobalConfigSnapshot);
diff --git a/Engine/Source/Programs/Horde/Horde.Server/Jobs/Bisect/BisectController.cs b/Engine/Source/Programs/Horde/Horde.Server/Jobs/Bisect/BisectController.cs
--- a/Engine/Source/Programs/Horde/Horde.Server/Jobs/Bisect/BisectController.cs
+++ b/Engine/Source/Programs/Horde/Horde.Server/Jobs/Bisect/BisectController.cs
@@ -2,6 +2,7 @@
 
 using System;
 using System.Collections.Generic;
+using System.Linq;
 using System.Threading;
 using System.Threading.Tasks;
 using Horde.Server.Acls;
@@ -103,10 +104,17 @@ namespace Horde.Server.Jobs.Bisect
 		/// <inheritdoc cref="IBisectTask.CurrentChange"/>
 		public int CurrentChange => _bisectTask.CurrentChange;
 
-		internal GetBisectTaskResponse(IBisectTask bisectTask, GetThinUserInfoResponse owner)
+		/// <summary>
+		/// The steps involved in the bisection
+		/// </summary>
+		public List<GetJobStepRefResponse> Steps;
+
+		internal GetBisectTaskResponse(IBisectTask bisectTask, GetThinUserInfoResponse owner, List<IJobStepRef> steps)
 		{
 			_bisectTask = bisectTask;
 			Owner = owner;
+			Steps = steps.Select(s => new GetJobStepRefResponse(s)).ToList();
+
 		}
 	}
 
@@ -152,6 +160,7 @@ namespace Horde.Server.Jobs.Bisect
 	{
 		readonly IBisectTaskCollection _bisectTaskCollection;
 		readonly IJobCollection _jobCollection;
+		readonly IJobStepRefCollection _jobStepRefs;
 		readonly IGraphCollection _graphCollection;
 		readonly IUserCollection _userCollection;
 		readonly IOptionsSnapshot<GlobalConfig> _globalConfig;
@@ -159,10 +168,11 @@ namespace Horde.Server.Jobs.Bisect
 		/// <summary>
 		/// Constructor
 		/// </summary>
-		public BisectTasksController(IBisectTaskCollection bisectTaskCollection, IJobCollection jobCollection, IGraphCollection graphCollection, IUserCollection userCollection, IOptionsSnapshot<GlobalConfig> globalConfig)
+		public BisectTasksController(IBisectTaskCollection bisectTaskCollection, IJobCollection jobCollection, IJobStepRefCollection jobStepRefs, IGraphCollection graphCollection, IUserCollection userCollection, IOptionsSnapshot<GlobalConfig> globalConfig)
 		{
 			_bisectTaskCollection = bisectTaskCollection;
 			_jobCollection = jobCollection;
+			_jobStepRefs = jobStepRefs;
 			_graphCollection = graphCollection;
 			_userCollection = userCollection;
 			_globalConfig = globalConfig;
@@ -176,7 +186,7 @@ namespace Horde.Server.Jobs.Bisect
 		/// <returns>Id of the new bisect task</returns>
 		[HttpPost]
 		[Route("/api/v1/bisect")]
-		public async Task<ActionResult<CreateBisectTaskResponse>> CreateAsync([FromBody] CreateBisectTaskRequest create, CancellationToken cancellationToken)
+		public async Task<ActionResult<CreateBisectTaskResponse>> CreateAsync([FromBody] CreateBisectTaskRequest create, CancellationToken cancellationToken = default)
 		{
 			IJob? job = await _jobCollection.GetAsync(create.JobId);
 			if (job == null)
@@ -237,7 +247,7 @@ namespace Horde.Server.Jobs.Bisect
 		/// <returns>Id of the new bisect task</returns>
 		[HttpGet]
 		[Route("/api/v1/bisect/{bisectTaskId}")]
-		public async Task<ActionResult<GetBisectTaskResponse>> GetAsync([FromRoute] BisectTaskId bisectTaskId, CancellationToken cancellationToken)
+		public async Task<ActionResult<GetBisectTaskResponse>> GetAsync([FromRoute] BisectTaskId bisectTaskId, CancellationToken cancellationToken = default)
 		{
 			IBisectTask? bisectTask = await _bisectTaskCollection.GetAsync(bisectTaskId, cancellationToken);
 			if (bisectTask == null)
@@ -245,6 +255,12 @@ namespace Horde.Server.Jobs.Bisect
 				return NotFound(bisectTaskId);
 			}
 
+			IJob? initialJob = await _jobCollection.GetAsync(bisectTask.InitialJobId);
+			if (initialJob == null)
+			{
+				return NotFound(bisectTask.InitialJobId);
+			}
+
 			StreamConfig? streamConfig;
 			if (!_globalConfig.Value.TryGetStream(bisectTask.StreamId, out streamConfig))
 			{
@@ -256,7 +272,9 @@ namespace Horde.Server.Jobs.Bisect
 			}
 
 			IUser? user = await _userCollection.GetUserAsync(bisectTask.OwnerId);
-			return new GetBisectTaskResponse(bisectTask, new GetThinUserInfoResponse(user));
+			List<IJobStepRef> steps = await _jobStepRefs.GetStepsForNodeAsync(initialJob.StreamId, initialJob.TemplateId, bisectTask.NodeName, null, true, 1024, bisectTask.Id, cancellationToken);
+
+			return new GetBisectTaskResponse(bisectTask, new GetThinUserInfoResponse(user), steps);
 		}
 
 		/// <summary>
@@ -302,5 +320,50 @@ namespace Horde.Server.Jobs.Bisect
 				}
 			}
 		}
+
+		/// <summary>
+		/// Gets the bisections run on a specific job
+		/// </summary>
+		/// <param name="jobId">Id of the job to retrieve</param>		
+		/// <param name="cancellationToken" />
+		/// <returns>Id of the new bisect task</returns>
+		[HttpGet]
+		[Route("/api/v1/bisect/job/{jobId}")]
+		public async Task<ActionResult<List<GetBisectTaskResponse>>> GetJobBisectTasksAsync([FromRoute] JobId jobId, CancellationToken cancellationToken = default)
+		{
+			IJob? job = await _jobCollection.GetAsync(jobId);
+			if (job == null)
+			{
+				return NotFound(jobId);
+			}
+
+			IReadOnlyList<IBisectTask> tasks = await _bisectTaskCollection.FindAsync(jobId, cancellationToken);
+
+			List<GetBisectTaskResponse> response = new List<GetBisectTaskResponse>();
+
+			foreach (IBisectTask bisectTask in tasks)
+			{
+				StreamConfig? streamConfig;
+				if (!_globalConfig.Value.TryGetStream(bisectTask.StreamId, out streamConfig))
+				{
+					return NotFound(bisectTask.StreamId);
+				}
+				if (!streamConfig.Authorize(BisectTaskAclAction.ViewBisectTask, User))
+				{
+					return Forbid(BisectTaskAclAction.ViewBisectTask, streamConfig.Id);
+				}
+
+
+
+				IUser? user = await _userCollection.GetUserAsync(bisectTask.OwnerId);
+				List<IJobStepRef> steps = await _jobStepRefs.GetStepsForNodeAsync(job.StreamId, job.TemplateId, bisectTask.NodeName, null, true, 1024, bisectTask.Id, cancellationToken);
+
+				response.Add(new GetBisectTaskResponse(bisectTask, new GetThinUserInfoResponse(user), steps));
+				
+			}
+
+			return response;
+		}
 	}
+
 }
diff --git a/Engine/Source/Programs/Horde/Horde.Server/Jobs/Bisect/BisectService.cs b/Engine/Source/Programs/Horde/Horde.Server/Jobs/Bisect/BisectService.cs
--- a/Engine/Source/Programs/Horde/Horde.Server/Jobs/Bisect/BisectService.cs
+++ b/Engine/Source/Programs/Horde/Horde.Server/Jobs/Bisect/BisectService.cs
@@ -174,7 +174,7 @@ namespace Horde.Server.Jobs.Bisect
 			// Query jobs before the start position
 			for (int maxJobCount = 20; ; maxJobCount += 20)
 			{
-				List<IJobStepRef> jobStepRefs = await _jobStepRefCollection.GetStepsForNodeAsync(bisectTask.StreamId, bisectTask.TemplateId, bisectTask.NodeName, bisectTask.InitialChange, true, maxJobCount, cancellationToken);
+				List<IJobStepRef> jobStepRefs = await _jobStepRefCollection.GetStepsForNodeAsync(bisectTask.StreamId, bisectTask.TemplateId, bisectTask.NodeName, bisectTask.InitialChange, true, maxJobCount, null, cancellationToken);
 
 				IJobStepRef? currentJobStepRef = null;
 				foreach(IJobStepRef jobStepRef in jobStepRefs.OrderByDescending(x => x.Change).ThenBy(x => x.Id))
diff --git a/Engine/Source/Programs/Horde/Horde.Server/Jobs/Bisect/BisectTaskCollection.cs b/Engine/Source/Programs/Horde/Horde.Server/Jobs/Bisect/BisectTaskCollection.cs
--- a/Engine/Source/Programs/Horde/Horde.Server/Jobs/Bisect/BisectTaskCollection.cs
+++ b/Engine/Source/Programs/Horde/Horde.Server/Jobs/Bisect/BisectTaskCollection.cs
@@ -80,7 +80,7 @@ namespace Horde.Server.Jobs.Bisect
 		public BisectTaskCollection(MongoService mongoService)
 		{
 			List<MongoIndex<BisectTaskDoc>> indexes = new List<MongoIndex<BisectTaskDoc>>();
-			indexes.Add(keys => keys.Ascending(x => x.Id).Ascending(x => x.Running), sparse: true);
+			indexes.Add(keys => keys.Ascending(x => x.Id).Ascending(x => x.Running).Ascending(x => x.InitialJobId), sparse: true);
 
 			_bisectTasks = mongoService.GetCollection<BisectTaskDoc>("BisectTasks", indexes);
 		}
@@ -143,6 +143,19 @@ namespace Horde.Server.Jobs.Bisect
 			return await _bisectTasks.Find(x => x.Id == bisectTaskId).FirstOrDefaultAsync(cancellationToken);
 		}
 
+		/// <inheritdoc/>
+		public async Task<IReadOnlyList<IBisectTask>> FindAsync(JobId? jobId = null, CancellationToken cancellationToken = default)
+		{
+			// Find all the bisection tasks matching the given criteria
+			FilterDefinitionBuilder<BisectTaskDoc> filterBuilder = Builders<BisectTaskDoc>.Filter;
+
+			FilterDefinition<BisectTaskDoc> filter = FilterDefinition<BisectTaskDoc>.Empty;
+			filter &= filterBuilder.Eq(x => x.InitialJobId, jobId);
+
+			List<BisectTaskDoc> steps = await _bisectTasks.Find(filter).SortByDescending(x => x.InitialChange).ToListAsync(cancellationToken);
+			return steps.ConvertAll<IBisectTask>(x => x);
+		}
+
 		/// <inheritdoc/>
 		public async Task<IBisectTask?> TryUpdateAsync(IBisectTask bisectTask, UpdateBisectTaskOptions options, CancellationToken cancellationToken = default)
 		{
diff --git a/Engine/Source/Programs/Horde/Horde.Server/Jobs/Bisect/IBisectTaskCollection.cs b/Engine/Source/Programs/Horde/Horde.Server/Jobs/Bisect/IBisectTaskCollection.cs
--- a/Engine/Source/Programs/Horde/Horde.Server/Jobs/Bisect/IBisectTaskCollection.cs
+++ b/Engine/Source/Programs/Horde/Horde.Server/Jobs/Bisect/IBisectTaskCollection.cs
@@ -33,6 +33,14 @@ namespace Horde.Server.Jobs.Bisect
 		/// <returns>List of bisect tasks</returns>
 		public IAsyncEnumerable<IBisectTask> FindActiveAsync(CancellationToken cancellationToken = default);
 
+		/// <summary>
+		/// Finds all bisect tasks for the provided criteria
+		/// </summary>
+		/// <param name="jobId">jobId of initial bisection</param>
+		/// <param name="cancellationToken">Cancellation token for the operation</param>
+		/// <returns>List of bisect tasks</returns>
+		public Task<IReadOnlyList<IBisectTask>> FindAsync(JobId? jobId = null, CancellationToken cancellationToken = default);
+
 		/// <summary>
 		/// Gets a bisect task by id
 		/// </summary>
diff --git a/Engine/Source/Programs/Horde/Horde.Server/Jobs/IJobStepRef.cs b/Engine/Source/Programs/Horde/Horde.Server/Jobs/IJobStepRef.cs
--- a/Engine/Source/Programs/Horde/Horde.Server/Jobs/IJobStepRef.cs
+++ b/Engine/Source/Programs/Horde/Horde.Server/Jobs/IJobStepRef.cs
@@ -4,6 +4,7 @@ using System;
 using System.Collections.Generic;
 using Horde.Server.Agents;
 using Horde.Server.Agents.Pools;
+using Horde.Server.Jobs.Bisect;
 using Horde.Server.Logs;
 using Horde.Server.Streams;
 using HordeCommon;
@@ -171,6 +172,11 @@ namespace Horde.Server.Jobs
 		/// </summary>
 		public IReadOnlyList<int>? IssueIds { get; }
 
+		/// <summary>
+		/// Whether this step is part of a bisection
+		/// </summary>
+		public BisectTaskId? BisectTaskId { get; }
+
 		/// <summary>
 		/// The last change that succeeded. Note that this is only set when the ref is updated; it is not necessarily consistent with steps run later.
 		/// </summary>
diff --git a/Engine/Source/Programs/Horde/Horde.Server/Jobs/IJobStepRefCollection.cs b/Engine/Source/Programs/Horde/Horde.Server/Jobs/IJobStepRefCollection.cs
--- a/Engine/Source/Programs/Horde/Horde.Server/Jobs/IJobStepRefCollection.cs
+++ b/Engine/Source/Programs/Horde/Horde.Server/Jobs/IJobStepRefCollection.cs
@@ -6,6 +6,7 @@ using System.Threading;
 using System.Threading.Tasks;
 using Horde.Server.Agents;
 using Horde.Server.Agents.Pools;
+using Horde.Server.Jobs.Bisect;
 using Horde.Server.Jobs.Graphs;
 using Horde.Server.Logs;
 using Horde.Server.Streams;
@@ -40,7 +41,8 @@ namespace Horde.Server.Jobs
 		/// <param name="jobStartTimeUtc">Start time of the job</param>
 		/// <param name="startTimeUtc">Start time</param>
 		/// <param name="finishTimeUtc">Finish time for the step, if known</param>
-		Task<IJobStepRef> InsertOrReplaceAsync(JobStepRefId id, string jobName, string stepName, StreamId streamId, TemplateId templateId, int change, LogId? logId, PoolId? poolId, AgentId? agentId, JobStepOutcome? outcome, bool updateIssues, int? lastSuccess, int? lastWarning, float waitTime, float initTime, DateTime jobStartTimeUtc, DateTime startTimeUtc, DateTime? finishTimeUtc);
+		/// <param name="bisectTaskId">The bisection task id, if part of a bisection</param>
+		Task<IJobStepRef> InsertOrReplaceAsync(JobStepRefId id, string jobName, string stepName, StreamId streamId, TemplateId templateId, int change, LogId? logId, PoolId? poolId, AgentId? agentId, JobStepOutcome? outcome, bool updateIssues, int? lastSuccess, int? lastWarning, float waitTime, float initTime, DateTime jobStartTimeUtc, DateTime startTimeUtc, DateTime? finishTimeUtc, BisectTaskId? bisectTaskId);
 
 		/// <summary>
 		/// Updates a job step ref 
@@ -70,9 +72,10 @@ namespace Horde.Server.Jobs
 		/// <param name="change">The current change</param>
 		/// <param name="includeFailed">Whether to include failed nodes</param>
 		/// <param name="maxCount">Number of results to return</param>
+		/// <param name="bisectTaskId">The bisection task to get nodes for</param>
 		/// <param name="cancellationToken">Cancellation token for the operation</param>
 		/// <returns>List of step references</returns>
-		Task<List<IJobStepRef>> GetStepsForNodeAsync(StreamId streamId, TemplateId templateId, string nodeName, int? change, bool includeFailed, int maxCount, CancellationToken cancellationToken = default);
+		Task<List<IJobStepRef>> GetStepsForNodeAsync(StreamId streamId, TemplateId templateId, string nodeName, int? change, bool includeFailed, int maxCount, BisectTaskId? bisectTaskId = null, CancellationToken cancellationToken = default);
 
 		/// <summary>
 		/// Gets the previous job that ran a given step
@@ -142,7 +145,7 @@ namespace Horde.Server.Jobs
 					logger.LogInformation("Updating step reference {StepId} for job {JobId}, batch {BatchId}, with outcome {JobStepOutcome}", step.Id, job.Id, batch.Id, outcome);
 				}				
 
-				await jobStepRefs.InsertOrReplaceAsync(new JobStepRefId(job.Id, batch.Id, step.Id), job.Name, nodeName, job.StreamId, job.TemplateId, job.Change, step.LogId, batch.PoolId, batch.AgentId, outcome, job.UpdateIssues, lastSuccess, lastWarning, waitTime, initTime, job.CreateTimeUtc, step.StartTimeUtc ?? DateTime.UtcNow, step.FinishTimeUtc);
+				await jobStepRefs.InsertOrReplaceAsync(new JobStepRefId(job.Id, batch.Id, step.Id), job.Name, nodeName, job.StreamId, job.TemplateId, job.Change, step.LogId, batch.PoolId, batch.AgentId, outcome, job.UpdateIssues, lastSuccess, lastWarning, waitTime, initTime, job.CreateTimeUtc, step.StartTimeUtc ?? DateTime.UtcNow, step.FinishTimeUtc, job.StartedByBisectTaskId);
 			}
 		}
 	}
diff --git a/Engine/Source/Programs/Horde/Horde.Server/Jobs/JobStepRefCollection.cs b/Engine/Source/Programs/Horde/Horde.Server/Jobs/JobStepRefCollection.cs
--- a/Engine/Source/Programs/Horde/Horde.Server/Jobs/JobStepRefCollection.cs
+++ b/Engine/Source/Programs/Horde/Horde.Server/Jobs/JobStepRefCollection.cs
@@ -7,6 +7,7 @@ using System.Threading;
 using System.Threading.Tasks;
 using Horde.Server.Agents;
 using Horde.Server.Agents.Pools;
+using Horde.Server.Jobs.Bisect;
 using Horde.Server.Logs;
 using Horde.Server.Server;
 using Horde.Server.Streams;
@@ -69,13 +70,16 @@ namespace Horde.Server.Jobs
 			[BsonIgnoreIfNull]
 			public List<int>? IssueIds { get; set; }
 
+			[BsonElement("btid"), BsonIgnoreIfNull]
+			public BisectTaskId? BisectTaskId { get; set; }
+
 			DateTime IJobStepRef.StartTimeUtc => StartTimeUtc ?? StartTime?.UtcDateTime ?? default;
 			DateTime? IJobStepRef.FinishTimeUtc => FinishTimeUtc ?? FinishTime?.UtcDateTime;
 			string IJobStepRef.NodeName => Name;
 			bool IJobStepRef.UpdateIssues => UpdateIssues ?? false;			
 			IReadOnlyList<int>? IJobStepRef.IssueIds => IssueIds;
 
-			public JobStepRef(JobStepRefId id, string jobName, string nodeName, StreamId streamId, TemplateId templateId, int change, LogId? logId, PoolId? poolId, AgentId? agentId, JobStepOutcome? outcome, bool updateIssues, int? lastSuccess, int? lastWarning, float batchWaitTime, float batchInitTime, DateTime jobStartTimeUtc, DateTime startTimeUtc, DateTime? finishTimeUtc)
+			public JobStepRef(JobStepRefId id, string jobName, string nodeName, StreamId streamId, TemplateId templateId, int change, LogId? logId, PoolId? poolId, AgentId? agentId, JobStepOutcome? outcome, bool updateIssues, int? lastSuccess, int? lastWarning, float batchWaitTime, float batchInitTime, DateTime jobStartTimeUtc, DateTime startTimeUtc, DateTime? finishTimeUtc, BisectTaskId? bisectTaskId)
 			{
 				Id = id;
 				JobName = jobName;
@@ -95,6 +99,7 @@ namespace Horde.Server.Jobs
 				JobStartTimeUtc = jobStartTimeUtc;
 				StartTimeUtc = startTimeUtc;
 				FinishTimeUtc = finishTimeUtc;
+				BisectTaskId = bisectTaskId;
 			}
 		}
 
@@ -113,9 +118,9 @@ namespace Horde.Server.Jobs
 		}
 
 		/// <inheritdoc/>
-		public async Task<IJobStepRef> InsertOrReplaceAsync(JobStepRefId id, string jobName, string stepName, StreamId streamId, TemplateId templateId, int change, LogId? logId, PoolId? poolId, AgentId? agentId, JobStepOutcome? outcome, bool updateIssues, int? lastSuccess, int? lastWarning, float waitTime, float initTime, DateTime jobStartTimeUtc, DateTime startTimeUtc, DateTime? finishTimeUtc)
+		public async Task<IJobStepRef> InsertOrReplaceAsync(JobStepRefId id, string jobName, string stepName, StreamId streamId, TemplateId templateId, int change, LogId? logId, PoolId? poolId, AgentId? agentId, JobStepOutcome? outcome, bool updateIssues, int? lastSuccess, int? lastWarning, float waitTime, float initTime, DateTime jobStartTimeUtc, DateTime startTimeUtc, DateTime? finishTimeUtc, BisectTaskId? bisectTaskId)
 		{
-			JobStepRef newJobStepRef = new JobStepRef(id, jobName, stepName, streamId, templateId, change, logId, poolId, agentId, outcome, updateIssues, lastSuccess, lastWarning, waitTime, initTime, jobStartTimeUtc, startTimeUtc, finishTimeUtc);
+			JobStepRef newJobStepRef = new JobStepRef(id, jobName, stepName, streamId, templateId, change, logId, poolId, agentId, outcome, updateIssues, lastSuccess, lastWarning, waitTime, initTime, jobStartTimeUtc, startTimeUtc, finishTimeUtc, bisectTaskId);
 			await _jobStepRefs.ReplaceOneAsync(Builders<JobStepRef>.Filter.Eq(x => x.Id, newJobStepRef.Id), newJobStepRef, new ReplaceOptions { IsUpsert = true });
 
 			if (_telemetrySink.Enabled)
@@ -175,7 +180,7 @@ namespace Horde.Server.Jobs
 		}
 
 		/// <inheritdoc/>
-		public async Task<List<IJobStepRef>> GetStepsForNodeAsync(StreamId streamId, TemplateId templateId, string nodeName, int? change, bool includeFailed, int maxCount, CancellationToken cancellationToken)
+		public async Task<List<IJobStepRef>> GetStepsForNodeAsync(StreamId streamId, TemplateId templateId, string nodeName, int? change, bool includeFailed, int maxCount, BisectTaskId? bisectTaskId, CancellationToken cancellationToken)
 		{
 			// Find all the steps matching the given criteria
 			FilterDefinitionBuilder<JobStepRef> filterBuilder = Builders<JobStepRef>.Filter;
@@ -192,6 +197,10 @@ namespace Horde.Server.Jobs
 			{
 				filter &= filterBuilder.Ne(x => x.Outcome, JobStepOutcome.Failure);
 			}
+			if (bisectTaskId != null)
+			{
+				filter &= filterBuilder.Eq(x => x.BisectTaskId, bisectTaskId.Value);
+			}
 
 			List<JobStepRef> steps = await _jobStepRefs.Find(filter).SortByDescending(x => x.Change).ThenByDescending(x => x.StartTimeUtc).Limit(maxCount).ToListAsync(cancellationToken);
 			return steps.ConvertAll<IJobStepRef>(x => x);
