Horde: Fix issue with agents in agent view not showing up due to case
#skipci
#rnx

diff --git a/Engine/Source/Programs/Horde/HordeDashboard/src/components/AgentView.tsx b/Engine/Source/Programs/Horde/HordeDashboard/src/components/AgentView.tsx
--- a/Engine/Source/Programs/Horde/HordeDashboard/src/components/AgentView.tsx
+++ b/Engine/Source/Programs/Horde/HordeDashboard/src/components/AgentView.tsx
@@ -1726,6 +1726,9 @@ export const SearchUpdate: React.FC = observer(() => {
 });
 
 export const AgentViewInner: React.FC<{ agentId?: string, poolId?: string, searchParams?: URLSearchParams, agentView?: boolean }> = observer(({ agentId, poolId, searchParams, agentView }) => {
+
+   poolId = poolId?.toLowerCase();
+
    const [initAgentUpdater, setInitAgentUpdater] = useState(false);
 
    useEffect(() => {
@@ -2025,7 +2028,7 @@ export const AgentViewInner: React.FC<{ agentId?: string, poolId?: string, searc
 
    if (poolId) {
       agentItems = agentItems.filter((item => {
-         return !!item.pools && item.pools.indexOf(poolId) !== -1;
+         return !!item.pools && item.pools.indexOf(poolId!) !== -1;
       }))
    }
 
