Horde: Add an endpoint to return the public IP/port of the current connection.

diff --git a/Engine/Source/Programs/Horde/Horde.Server/Server/ServerController.cs b/Engine/Source/Programs/Horde/Horde.Server/Server/ServerController.cs
--- a/Engine/Source/Programs/Horde/Horde.Server/Server/ServerController.cs
+++ b/Engine/Source/Programs/Horde/Horde.Server/Server/ServerController.cs
@@ -72,6 +72,20 @@ namespace Horde.Server.Server
 			return new GetServerInfoResponse(agentVersion);
 		}
 
+		/// <summary>
+		/// Gets connection information
+		/// </summary>
+		[HttpGet]
+		[AllowAnonymous]
+		[Route("/api/v1/server/connection")]
+		public ActionResult<GetConnectionResponse> GetConnection()
+		{
+			GetConnectionResponse response = new GetConnectionResponse();
+			response.Ip = HttpContext.Connection.RemoteIpAddress?.ToString();
+			response.Port = HttpContext.Connection.RemotePort;
+			return response;
+		}
+
 		/// <summary>
 		/// Gets ports used by the server
 		/// </summary>
diff --git a/Engine/Source/Programs/Horde/Horde.Server/Server/ServerMessages.cs b/Engine/Source/Programs/Horde/Horde.Server/Server/ServerMessages.cs
--- a/Engine/Source/Programs/Horde/Horde.Server/Server/ServerMessages.cs
+++ b/Engine/Source/Programs/Horde/Horde.Server/Server/ServerMessages.cs
@@ -39,6 +39,22 @@ namespace Horde.Server.Server
 		}
 	}
 
+	/// <summary>
+	/// Gets connection information to the server
+	/// </summary>
+	public class GetConnectionResponse
+	{
+		/// <summary>
+		/// Public IP address of the remote machine
+		/// </summary>
+		public string? Ip { get; set; }
+
+		/// <summary>
+		/// Public port of the connecting machine
+		/// </summary>
+		public int Port { get; set; }
+	}
+
 	/// <summary>
 	/// Gets ports configured for this server
 	/// </summary>
