Horde: Make JobStepRefId comparable and add to bisect unit test case
#rnx

diff --git a/Engine/Source/Programs/Horde/Horde.Server.Tests/BisectTests.cs b/Engine/Source/Programs/Horde/Horde.Server.Tests/BisectTests.cs
--- a/Engine/Source/Programs/Horde/Horde.Server.Tests/BisectTests.cs
+++ b/Engine/Source/Programs/Horde/Horde.Server.Tests/BisectTests.cs
@@ -143,6 +143,8 @@ namespace Horde.Server.Tests
 			options.Arguments.Add($"{IJob.TargetArgumentPrefix}CompileEditor");
 
 			IJob succeededJob = await CreateJob(10, graph, JobStepOutcome.Success, options);
+			// create an interim job at the same change as the one that will be bisected
+			await CreateJob(20, graph, JobStepOutcome.Failure, options);
 			IJob failedJob = await CreateJob(20, graph, JobStepOutcome.Failure, options);
 
 			return (failedJob, graph);
diff --git a/Engine/Source/Programs/Horde/Horde.Server/Jobs/IJobStepRef.cs b/Engine/Source/Programs/Horde/Horde.Server/Jobs/IJobStepRef.cs
--- a/Engine/Source/Programs/Horde/Horde.Server/Jobs/IJobStepRef.cs
+++ b/Engine/Source/Programs/Horde/Horde.Server/Jobs/IJobStepRef.cs
@@ -17,7 +17,7 @@ namespace Horde.Server.Jobs
 	/// Unique id struct for JobStepRef objects. Includes a job id, batch id, and step id to uniquely identify the step.
 	/// </summary>
 	[BsonSerializer(typeof(JobStepRefIdSerializer))]
-	public struct JobStepRefId
+	public struct JobStepRefId : IComparable<JobStepRefId>
 	{
 		/// <summary>
 		/// The job id
@@ -72,6 +72,24 @@ namespace Horde.Server.Jobs
 		{
 			return JobId.Equals(other.JobId) && BatchId.Equals(other.BatchId) && StepId.Equals(other.StepId);
 		}
+
+		/// <inheritdoc/>
+		public int CompareTo(JobStepRefId other)
+		{
+			int result = JobId.Id.CompareTo(other.JobId.Id);
+			if (result != 0)
+			{
+				return result;
+			}
+
+			result = BatchId.Value.CompareTo(other.BatchId.Value);
+			if (result != 0)
+			{
+				return result;
+			}
+
+			return StepId.Value.CompareTo(other.StepId.Value);
+		}
 	}
 
 	/// <summary>
