Horde: Add missing property method for bisection steps
#rnx

diff --git a/Engine/Source/Programs/Horde/Horde.Server/Jobs/Bisect/BisectController.cs b/Engine/Source/Programs/Horde/Horde.Server/Jobs/Bisect/BisectController.cs
--- a/Engine/Source/Programs/Horde/Horde.Server/Jobs/Bisect/BisectController.cs
+++ b/Engine/Source/Programs/Horde/Horde.Server/Jobs/Bisect/BisectController.cs
@@ -107,7 +107,7 @@ namespace Horde.Server.Jobs.Bisect
 		/// <summary>
 		/// The steps involved in the bisection
 		/// </summary>
-		public List<GetJobStepRefResponse> Steps;
+		public List<GetJobStepRefResponse> Steps { get; }
 
 		internal GetBisectTaskResponse(IBisectTask bisectTask, GetThinUserInfoResponse owner, List<IJobStepRef> steps)
 		{
