Horde: Guard against null "user" properties in UGS metadata.

diff --git a/Engine/Source/Programs/Horde/Horde.Server/Ugs/UgsMetadataCollection.cs b/Engine/Source/Programs/Horde/Horde.Server/Ugs/UgsMetadataCollection.cs
--- a/Engine/Source/Programs/Horde/Horde.Server/Ugs/UgsMetadataCollection.cs
+++ b/Engine/Source/Programs/Horde/Horde.Server/Ugs/UgsMetadataCollection.cs
@@ -37,6 +37,7 @@ namespace Horde.Server.Ugs
 			[BsonIgnoreIfNull]
 			public string? Comment { get; set; }
 
+			string IUgsUserData.User => User ?? "Unknown";
 			long? IUgsUserData.SyncTime => SyncTime;
 			UgsUserVote IUgsUserData.Vote => Vote ?? UgsUserVote.None;
 
@@ -144,7 +145,7 @@ namespace Horde.Server.Ugs
 			{
 				UgsMetadataDocument document = (UgsMetadataDocument)metadata;
 
-				int userIdx = document.Users.FindIndex(x => x.User.Equals(userName, StringComparison.OrdinalIgnoreCase));
+				int userIdx = document.Users.FindIndex(x => x.User != null && x.User.Equals(userName, StringComparison.OrdinalIgnoreCase));
 				if (userIdx == -1)
 				{
 					// Create a new user entry
