Horde: Handle AWS auto-scaling events to ensure agents are terminated gracefully without interrupting leases

diff --git a/Engine/Source/Programs/Horde/Horde.Server.Tests/Fleet/AwsAutoScalingLifecycleTest.cs b/Engine/Source/Programs/Horde/Horde.Server.Tests/Fleet/AwsAutoScalingLifecycleTest.cs
new file mode 100644
--- /dev/null
+++ b/Engine/Source/Programs/Horde/Horde.Server.Tests/Fleet/AwsAutoScalingLifecycleTest.cs
@@ -0,0 +1,168 @@
+// Copyright Epic Games, Inc. All Rights Reserved.
+
+using System;
+using System.Collections.Generic;
+using System.Net;
+using System.Text.Json;
+using System.Threading;
+using System.Threading.Tasks;
+using Amazon.AutoScaling;
+using Amazon.AutoScaling.Model;
+using Horde.Server.Agents;
+using Horde.Server.Agents.Fleet;
+using Horde.Server.Agents.Pools;
+using HordeCommon;
+using Microsoft.Extensions.DependencyInjection;
+using Microsoft.Extensions.Logging;
+using Microsoft.VisualStudio.TestTools.UnitTesting;
+using Moq;
+
+namespace Horde.Server.Tests.Fleet;
+
+[TestClass]
+public class LifecycleHookMessageTest
+{
+	[TestMethod]
+	public void DeserializeLifecycleHookRequest()
+	{
+		string rawText = @"
+{
+	""version"": ""0"",
+	""id"": ""d2b0a72e-03a4-11ee-be56-0242ac120002"",
+	""detail-type"": ""EC2 Instance-terminate Lifecycle Action"",
+	""source"": ""aws.autoscaling"",
+	""account"": ""111111111111"",
+	""time"": ""2023-06-05T10:48:08Z"",
+	""region"": ""us-east-1"",
+	""resources"": [""arn:aws:autoscaling:us-east-1:111111111111:autoScalingGroup:d816d436-03a4-11ee-be56-0242ac120002:autoScalingGroupName/my-asg-name""],
+	""detail"": {
+		""LifecycleActionToken"": ""f191f274-03a4-11ee-be56-0242ac120002"",
+		""AutoScalingGroupName"": ""my-asg-name"",
+		""LifecycleHookName"": ""my-hook-name"",
+		""EC2InstanceId"": ""i-11234567890abcdef"",
+		""LifecycleTransition"": ""autoscaling:EC2_INSTANCE_TERMINATING"",
+		""Origin"": ""AutoScalingGroup"",
+		""Destination"": ""EC2""
+	}
+}
+";
+
+		LifecycleActionRequest? req = JsonSerializer.Deserialize<LifecycleActionRequest>(rawText);
+		Assert.AreEqual("f191f274-03a4-11ee-be56-0242ac120002", req!.Event.LifecycleActionToken);
+		Assert.AreEqual("my-asg-name", req!.Event.AutoScalingGroupName);
+		Assert.AreEqual("my-hook-name", req!.Event.LifecycleHookName);
+		Assert.AreEqual("i-11234567890abcdef", req!.Event.Ec2InstanceId);
+		Assert.AreEqual("autoscaling:EC2_INSTANCE_TERMINATING", req!.Event.LifecycleTransition);
+		Assert.AreEqual("AutoScalingGroup", req!.Event.Origin);
+	}
+}
+
+[TestClass]
+public class AwsAutoScalingLifecycleServiceTest : TestSetup
+{
+	private AwsAutoScalingLifecycleService _asgLifecycleService = default!;
+	private Mock<IAmazonAutoScaling> _asgMock = default!;
+	private readonly List<CompleteLifecycleActionRequest> _lifecycleUpdates = new();
+	
+	[TestInitialize]
+	public async Task Setup()
+	{
+		ILogger<AwsAutoScalingLifecycleService> logger = ServiceProvider.GetRequiredService<ILogger<AwsAutoScalingLifecycleService>>();
+		_asgMock = new Mock<IAmazonAutoScaling>(MockBehavior.Strict);
+
+		Task<CompleteLifecycleActionResponse> OnCompleteLifecycleActionAsync(CompleteLifecycleActionRequest request, CancellationToken cancellationToken)
+		{
+			_lifecycleUpdates.Add(request);
+			return Task.FromResult(new CompleteLifecycleActionResponse() { HttpStatusCode = HttpStatusCode.OK });
+		}
+
+		_asgMock
+			.Setup(x => x.CompleteLifecycleActionAsync(It.IsAny<CompleteLifecycleActionRequest>(), It.IsAny<CancellationToken>()))
+			.Returns(OnCompleteLifecycleActionAsync);
+		
+		_asgLifecycleService = new AwsAutoScalingLifecycleService(AgentService, GetRedisServiceSingleton(), AgentCollection, Clock, ServiceProvider, Tracer, logger);
+		_asgLifecycleService.SetAmazonAutoScalingClient(_asgMock.Object);
+		await _asgLifecycleService.StartAsync(CancellationToken.None);
+	}
+
+	[TestMethod]
+	public async Task TerminationRequested_WithAgentInService_MarkedForShutdown()
+	{
+		// Arrange
+		LifecycleActionEvent lae = new() { Ec2InstanceId = "i-1234", LifecycleActionToken = "action-token-test", Origin = "AutoScalingGroup" };
+		IAgent agent = await CreateAgentAsync(new PoolId("pool1"), properties: new() { KnownPropertyNames.AwsInstanceId + "=" + lae.Ec2InstanceId });
+
+		// Act
+		await AwsAsgLifecycleService.InitiateTerminationAsync(lae, CancellationToken.None);
+		
+		// Assert
+		agent = (await AgentService.GetAgentAsync(agent.Id))!;
+		Assert.IsTrue(agent.RequestShutdown);
+	}
+	
+	[TestMethod]
+	public async Task TerminationRequested_WithAgentOnline_LifecycleIsContinued()
+	{
+		// Arrange
+		LifecycleActionEvent lae = new() { Ec2InstanceId = "i-1234", LifecycleActionToken = "action-token-test", Origin = "AutoScalingGroup" };
+		await CreateAgentAsync(new PoolId("pool1"), properties: new() { KnownPropertyNames.AwsInstanceId + "=" + lae.Ec2InstanceId });
+		await _asgLifecycleService.InitiateTerminationAsync(lae, CancellationToken.None);
+		
+		// Act
+		await Clock.AdvanceAsync(_asgLifecycleService.LifecycleUpdaterInterval + TimeSpan.FromSeconds(10));
+		
+		// Assert
+		AssertLifecycleUpdate(lae, AwsAutoScalingLifecycleService.ActionContinue, _lifecycleUpdates[0]);
+	}
+	
+	[TestMethod]
+	public async Task TerminationRequested_WithAgentGoingFromOnlineToOffline()
+	{
+		// Arrange
+		LifecycleActionEvent lae = new() { Ec2InstanceId = "i-1234", LifecycleActionToken = "action-token-test", Origin = "AutoScalingGroup" };
+		List<string> props = new() { KnownPropertyNames.AwsInstanceId + "=" + lae.Ec2InstanceId };
+		IAgent agent = await CreateAgentAsync(new PoolId("pool1"), properties: props);
+		agent = await AgentService.CreateSessionAsync(agent, AgentStatus.Ok, props, new Dictionary<string, int>(), "v1");
+		await _asgLifecycleService.InitiateTerminationAsync(lae, CancellationToken.None);
+		TimeSpan extraMargin = TimeSpan.FromSeconds(10);
+		
+		// Agent is online (e.g still running a job) and the lifecycle update tick is triggered
+		await Clock.AdvanceAsync(_asgLifecycleService.LifecycleUpdaterInterval + extraMargin);
+		AssertLifecycleUpdate(lae, AwsAutoScalingLifecycleService.ActionContinue, _lifecycleUpdates[0]);
+		
+		// Agent still online
+		await Clock.AdvanceAsync(_asgLifecycleService.LifecycleUpdaterInterval + extraMargin);
+		AssertLifecycleUpdate(lae, AwsAutoScalingLifecycleService.ActionContinue, _lifecycleUpdates[1]);
+		
+		// Agent responded to shutdown request, now safe to notify AWS ASG the termination can take place
+		agent = await AgentService.CreateSessionAsync(agent, AgentStatus.Stopped, props, new Dictionary<string, int>(), "v1");
+		await Clock.AdvanceAsync(_asgLifecycleService.LifecycleUpdaterInterval + extraMargin);
+		AssertLifecycleUpdate(lae, AwsAutoScalingLifecycleService.ActionAbandon, _lifecycleUpdates[2]);
+		
+		// Trigger lifecycle update tick once again and no further updates should have been sent
+		await Clock.AdvanceAsync(_asgLifecycleService.LifecycleUpdaterInterval + extraMargin);
+		Assert.AreEqual(3, _lifecycleUpdates.Count);
+	}
+	
+	[TestMethod]
+	public async Task TerminationRequested_WithAgentInWarmPool_LifecycleIsAbandoned()
+	{
+		// Arrange
+		LifecycleActionEvent lae = new() { Ec2InstanceId = "i-1234", LifecycleActionToken = "action-token-test", Origin = "WarmPool" };
+
+		// Act
+		await _asgLifecycleService.InitiateTerminationAsync(lae, CancellationToken.None);
+		
+		// Assert
+		AssertLifecycleUpdate(lae, AwsAutoScalingLifecycleService.ActionAbandon, _lifecycleUpdates[0]);
+	}
+
+	private static void AssertLifecycleUpdate(LifecycleActionEvent expectedEvent, string expectedResult, CompleteLifecycleActionRequest actual)
+	{
+		Assert.AreEqual(expectedEvent.Ec2InstanceId, actual.InstanceId);
+		Assert.AreEqual(expectedResult, actual.LifecycleActionResult);
+		Assert.AreEqual(expectedEvent.LifecycleActionToken, actual.LifecycleActionToken);
+		Assert.AreEqual(expectedEvent.LifecycleHookName, actual.LifecycleHookName);
+		Assert.AreEqual(expectedEvent.AutoScalingGroupName, actual.AutoScalingGroupName);
+	}
+}
diff --git a/Engine/Source/Programs/Horde/Horde.Server.Tests/TestSetup.cs b/Engine/Source/Programs/Horde/Horde.Server.Tests/TestSetup.cs
--- a/Engine/Source/Programs/Horde/Horde.Server.Tests/TestSetup.cs
+++ b/Engine/Source/Programs/Horde/Horde.Server.Tests/TestSetup.cs
@@ -100,6 +100,7 @@ namespace Horde.Server.Tests
 		public AclService AclService => ServiceProvider.GetRequiredService<AclService>();
 		public FleetService FleetService => ServiceProvider.GetRequiredService<FleetService>();
 		public AgentService AgentService => ServiceProvider.GetRequiredService<AgentService>();
+		public AwsAutoScalingLifecycleService AwsAsgLifecycleService => ServiceProvider.GetRequiredService<AwsAutoScalingLifecycleService>();
 		public ICommitService CommitService => ServiceProvider.GetRequiredService<ICommitService>();
 		public GlobalsService GlobalsService => ServiceProvider.GetRequiredService<GlobalsService>();
 		public MongoService MongoService => ServiceProvider.GetRequiredService<MongoService>();
@@ -133,6 +134,7 @@ namespace Horde.Server.Tests
 		public DashboardController DashboardController => GetDashboardController();
 		public TestDataController TestDataController => GetTestDataController();
 		public BisectTasksController BisectTasksController => GetBisectTasksController();
+		public AwsAutoScalingLifecycleController AwsAutoScalingLifecycleController => GetAwsAutoScalingLifecycleController();
 
 		public OpenTelemetry.Trace.Tracer Tracer => ServiceProvider.GetRequiredService<OpenTelemetry.Trace.Tracer>();
 		public Meter Meter => ServiceProvider.GetRequiredService<Meter>();
@@ -241,6 +243,7 @@ namespace Horde.Server.Tests
 
 			services.AddSingleton<AclService>();
 			services.AddSingleton<AgentService>();
+			services.AddSingleton<AwsAutoScalingLifecycleService>();
 			services.AddSingleton<ArtifactExpirationService>();
 			services.AddSingleton<FleetService>();
 			services.AddSingleton<ConsistencyService>();
@@ -315,6 +318,13 @@ namespace Horde.Server.Tests
 			dataCtrl.ControllerContext = GetControllerContext();
 			return dataCtrl;
 		}
+		
+		private AwsAutoScalingLifecycleController GetAwsAutoScalingLifecycleController()
+		{
+			AwsAutoScalingLifecycleController ctrl = new (AwsAsgLifecycleService);
+			ctrl.ControllerContext = GetControllerContext();
+			return ctrl;
+		}
 
 		private BisectTasksController GetBisectTasksController()
 		{
@@ -354,7 +364,12 @@ namespace Horde.Server.Tests
 		}
 		
 		private static int s_agentIdCounter = 1;
-		public async Task<IAgent> CreateAgentAsync(IPool pool, bool enabled = true, bool requestShutdown = false, TimeSpan? adjustClockBy = null)
+		public Task<IAgent> CreateAgentAsync(IPool pool, bool enabled = true, bool requestShutdown = false, List<string>? properties = null, TimeSpan? adjustClockBy = null)
+		{
+			return CreateAgentAsync(pool.Id, enabled, requestShutdown, properties, adjustClockBy);
+		}
+		
+		public async Task<IAgent> CreateAgentAsync(PoolId poolId, bool enabled = true, bool requestShutdown = false, List<string>? properties = null, TimeSpan? adjustClockBy = null)
 		{
 			DateTime now = Clock.UtcNow;
 			if (adjustClockBy != null)
@@ -362,8 +377,8 @@ namespace Horde.Server.Tests
 				Clock.UtcNow = now + adjustClockBy.Value;
 			}
 			
-			IAgent agent = await AgentService.CreateAgentAsync("TestAgent" + s_agentIdCounter++, enabled, new List<PoolId> { pool.Id });
-			agent = await AgentService.CreateSessionAsync(agent, AgentStatus.Ok, new List<string>(), new Dictionary<string, int>(), null);
+			IAgent agent = await AgentService.CreateAgentAsync("TestAgent" + s_agentIdCounter++, enabled, new List<PoolId> { poolId });
+			agent = await AgentService.CreateSessionAsync(agent, AgentStatus.Ok, properties ?? new List<string>(), new Dictionary<string, int>(), null);
 			if (requestShutdown)
 			{
 				await AgentCollection.TryUpdateSettingsAsync(agent, requestShutdown: true);
diff --git a/Engine/Source/Programs/Horde/Horde.Server/Agents/AgentCollection.cs b/Engine/Source/Programs/Horde/Horde.Server/Agents/AgentCollection.cs
--- a/Engine/Source/Programs/Horde/Horde.Server/Agents/AgentCollection.cs
+++ b/Engine/Source/Programs/Horde/Horde.Server/Agents/AgentCollection.cs
@@ -185,7 +185,14 @@ namespace Horde.Server.Agents
 		}
 
 		/// <inheritdoc/>
-		public async Task<List<IAgent>> FindAsync(PoolId? poolId, DateTime? modifiedAfter, AgentStatus? status, bool? enabled, int? index, int? count)
+		public async Task<List<IAgent>> GetManyAsync(List<AgentId> agentIds)
+		{
+			List<AgentDocument> agentDocuments = await _agents.Find(p => agentIds.Contains(p.Id)).ToListAsync();
+			return new List<IAgent>(agentDocuments);
+		}
+
+		/// <inheritdoc/>
+		public async Task<List<IAgent>> FindAsync(PoolId? poolId, DateTime? modifiedAfter, string? property, AgentStatus? status, bool? enabled, int? index, int? count)
 		{
 			FilterDefinitionBuilder<AgentDocument> filterBuilder = new FilterDefinitionBuilder<AgentDocument>();
 
@@ -201,6 +208,11 @@ namespace Horde.Server.Agents
 				filter &= filterBuilder.Gt(x => x.UpdateTime, modifiedAfter.Value);
 			}
 			
+			if (property != null)
+			{
+				filter &= filterBuilder.AnyEq(x => x.Properties, property);
+			}
+			
 			if (status != null)
 			{
 				filter &= filterBuilder.Eq(x => x.Status, status.Value);
diff --git a/Engine/Source/Programs/Horde/Horde.Server/Agents/AgentService.cs b/Engine/Source/Programs/Horde/Horde.Server/Agents/AgentService.cs
--- a/Engine/Source/Programs/Horde/Horde.Server/Agents/AgentService.cs
+++ b/Engine/Source/Programs/Horde/Horde.Server/Agents/AgentService.cs
@@ -209,12 +209,13 @@ namespace Horde.Server.Agents
 		/// </summary>
 		/// <param name="poolId">The pool containing the agent</param>
 		/// <param name="modifiedAfter">If set, only returns agents modified after this time</param>
+		/// <param name="property">If set, only return agents matching this property</param>
 		/// <param name="index">Index within the list of results</param>
 		/// <param name="count">Number of results to return</param>
 		/// <returns>List of agents matching the given criteria</returns>
-		public Task<List<IAgent>> FindAgentsAsync(PoolId? poolId, DateTime? modifiedAfter, int? index, int? count)
+		public Task<List<IAgent>> FindAgentsAsync(PoolId? poolId, DateTime? modifiedAfter, string? property, int? index, int? count)
 		{
-			return Agents.FindAsync(poolId, modifiedAfter, null, null, index, count);
+			return Agents.FindAsync(poolId, modifiedAfter, property, null, null, index, count);
 		}
 
 		/// <summary>
diff --git a/Engine/Source/Programs/Horde/Horde.Server/Agents/AgentsController.cs b/Engine/Source/Programs/Horde/Horde.Server/Agents/AgentsController.cs
--- a/Engine/Source/Programs/Horde/Horde.Server/Agents/AgentsController.cs
+++ b/Engine/Source/Programs/Horde/Horde.Server/Agents/AgentsController.cs
@@ -58,7 +58,7 @@ namespace Horde.Server.Agents
 				return Forbid(AgentAclAction.ListAgents);
 			}
 
-			List<IAgent> agents = await _agentService.FindAgentsAsync(poolId, modifiedAfter?.UtcDateTime, index, count);
+			List<IAgent> agents = await _agentService.FindAgentsAsync(poolId, modifiedAfter?.UtcDateTime, null, index, count);
 
 			List<object> responses = new List<object>();
 			foreach (IAgent agent in agents)
diff --git a/Engine/Source/Programs/Horde/Horde.Server/Agents/Fleet/AwsAutoScalingLifecycle.cs b/Engine/Source/Programs/Horde/Horde.Server/Agents/Fleet/AwsAutoScalingLifecycle.cs
new file mode 100644
--- /dev/null
+++ b/Engine/Source/Programs/Horde/Horde.Server/Agents/Fleet/AwsAutoScalingLifecycle.cs
@@ -0,0 +1,350 @@
+// Copyright Epic Games, Inc. All Rights Reserved.
+
+using System;
+using System.Collections.Generic;
+using System.IO;
+using System.Linq;
+using System.Text.Json;
+using System.Text.Json.Serialization;
+using System.Threading;
+using System.Threading.Tasks;
+using Amazon.AutoScaling;
+using Amazon.AutoScaling.Model;
+using Horde.Server.Server;
+using Horde.Server.Utilities;
+using HordeCommon;
+using Microsoft.AspNetCore.Authorization;
+using Microsoft.AspNetCore.Mvc;
+using Microsoft.Extensions.DependencyInjection;
+using Microsoft.Extensions.Hosting;
+using Microsoft.Extensions.Logging;
+using OpenTelemetry.Trace;
+using StackExchange.Redis;
+
+namespace Horde.Server.Agents.Fleet;
+
+/// <summary>
+/// Service handling callbacks and lifecycle events triggered by EC2 auto-scaling in AWS
+/// </summary>
+public sealed class AwsAutoScalingLifecycleService : IHostedService, IDisposable
+{
+	/// <summary>
+	/// Lifecycle action for continue
+	/// In the case of terminating instances, setting this result means the instance should *continue* be kept alive
+	/// </summary>
+	public const string ActionContinue = "CONTINUE";
+	
+	/// <summary>
+	/// Lifecycle action for abandon
+	/// In the case of terminating instances, setting this result means the instance should be *terminated*
+	/// </summary>
+	public const string ActionAbandon = "ABANDON";
+	
+	/// <summary>
+	/// Warm pool origin
+	/// </summary>
+	public const string OriginWarmPool = "WarmPool";
+	
+	/// <summary>
+	/// Normal auto-scaling group origin
+	/// </summary>
+	public const string OriginAsg = "AutoScalingGroup";
+	
+	/// <summary>
+	/// How often instance lifecycles for auto-scaling should be updated with AWS 
+	/// </summary>
+	public TimeSpan LifecycleUpdaterInterval { get; } = TimeSpan.FromMinutes(20);
+
+	private const string RedisKey = "aws-asg-lifecycle";
+	private readonly AgentService _agentService;
+	private readonly RedisService _redisService;
+	private readonly IAgentCollection _agents;
+	private readonly IClock _clock;
+	private readonly Tracer _tracer;
+	private readonly ILogger<AwsAutoScalingLifecycleService> _logger;
+	private readonly ITicker _updateLifecyclesTicker;
+
+#pragma warning disable CA2213 // Disposable fields should be disposed
+	private IAmazonAutoScaling? _awsAutoScaling;
+#pragma warning restore CA2213
+	
+	/// <summary>
+	/// Constructor
+	/// </summary>
+	public AwsAutoScalingLifecycleService(
+		AgentService agentService,
+		RedisService redisService,
+		IAgentCollection agents,
+		IClock clock,
+		IServiceProvider serviceProvider,
+		Tracer tracer,
+		ILogger<AwsAutoScalingLifecycleService> logger)
+	{
+		_agentService = agentService;
+		_redisService = redisService;
+		_agents = agents;
+		_clock = clock;
+		_tracer = tracer;
+		_logger = logger;
+		_awsAutoScaling = serviceProvider.GetService<IAmazonAutoScaling>();
+
+		string tickerName = $"{nameof(AwsAutoScalingLifecycleService)}.{nameof(UpdateLifecyclesAsync)}";
+		_updateLifecyclesTicker = clock.AddSharedTicker(tickerName, LifecycleUpdaterInterval, UpdateLifecyclesAsync, logger);
+	}
+
+	internal void SetAmazonAutoScalingClient(IAmazonAutoScaling awsAutoScaling)
+	{
+		_awsAutoScaling = awsAutoScaling;
+	}
+	
+	/// <inheritdoc/>
+	public async Task StartAsync(CancellationToken cancellationToken)
+	{
+		await _updateLifecyclesTicker.StartAsync();
+	}
+
+	/// <inheritdoc/>
+	public async Task StopAsync(CancellationToken cancellationToken)
+	{
+		await _updateLifecyclesTicker.StopAsync();
+	} 
+
+	/// <inheritdoc/>
+	public void Dispose()
+	{
+		_updateLifecyclesTicker.Dispose();
+	}
+	
+	/// <summary>
+	/// Gracefully terminate an EC2 instance running the Horde agent
+	/// This call is initiated by an event coming from AWS auto-scaling group and will request the agent to shutdown.
+	/// Since the agent can be busy running a job, it's possible it cannot be terminated immediately.
+	/// The termination request from AWS is therefore tracked and updated continuously as the agent shutdown progresses. 
+	/// </summary>
+	/// <param name="e">Lifecycle action event received</param>
+	/// <param name="cancellationToken">Cancellation token</param>
+	/// <returns>True if handled</returns>
+	public async Task<bool> InitiateTerminationAsync(LifecycleActionEvent e, CancellationToken cancellationToken)
+	{
+		if (e.Origin == OriginAsg)
+		{
+			string instanceIdProp = KnownPropertyNames.AwsInstanceId + "=" + e.Ec2InstanceId;
+			IAgent? agent = (await _agentService.FindAgentsAsync(null, null, instanceIdProp, null, null)).FirstOrDefault();
+			if (agent == null)
+			{
+				_logger.LogWarning("Lifecycle action received but no agent with instance ID {InstanceId} found", e.Ec2InstanceId);
+				return false;
+			}
+		
+			await _agents.TryUpdateSettingsAsync(agent, requestShutdown: true);
+			await TrackAgentLifecycleAsync(agent.Id, e);
+			return true;
+		}
+		else if (e.Origin == OriginWarmPool)
+		{
+			// EC2 instance in warm pool and not current in use by Horde
+			await SendCompleteLifecycleActionAsync(e, ActionResult.Abandon, cancellationToken);
+			return true;
+		}
+		else
+		{
+			_logger.LogWarning("Unknown origin {Origin}", e.Origin);
+			return false;
+		}
+	}
+
+	private async Task TrackAgentLifecycleAsync(AgentId agentId, LifecycleActionEvent e)
+	{
+		DateTime utcNow = _clock.UtcNow;
+		AgentLifecycleInfo info = new(agentId.ToString(), utcNow, utcNow, e);
+		using MemoryStream ms = new (500);
+		await JsonSerializer.SerializeAsync(ms, info);
+		await _redisService.GetDatabase().HashSetAsync(RedisKey, agentId.ToString(), ms.ToArray());
+	}
+
+	private async Task SendCompleteLifecycleActionAsync(LifecycleActionEvent laEvent, ActionResult result, CancellationToken cancellationToken)
+	{
+		if (_awsAutoScaling == null)
+		{
+			throw new Exception($"{nameof(IAmazonAutoScaling)} client not initialized! Cannot make AWS API call.");
+		}
+		
+		using TelemetrySpan span = _tracer.StartActiveSpan($"{nameof(AwsAutoScalingLifecycleService)}.{nameof(SendCompleteLifecycleActionAsync)}");
+		CompleteLifecycleActionRequest request = new()
+		{
+			InstanceId = laEvent.Ec2InstanceId,
+			LifecycleActionResult = result == ActionResult.Continue ? ActionContinue : ActionAbandon,
+			LifecycleActionToken = laEvent.LifecycleActionToken,
+			LifecycleHookName = laEvent.LifecycleHookName,
+			AutoScalingGroupName = laEvent.AutoScalingGroupName
+		};
+		span.SetAttribute("instanceId", request.InstanceId);
+		span.SetAttribute("lifecycleActionResult", request.LifecycleActionResult);
+		span.SetAttribute("lifecycleActionToken", request.LifecycleActionToken);
+		span.SetAttribute("lifecycleHookName", request.LifecycleHookName);
+		span.SetAttribute("autoScalingGroupName", request.AutoScalingGroupName);
+		
+		CompleteLifecycleActionResponse response = await _awsAutoScaling.CompleteLifecycleActionAsync(request, cancellationToken);
+		span.SetAttribute("statusCode", response.HttpStatusCode.ToString());
+	}
+	
+	/// <summary>
+	/// Update AWS auto-scaling on the lifecycle status for each agent being tracked for shutdown
+	/// </summary>
+	/// <param name="cancellationToken">Cancellation token for the async task</param>
+	/// <returns>Async task</returns>
+	internal async ValueTask UpdateLifecyclesAsync(CancellationToken cancellationToken)
+	{
+		using TelemetrySpan span = _tracer.StartActiveSpan($"{nameof(AwsAutoScalingLifecycleService)}.{nameof(UpdateLifecyclesAsync)}");
+
+		Dictionary<AgentId, AgentLifecycleInfo> infos = new();
+		IDatabase redis = _redisService.GetDatabase();
+		
+		foreach (HashEntry entry in await redis.HashGetAllAsync(RedisKey))
+		{
+			try
+			{
+				AgentId agentId = new(entry.Name.ToString());
+				AgentLifecycleInfo info = JsonSerializer.Deserialize<AgentLifecycleInfo>(entry.Value.ToString()) ?? throw new Exception("Deserialization returned null");
+				infos[agentId] = info;
+			}
+			catch (Exception e)
+			{
+				_logger.LogWarning(e, "Unable to deserialize agent lifecycle info object from Redis. {Key} = {Value}", entry.Name.ToString(), entry.Value.ToString());
+				await redis.HashDeleteAsync(RedisKey, entry.Name);
+			}
+		}
+
+		span.SetAttribute("numTrackedAgents", infos.Count);
+
+		List<Task> tasks = new();
+		foreach (IAgent agent in await _agents.GetManyAsync(infos.Keys.ToList()))
+		{
+			ActionResult result = agent.Status == AgentStatus.Stopped ? ActionResult.Abandon : ActionResult.Continue;
+			tasks.Add(Task.Run(async () =>
+			{
+				await SendCompleteLifecycleActionAsync(infos[agent.Id].Event, result, cancellationToken);
+				if (result == ActionResult.Abandon)
+				{
+					await redis.HashDeleteAsync(RedisKey, agent.Id.ToString());
+				}
+			}, cancellationToken));
+		}
+		
+		span.SetAttribute("numLifecycleActionsSent", tasks.Count);
+		await Task.WhenAll(tasks);
+	}
+
+	private class AgentLifecycleInfo
+	{
+		public string AgentId { get; private set; }
+		public DateTime CreatedAt { get; private set; }
+		public DateTime UpdatedAt { get; private set; }
+		public LifecycleActionEvent Event { get; private set; }
+
+		public AgentLifecycleInfo(string agentId, DateTime createdAt, DateTime updatedAt, LifecycleActionEvent @event)
+		{
+			AgentId = agentId;
+			CreatedAt = createdAt;
+			UpdatedAt = updatedAt;
+			Event = @event;
+		}
+	}
+
+	private enum ActionResult
+	{
+		Continue, Abandon
+	}
+}
+
+/// <summary>
+/// An AWS EventBridge event containing a lifecycle hook message for an auto-scaling group
+/// </summary>
+public class LifecycleActionRequest
+{
+	/// <summary>
+	/// Event
+	/// </summary>
+	[JsonPropertyName("detail")]
+	public LifecycleActionEvent Event { get; set; }
+
+	/// <summary>
+	/// Constructor
+	/// </summary>
+	/// <param name="event"></param>
+	public LifecycleActionRequest(LifecycleActionEvent @event)
+	{
+		Event = @event;
+	}
+}
+
+/// <summary>
+/// Lifecycle action event from AWS auto-scaling group
+/// The property names are explicitly set to highlight these are not controlled by Horde but sent from AWS API.
+/// </summary>
+public class LifecycleActionEvent
+{
+	/// <summary>
+	/// Name of the lifecycle hook
+	/// </summary>
+	[JsonPropertyName("LifecycleHookName")] public string LifecycleHookName { get; set; } = "";
+	
+	/// <summary>
+	/// Token for the lifecycle action
+	/// </summary>
+	[JsonPropertyName("LifecycleActionToken")] public string LifecycleActionToken { get; set; } = "";
+	
+	/// <summary>
+	/// Lifecycle transition
+	/// </summary>
+	[JsonPropertyName("LifecycleTransition")] public string LifecycleTransition { get; set; } = "";
+	
+	/// <summary>
+	/// Name of the auto-scaling group being handled
+	/// </summary>
+	[JsonPropertyName("AutoScalingGroupName")] public string AutoScalingGroupName { get; set; } = "";
+	
+	/// <summary>
+	/// Instance ID
+	/// </summary>
+	[JsonPropertyName("EC2InstanceId")] public string Ec2InstanceId { get; set; } = "";
+	
+	/// <summary>
+	/// Origin of the instance (e.g an agent in service or in a warm pool)
+	/// </summary>
+	[JsonPropertyName("Origin")] public string Origin { get; set; } = "";
+}
+
+/// <summary>
+/// Controller handling callbacks for AWS auto-scaling
+/// </summary>
+[ApiController]
+[Authorize]
+[Route("[controller]")]
+public class AwsAutoScalingLifecycleController : HordeControllerBase
+{
+	private AwsAutoScalingLifecycleService _lifecycleService;
+
+	/// <summary>
+	/// Constructor
+	/// </summary>
+	public AwsAutoScalingLifecycleController(AwsAutoScalingLifecycleService lifecycleService)
+	{
+		_lifecycleService = lifecycleService;
+	}
+	
+	/// <summary>
+	/// Callback handling auto-scaling notification generated by the ASG
+	/// A Lambda function should be configured first to receive the notification from AWS EventBridge.
+	/// Then in turn, call the Horde server on this endpoint. 
+	/// </summary>
+	/// <param name="request">Request coming from</param>
+	/// <returns>Human readable message</returns>
+	[HttpPost]
+	[Route("/api/v1/aws/asg-lifecycle-hook")]
+	public async Task<ActionResult> LifecycleHookAsync([FromBody] LifecycleActionRequest request)
+	{
+		await _lifecycleService.InitiateTerminationAsync(request.Event, CancellationToken.None);
+		return new JsonResult(new { message = "Handled" });
+	}
+}
\ No newline at end of file
diff --git a/Engine/Source/Programs/Horde/Horde.Server/Agents/IAgent.cs b/Engine/Source/Programs/Horde/Horde.Server/Agents/IAgent.cs
--- a/Engine/Source/Programs/Horde/Horde.Server/Agents/IAgent.cs
+++ b/Engine/Source/Programs/Horde/Horde.Server/Agents/IAgent.cs
@@ -354,6 +354,11 @@ namespace Horde.Server.Agents
 		/// </summary>
 		public const string Ram = "RAM";
 		
+		/// <summary>
+		/// AWS: Instance ID
+		/// </summary>
+		public const string AwsInstanceId = "aws-instance-id";
+		
 		/// <summary>
 		/// AWS: Instance type
 		/// </summary>
diff --git a/Engine/Source/Programs/Horde/Horde.Server/Agents/IAgentCollection.cs b/Engine/Source/Programs/Horde/Horde.Server/Agents/IAgentCollection.cs
--- a/Engine/Source/Programs/Horde/Horde.Server/Agents/IAgentCollection.cs
+++ b/Engine/Source/Programs/Horde/Horde.Server/Agents/IAgentCollection.cs
@@ -45,21 +45,29 @@ namespace Horde.Server.Agents
 		/// <param name="agentId">Unique id of the agent</param>
 		/// <returns>The agent document</returns>
 		Task<IAgent?> GetAsync(AgentId agentId);
+		
+		/// <summary>
+		/// Gets multiple agents by ID
+		/// </summary>
+		/// <param name="agentIds">List of unique IDs of the agents</param>
+		/// <returns>The agent documents</returns>
+		Task<List<IAgent>> GetManyAsync(List<AgentId> agentIds);
 
 		/// <summary>
 		/// Finds all agents matching certain criteria
 		/// </summary>
 		/// <param name="poolId">The pool ID in string form containing the agent</param>
 		/// <param name="modifiedAfter">If set, only returns agents modified after this time</param>
+		/// <param name="property">Property to look for</param>
 		/// <param name="status">Status to look for</param>
 		/// <param name="enabled">Enabled/disabled status to look for</param>
 		/// <param name="index">Index of the first result</param>
 		/// <param name="count">Number of results to return</param>
 		/// <returns>List of agents matching the given criteria</returns>
-		Task<List<IAgent>> FindAsync(PoolId? poolId = null, DateTime? modifiedAfter = null, AgentStatus? status = null, bool? enabled = null, int? index = null, int? count = null);
+		Task<List<IAgent>> FindAsync(PoolId? poolId = null, DateTime? modifiedAfter = null, string? property = null, AgentStatus? status = null, bool? enabled = null, int? index = null, int? count = null);
 
 		/// <summary>
-		/// Finds all the expried agents
+		/// Finds all the expired agents
 		/// </summary>
 		/// <param name="utcNow">The current time</param>
 		/// <param name="maxAgents">Maximum number of agents to return</param>
diff --git a/Engine/Source/Programs/Horde/Horde.Server/Startup.cs b/Engine/Source/Programs/Horde/Horde.Server/Startup.cs
--- a/Engine/Source/Programs/Horde/Horde.Server/Startup.cs
+++ b/Engine/Source/Programs/Horde/Horde.Server/Startup.cs
@@ -433,6 +433,7 @@ namespace Horde.Server
 			services.AddSingleton(typeof(IAuditLogFactory<>), typeof(AuditLogFactory<>));
 			services.AddSingleton(typeof(ISingletonDocument<>), typeof(SingletonDocument<>));
 
+			services.AddSingleton<AwsAutoScalingLifecycleService>();
 			services.AddSingleton<FleetService>();
 			services.AddSingleton<IFleetManagerFactory, FleetManagerFactory>();
 			
@@ -716,6 +717,11 @@ namespace Horde.Server
 				{
 					services.AddHostedService(provider => provider.GetRequiredService<SlackNotificationSink>());
 				}
+				
+				if (settings.WithAws)
+				{
+					services.AddHostedService(provider => provider.GetRequiredService<AwsAutoScalingLifecycleService>());
+				}
 			}
 
 			services.AddHostedService(provider => provider.GetRequiredService<IExternalIssueService>());
diff --git a/Engine/Source/Programs/Horde/Horde.Server/Utilities/Clock.cs b/Engine/Source/Programs/Horde/Horde.Server/Utilities/Clock.cs
--- a/Engine/Source/Programs/Horde/Horde.Server/Utilities/Clock.cs
+++ b/Engine/Source/Programs/Horde/Horde.Server/Utilities/Clock.cs
@@ -72,7 +72,8 @@ namespace HordeCommon
 		ITicker AddTicker(string name, TimeSpan interval, Func<CancellationToken, ValueTask<TimeSpan?>> tickAsync, ILogger logger);
 
 		/// <summary>
-		/// Create a ticker shared between all server pods
+		/// Create a ticker shared between all server processes.
+		/// Callback can be run inside any available process but will still only be called once per tick.
 		/// </summary>
 		/// <param name="name">Name of the event</param>
 		/// <param name="interval">Time after which the event will trigger</param>
