Horde: Allow specifying the conditions under which a process should be terminated, via the ProcessesToTerminate property.

diff --git a/Engine/Source/Programs/Horde/Horde.Agent.Tests/WorkerServiceTest.cs b/Engine/Source/Programs/Horde/Horde.Agent.Tests/WorkerServiceTest.cs
--- a/Engine/Source/Programs/Horde/Horde.Agent.Tests/WorkerServiceTest.cs
+++ b/Engine/Source/Programs/Horde/Horde.Agent.Tests/WorkerServiceTest.cs
@@ -270,7 +270,7 @@ namespace Horde.Agent.Tests
 			fakeSession.Setup(x => x.AgentId).Returns("LocalAgent");
 			fakeSession.Setup(x => x.SessionId).Returns("Session");
 			fakeSession.Setup(x => x.RpcConnection).Returns(rpcConnection);
-			fakeSession.Setup(x => x.TerminateProcessesAsync(It.IsAny<ILogger>(), It.IsAny<CancellationToken>())).Returns(Task.CompletedTask);
+			fakeSession.Setup(x => x.TerminateProcessesAsync(It.IsAny<TerminateCondition>(), It.IsAny<ILogger>(), It.IsAny<CancellationToken>())).Returns(Task.CompletedTask);
 			fakeSession.Setup(x => x.DisposeAsync()).Returns(new ValueTask());
 			fakeSession.Setup(x => x.WorkingDir).Returns(DirectoryReference.GetCurrentDirectory());
 			return fakeSession.Object;
diff --git a/Engine/Source/Programs/Horde/Horde.Agent/AgentSettings.cs b/Engine/Source/Programs/Horde/Horde.Agent/AgentSettings.cs
--- a/Engine/Source/Programs/Horde/Horde.Agent/AgentSettings.cs
+++ b/Engine/Source/Programs/Horde/Horde.Agent/AgentSettings.cs
@@ -6,6 +6,7 @@ using System.ComponentModel.DataAnnotations;
 using System.Linq;
 using EpicGames.Core;
 using EpicGames.Horde.Storage;
+using Horde.Agent.Services;
 using Microsoft.Extensions.Configuration;
 
 namespace Horde.Agent
@@ -110,6 +111,59 @@ namespace Horde.Agent
 		public bool RunConform { get; set; } = true;
 	}
 
+	/// <summary>
+	/// Flags for processes to terminate
+	/// </summary>
+	[Flags]
+	public enum TerminateCondition
+	{
+		/// <summary>
+		/// Not specified; terminate in all circumstances
+		/// </summary>
+		None = 0,
+
+		/// <summary>
+		/// When a session starts
+		/// </summary>
+		BeforeSession = 1,
+
+		/// <summary>
+		/// Before running a conform
+		/// </summary>
+		BeforeConform = 2,
+
+		/// <summary>
+		/// Before executing a batch
+		/// </summary>
+		BeforeBatch = 4,
+
+		/// <summary>
+		/// Terminate at the end of a batch
+		/// </summary>
+		AfterBatch = 8,
+
+		/// <summary>
+		/// After a step completes
+		/// </summary>
+		AfterStep = 16,
+	}
+
+	/// <summary>
+	/// Specifies a process to terminate
+	/// </summary>
+	public class ProcessToTerminate
+	{
+		/// <summary>
+		/// Name of the process
+		/// </summary>
+		public string Name { get; set; } = String.Empty;
+
+		/// <summary>
+		/// When to terminate this process
+		/// </summary>
+		public List<TerminateCondition>? When { get; init; }
+	}
+
 	/// <summary>
 	/// Global settings for the agent
 	/// </summary>
@@ -169,8 +223,14 @@ namespace Horde.Agent
 		/// <summary>
 		/// List of process names to terminate after a job
 		/// </summary>
+		[Obsolete("Prefer using the ProcessesToTerminate list instead")]
 		public List<string> ProcessNamesToTerminate { get; } = new List<string>();
 
+		/// <summary>
+		/// List of process names to terminate after a lease completes, but not after a job step
+		/// </summary>
+		public List<ProcessToTerminate> ProcessesToTerminate { get; } = new List<ProcessToTerminate>();
+
 		/// <summary>
 		/// Whether to write step output to the logging device
 		/// </summary>
@@ -235,7 +295,34 @@ namespace Horde.Agent
 
 			return GetServerProfile(Server);
 		}
-		
+
+		/// <summary>
+		/// Gets a lookup of process name to the circumstances in which it should be terminate
+		/// </summary>
+		public Dictionary<string, TerminateCondition> GetProcessesToTerminateMap()
+		{
+			Dictionary<string, TerminateCondition> processesToTerminate = new Dictionary<string, TerminateCondition>(StringComparer.OrdinalIgnoreCase);
+#pragma warning disable CS0618 // Type or member is obsolete
+			foreach (string processName in ProcessNamesToTerminate)
+			{
+				processesToTerminate[processName] = TerminateCondition.None;
+			}
+#pragma warning restore CS0618 // Type or member is obsolete
+			foreach (ProcessToTerminate processToTerminate in ProcessesToTerminate)
+			{
+				TerminateCondition condition = default;
+				if (processToTerminate.When != null)
+				{
+					foreach (TerminateCondition when in processToTerminate.When)
+					{
+						condition |= when;
+					}
+				}
+				processesToTerminate[processToTerminate.Name] = condition;
+			}
+			return processesToTerminate;
+		}
+
 		internal string GetAgentName()
 		{
 			return Name ?? Environment.MachineName;
diff --git a/Engine/Source/Programs/Horde/Horde.Agent/Commands/Execution/ExecuteJobCommand.cs b/Engine/Source/Programs/Horde/Horde.Agent/Commands/Execution/ExecuteJobCommand.cs
--- a/Engine/Source/Programs/Horde/Horde.Agent/Commands/Execution/ExecuteJobCommand.cs
+++ b/Engine/Source/Programs/Horde/Horde.Agent/Commands/Execution/ExecuteJobCommand.cs
@@ -1,6 +1,7 @@
 // Copyright Epic Games, Inc. All Rights Reserved.
 
 using System;
+using System.Collections.Generic;
 using System.Threading;
 using System.Threading.Tasks;
 using EpicGames.Core;
@@ -47,8 +48,10 @@ namespace Horde.Agent.Commands.Execution
 			ServerProfile serverProfile = _settings.GetCurrentServerProfile();
 			ExecuteJobTask executeTask = ExecuteJobTask.Parser.ParseFrom(Convert.FromBase64String(Task));
 
+			Dictionary<string, TerminateCondition> processNamesToTerminate = _settings.GetProcessesToTerminateMap();
+
 			await using RpcConnection rpcConnection = new RpcConnection(ctx => _grpcService.CreateGrpcChannelAsync(executeTask.Token, ctx), logger);
-			await using Session session = new Session(serverProfile.Url, AgentId, SessionId, executeTask.Token, rpcConnection, WorkingDir, _settings.ProcessNamesToTerminate, logger);
+			await using Session session = new Session(serverProfile.Url, AgentId, SessionId, executeTask.Token, rpcConnection, WorkingDir, processNamesToTerminate, logger);
 
 			await _jobHandler.ExecuteInternalAsync(session, LeaseId, executeTask, CancellationToken.None);
 			return 0;
diff --git a/Engine/Source/Programs/Horde/Horde.Agent/Leases/Handlers/ConformHandler.cs b/Engine/Source/Programs/Horde/Horde.Agent/Leases/Handlers/ConformHandler.cs
--- a/Engine/Source/Programs/Horde/Horde.Agent/Leases/Handlers/ConformHandler.cs
+++ b/Engine/Source/Programs/Horde/Horde.Agent/Leases/Handlers/ConformHandler.cs
@@ -50,7 +50,7 @@ namespace Horde.Agent.Leases.Handlers
 		async Task<LeaseResult> ExecuteInternalAsync(ISession session, string leaseId, ConformTask conformTask, IServerLogger conformLogger, CancellationToken cancellationToken)
 		{
 			conformLogger.LogInformation("Conforming, lease {LeaseId}", leaseId);
-			await session.TerminateProcessesAsync(conformLogger, cancellationToken);
+			await session.TerminateProcessesAsync(TerminateCondition.BeforeConform, conformLogger, cancellationToken);
 
 			bool removeUntrackedFiles = conformTask.RemoveUntrackedFiles;
 			IList<AgentWorkspace> pendingWorkspaces = conformTask.Workspaces;
diff --git a/Engine/Source/Programs/Horde/Horde.Agent/Leases/Handlers/JobHandler.cs b/Engine/Source/Programs/Horde/Horde.Agent/Leases/Handlers/JobHandler.cs
--- a/Engine/Source/Programs/Horde/Horde.Agent/Leases/Handlers/JobHandler.cs
+++ b/Engine/Source/Programs/Horde/Horde.Agent/Leases/Handlers/JobHandler.cs
@@ -198,7 +198,7 @@ namespace Horde.Agent.Leases.Handlers
 			string executorName = String.IsNullOrEmpty(options.JobOptions.Executor) ? _settings.Executor : options.JobOptions.Executor;
 			
 			batchLogger.LogInformation("Executing batch {BatchId} using {Executor} executor", options.BatchId, executorName);
-			await session.TerminateProcessesAsync(batchLogger, cancellationToken);
+			await session.TerminateProcessesAsync(TerminateCondition.BeforeBatch, batchLogger, cancellationToken);
 
 			IJobExecutorFactory? executorFactory = _executorFactories.FirstOrDefault(x => x.Name.Equals(executorName, StringComparison.OrdinalIgnoreCase));
 			if (executorFactory == null)
@@ -310,7 +310,7 @@ namespace Horde.Agent.Leases.Handlers
 							}
 
 							// Kill any processes spawned by the step
-							await session.TerminateProcessesAsync(batchLogger, cancellationToken);
+							await session.TerminateProcessesAsync(TerminateCondition.AfterStep, batchLogger, cancellationToken);
 
 							// Wait for the logger to finish
 							await stepLogger.StopAsync();
@@ -345,6 +345,16 @@ namespace Horde.Agent.Leases.Handlers
 				}
 			}
 
+			// Terminate any processes which are still running
+			try
+			{
+				await session.TerminateProcessesAsync(TerminateCondition.AfterBatch, batchLogger, CancellationToken.None);
+			}
+			catch(Exception ex)
+			{
+				batchLogger.LogWarning(ex, "Exception while terminating processes: {Message}", ex.Message);
+			}
+
 			// Clean the environment
 			batchLogger.LogInformation("Finalizing...");
 			using (batchLogger.BeginIndentScope("  "))
diff --git a/Engine/Source/Programs/Horde/Horde.Agent/Leases/LeaseManager.cs b/Engine/Source/Programs/Horde/Horde.Agent/Leases/LeaseManager.cs
--- a/Engine/Source/Programs/Horde/Horde.Agent/Leases/LeaseManager.cs
+++ b/Engine/Source/Programs/Horde/Horde.Agent/Leases/LeaseManager.cs
@@ -60,11 +60,6 @@ namespace Horde.Agent.Leases
 			}
 		}
 
-		/// <summary>
-		/// List of processes that should be terminated before running a job
-		/// </summary>
-		private readonly HashSet<string> _processNamesToTerminate = new HashSet<string>(StringComparer.OrdinalIgnoreCase);
-
 		/// <summary>
 		/// Object used for controlling access to the access tokens and active sessions list
 		/// </summary>
@@ -136,8 +131,6 @@ namespace Horde.Agent.Leases
 			_typeUrlToLeaseHandler = leaseHandlers.ToDictionary(x => x.LeaseType, x => x);
 			_settings = settings.Value;
 			_logger = logger;
-
-			_processNamesToTerminate.UnionWith(_settings.ProcessNamesToTerminate);
 		}
 
 		public LeaseManager(ISession session, IServiceProvider serviceProvider)
@@ -234,7 +227,7 @@ namespace Horde.Agent.Leases
 			IRpcConnection rpcCon = _session.RpcConnection;
 
 			// Terminate any remaining child processes from other instances
-			await _session.TerminateProcessesAsync(_logger, stoppingToken);
+			await _session.TerminateProcessesAsync(TerminateCondition.BeforeSession, _logger, stoppingToken);
 
 			// Track how many updates we get in 10 seconds. We'll start rate limiting this if it looks like we've got a problem that's causing us to spam the server.
 			Stopwatch updateTimer = Stopwatch.StartNew();
diff --git a/Engine/Source/Programs/Horde/Horde.Agent/Services/SessionFactory.cs b/Engine/Source/Programs/Horde/Horde.Agent/Services/SessionFactory.cs
--- a/Engine/Source/Programs/Horde/Horde.Agent/Services/SessionFactory.cs
+++ b/Engine/Source/Programs/Horde/Horde.Agent/Services/SessionFactory.cs
@@ -3,8 +3,6 @@
 using System;
 using System.Collections.Generic;
 using System.Linq;
-using System.Net.Http;
-using System.Net.Http.Json;
 using System.Reflection;
 using System.Runtime.InteropServices;
 using System.Text.Json;
@@ -59,9 +57,10 @@ namespace Horde.Agent.Services
 		/// <summary>
 		/// Terminate all processes in running in the working directory
 		/// </summary>
+		/// <param name="condition">Flags indicating which processes to terminate</param>
 		/// <param name="logger">Logger for any diagnostic messages</param>
 		/// <param name="cancellationToken">Cancellation token for the operation</param>
-		Task TerminateProcessesAsync(ILogger logger, CancellationToken cancellationToken);
+		Task TerminateProcessesAsync(TerminateCondition condition, ILogger logger, CancellationToken cancellationToken);
 	}
 
 	/// <summary>
@@ -112,13 +111,13 @@ namespace Horde.Agent.Services
 		/// </summary>
 		public DirectoryReference WorkingDir { get; }
 
-		readonly HashSet<string> _processNamesToTerminate;
+		readonly IReadOnlyDictionary<string, TerminateCondition> _processNamesToTerminate;
 		readonly ILogger _logger;
 
 		/// <summary>
 		/// Constructor
 		/// </summary>
-		public Session(Uri serverUrl, string agentId, string sessionId, string token, IRpcConnection rpcConnection, DirectoryReference workingDir, IEnumerable<string> processNamesToTerminate, ILogger logger)
+		public Session(Uri serverUrl, string agentId, string sessionId, string token, IRpcConnection rpcConnection, DirectoryReference workingDir, IReadOnlyDictionary<string, TerminateCondition> processNamesToTerminate, ILogger logger)
 		{
 			ServerUrl = serverUrl;
 			AgentId = agentId;
@@ -127,7 +126,7 @@ namespace Horde.Agent.Services
 			RpcConnection = rpcConnection;
 			WorkingDir = workingDir;
 
-			_processNamesToTerminate = new HashSet<string>(processNamesToTerminate, StringComparer.OrdinalIgnoreCase);
+			_processNamesToTerminate = processNamesToTerminate;
 			_logger = logger;
 		}
 
@@ -218,7 +217,7 @@ namespace Horde.Agent.Services
 			// Open a connection to the server
 #pragma warning disable CA2000 // False positive; ownership is transferred to new Session object.
 			IRpcConnection rpcConnection = new RpcConnection(createGrpcChannelAsync, logger);
-			return new Session(serverProfile.Url, createSessionResponse.AgentId, createSessionResponse.SessionId, createSessionResponse.Token, rpcConnection, workingDir, currentSettings.ProcessNamesToTerminate, logger);
+			return new Session(serverProfile.Url, createSessionResponse.AgentId, createSessionResponse.SessionId, createSessionResponse.Token, rpcConnection, workingDir, currentSettings.GetProcessesToTerminateMap(), logger);
 #pragma warning restore CA2000
 		}
 
@@ -288,17 +287,17 @@ namespace Horde.Agent.Services
 		}
 
 		/// <inheritdoc/>
-		public Task TerminateProcessesAsync(ILogger logger, CancellationToken cancellationToken)
+		public Task TerminateProcessesAsync(TerminateCondition condition, ILogger logger, CancellationToken cancellationToken)
 		{
 			// Terminate child processes from any previous runs
-			ProcessUtils.TerminateProcesses(ShouldTerminateProcess, logger, cancellationToken);
+			ProcessUtils.TerminateProcesses(x => ShouldTerminateProcess(x, condition), logger, cancellationToken);
 			return Task.CompletedTask;
 		}
 
 		/// <summary>
 		/// Callback for determining whether a process should be terminated
 		/// </summary>
-		bool ShouldTerminateProcess(FileReference imageFile)
+		bool ShouldTerminateProcess(FileReference imageFile, TerminateCondition condition)
 		{
 			if (imageFile.IsUnderDirectory(WorkingDir))
 			{
@@ -306,9 +305,12 @@ namespace Horde.Agent.Services
 			}
 
 			string fileName = imageFile.GetFileName();
-			if (_processNamesToTerminate.Contains(fileName))
+			if (_processNamesToTerminate.TryGetValue(fileName, out TerminateCondition terminateFlags))
 			{
-				return true;
+				if(terminateFlags == TerminateCondition.None || (terminateFlags & condition) != 0)
+				{
+					return true;
+				}
 			}
 
 			return false;
