Horde: Add support for a cleanup script to be executed on executor finalization.

diff --git a/Engine/Source/Programs/Horde/Horde.Agent/Execution/JobExecutor.cs b/Engine/Source/Programs/Horde/Horde.Agent/Execution/JobExecutor.cs
--- a/Engine/Source/Programs/Horde/Horde.Agent/Execution/JobExecutor.cs
+++ b/Engine/Source/Programs/Horde/Horde.Agent/Execution/JobExecutor.cs
@@ -1111,11 +1111,44 @@ namespace Horde.Agent.Execution
 			}
 		}
 
+		static FileReference GetLeaseCleanupScript(DirectoryReference workspaceDir)
+		{
+			if (RuntimeInformation.IsOSPlatform(OSPlatform.Windows))
+			{
+				return FileReference.Combine(workspaceDir, "CleanupLease.bat");
+			}
+			else
+			{
+				return FileReference.Combine(workspaceDir, "CleanupLease.sh");
+			}
+		}
+
 		internal IReadOnlyDictionary<string, string> GetEnvVars()
 		{
 			return new Dictionary<string, string>(_envVars);
 		}
-		
+
+		protected static async Task ExecuteLeaseCleanupScriptAsync(DirectoryReference workspaceDir, ILogger logger)
+		{
+			using (LogParser parser = new LogParser(logger, new List<string>()))
+			{
+				FileReference leaseCleanupScript = GetLeaseCleanupScript(workspaceDir);
+				await ExecuteCleanupScriptAsync(leaseCleanupScript, parser, logger);
+			}
+		}
+
+		protected async Task TerminateProcessesAsync(TerminateCondition condition, ILogger logger)
+		{
+			try
+			{
+				await _session.TerminateProcessesAsync(condition, logger, CancellationToken.None);
+			}
+			catch (Exception ex)
+			{
+				logger.LogWarning(ex, "Exception while terminating processes: {Message}", ex.Message);
+			}
+		}
+
 		static async Task ExecuteCleanupScriptAsync(FileReference cleanupScript, LogParser filter, ILogger logger)
 		{
 			if (FileReference.Exists(cleanupScript))
@@ -1317,6 +1350,10 @@ namespace Horde.Agent.Execution
 			FileReference cleanupScript = GetCleanupScript(workspaceDir);
 			newEnvVars["UE_HORDE_CLEANUP"] = cleanupScript.FullName;
 
+			// Pass the location of the cleanup script to the job
+			FileReference leaseCleanupScript = GetLeaseCleanupScript(workspaceDir);
+			newEnvVars["UE_HORDE_LEASE_CLEANUP"] = cleanupScript.FullName;
+
 			// Set up the shared working dir
 			newEnvVars["UE_HORDE_SHARED_DIR"] = DirectoryReference.Combine(_session.WorkingDir, "Saved").FullName;
 
diff --git a/Engine/Source/Programs/Horde/Horde.Agent/Execution/PerforceExecutor.cs b/Engine/Source/Programs/Horde/Horde.Agent/Execution/PerforceExecutor.cs
--- a/Engine/Source/Programs/Horde/Horde.Agent/Execution/PerforceExecutor.cs
+++ b/Engine/Source/Programs/Horde/Horde.Agent/Execution/PerforceExecutor.cs
@@ -15,6 +15,7 @@ using EpicGames.Core;
 using EpicGames.Horde.Storage;
 using EpicGames.Perforce;
 using EpicGames.Perforce.Managed;
+using Horde.Agent.Parser;
 using Horde.Agent.Services;
 using Horde.Agent.Utility;
 using Horde.Common.Rpc;
@@ -227,6 +228,9 @@ namespace Horde.Agent.Execution
 
 		public override async Task FinalizeAsync(ILogger logger, CancellationToken cancellationToken)
 		{
+			await ExecuteLeaseCleanupScriptAsync(_workspace.WorkspaceDir, logger);
+			await TerminateProcessesAsync(TerminateCondition.AfterBatch, logger);
+
 			await _workspace.CleanAsync(cancellationToken);
 			await base.FinalizeAsync(logger, cancellationToken);
 		}
diff --git a/Engine/Source/Programs/Horde/Horde.Agent/Execution/WorkspaceExecutor.cs b/Engine/Source/Programs/Horde/Horde.Agent/Execution/WorkspaceExecutor.cs
--- a/Engine/Source/Programs/Horde/Horde.Agent/Execution/WorkspaceExecutor.cs
+++ b/Engine/Source/Programs/Horde/Horde.Agent/Execution/WorkspaceExecutor.cs
@@ -118,6 +118,10 @@ namespace Horde.Agent.Execution
 		/// <inheritdoc/>
 		public override async Task FinalizeAsync(ILogger logger, CancellationToken cancellationToken)
 		{
+			DirectoryReference workspaceDir = (await _workspace.GetSettingsAsync(cancellationToken)).DirectoryPath;
+			await ExecuteLeaseCleanupScriptAsync(workspaceDir, logger);
+			await TerminateProcessesAsync(TerminateCondition.AfterBatch, logger);
+
 			if (_autoSdkWorkspace != null)
 			{
 				await _autoSdkWorkspace.FinalizeAsync(cancellationToken);
