Horde: Add job template user pref persistence
#rnx
#skipci

diff --git a/Engine/Source/Programs/Horde/HordeDashboard/src/backend/Api.ts b/Engine/Source/Programs/Horde/HordeDashboard/src/backend/Api.ts
--- a/Engine/Source/Programs/Horde/HordeDashboard/src/backend/Api.ts
+++ b/Engine/Source/Programs/Horde/HordeDashboard/src/backend/Api.ts
@@ -3100,16 +3100,29 @@ export enum DashboardPreference {
 	ShowPreflights = "ShowPreflights"
 }
 
+export type JobTemplatePreference = {		
+	streamId: string;
+	templateId: string;	
+	templateVersion: string;
+
+	// base64 encoded string array
+	arguments: string;
+
+	lastUsedUTC: string;
+	prefVersion: string;
+}
+
+
 export type DashboardSettings = {
 
 	preferences: Map<DashboardPreference, string>;
+	jobTemplatePrefs?: JobTemplatePreference[];
 
 }
 
 export type UpdateDashboardSettings = {
-
 	preferences?: Record<DashboardPreference, string>;
-
+	jobTemplatePrefs?: JobTemplatePreference[];
 }
 
 /** Settings for whether various features should be enabled on the dashboard */
diff --git a/Engine/Source/Programs/Horde/HordeDashboard/src/backend/Dashboard.ts b/Engine/Source/Programs/Horde/HordeDashboard/src/backend/Dashboard.ts
--- a/Engine/Source/Programs/Horde/HordeDashboard/src/backend/Dashboard.ts
+++ b/Engine/Source/Programs/Horde/HordeDashboard/src/backend/Dashboard.ts
@@ -3,7 +3,7 @@
 import { getTheme } from "@fluentui/react";
 import { action, makeObservable, observable } from 'mobx';
 import backend from '.';
-import { DashboardPreference, GetDashboardConfigResponse, GetUserResponse, UserClaim } from './Api';
+import { DashboardPreference, GetDashboardConfigResponse, GetUserResponse, JobTemplatePreference, UserClaim } from './Api';
 
 const theme = getTheme();
 
@@ -25,6 +25,8 @@ export enum WebBrowser {
     Other = "Other"
 }
 
+export const jobTemplatePreferenceVersion = "1";
+
 export class Dashboard {
 
     constructor() {
@@ -114,23 +116,23 @@ export class Dashboard {
     }
 
     get preview(): boolean {
-        
-        try {            
+
+        try {
             return !!(process.env.REACT_APP_HORDE_DEBUG_PREVIEW) || window?.location?.hostname?.indexOf("preview") !== -1;
         } catch (reason) {
             console.error(reason);
         }
-        return false;        
+        return false;
     }
 
     get development(): boolean {
 
-        try {            
+        try {
             return window?.location?.hostname?.indexOf("devtools-dev") !== -1;
         } catch (reason) {
             console.error(reason);
         }
-        return false;        
+        return false;
     }
 
     get browser(): WebBrowser {
@@ -210,6 +212,10 @@ export class Dashboard {
 
     }
 
+    get templatePreferences(): JobTemplatePreference[] {
+        return this.data.dashboardSettings?.jobTemplatePrefs ?? [];
+    }
+
     get darktheme(): boolean {
 
         if (!this.available) {
@@ -281,9 +287,9 @@ export class Dashboard {
             return true;
         }
         */
-        
+
         const value = this.preferences.get(DashboardPreference.LocalCache) !== 'false';
-        
+
         if (value && !this.hasLoggedLocalCache) {
             this.hasLoggedLocalCache = true;
             console.log("Local graph and template caching is enabled")
@@ -484,6 +490,113 @@ export class Dashboard {
         }
     }
 
+    getJobTemplatePref(streamId: string, templateIds: string[]): JobTemplatePreference | undefined {
+
+        let templatePrefs = this.data.dashboardSettings?.jobTemplatePrefs ?? [];
+        templatePrefs = templatePrefs.filter(p => p.streamId === streamId && templateIds.indexOf(p.templateId) !== -1);
+        const dates = new Map<string, Date>();
+
+        templatePrefs.forEach(p => {
+            dates.set(p.streamId + p.templateId, new Date(p.lastUsedUTC));
+        })
+
+        templatePrefs = templatePrefs.sort((a, b) => dates.get(b.streamId + b.templateId)!.getTime() - dates.get(a.streamId + a.templateId)!.getTime());
+
+        if (templatePrefs.length > 0) {
+            return templatePrefs[0];
+        }
+
+        return undefined;
+    }
+
+    async registerJobTemplatePref(streamId: string, templateId: string, templateVersion: string, args: string[]) {
+
+        // ignore non-preflight templates while get this working, if persisting arguments
+        /*
+        if (templateId.indexOf("presubmit") === -1) {
+            console.log("not a presubmit")
+            return;
+        }
+        */
+
+        let pref: JobTemplatePreference | undefined;
+
+        let templatePrefs = this.data.dashboardSettings?.jobTemplatePrefs ?? [];
+
+        const existing = templatePrefs.find(tp => tp.streamId === streamId && tp.templateId === templateId);
+
+        if (existing) {
+
+            pref = existing;
+
+            try {
+                existing.arguments = btoa(JSON.stringify(args));
+            } catch (msg) {
+                console.error(msg);
+                return;
+            }
+
+            existing.templateVersion = templateVersion;
+            existing.lastUsedUTC = new Date().toISOString();
+            existing.prefVersion = jobTemplatePreferenceVersion;
+
+        } else {
+
+            pref = {
+                streamId: streamId,
+                templateId: templateId,
+                templateVersion: templateVersion,
+                arguments: "",
+                lastUsedUTC: new Date().toISOString(),
+                prefVersion: jobTemplatePreferenceVersion
+            }
+
+            try {
+                pref.arguments = btoa(JSON.stringify(args));
+            } catch (msg) {
+                console.error(msg);
+                return;
+            }
+
+            templatePrefs.push(pref);
+        }
+
+        const dates = new Map<string, Date>();
+
+        templatePrefs.forEach(p => {
+            dates.set(p.streamId + p.templateId, new Date(p.lastUsedUTC));
+        })
+
+        templatePrefs = templatePrefs.sort((a, b) => dates.get(b.streamId + b.templateId)!.getTime() - dates.get(a.streamId + a.templateId)!.getTime());
+
+        const dupes = new Set<string>();
+
+        templatePrefs = templatePrefs.filter((a) => {
+            if (dupes.has(a.streamId + a.templateId)) {
+                return false;
+            }
+
+            dupes.add(a.streamId + a.templateId);
+            return true;
+        });
+
+        templatePrefs = templatePrefs.slice(0, 256);
+
+        if (!this.data.dashboardSettings?.preferences) {
+            this.data.dashboardSettings = { preferences: new Map() };
+        }
+
+        this.data.dashboardSettings.jobTemplatePrefs = templatePrefs;
+
+        try {
+            await this.postPreferences();
+        } catch (msg) {
+            console.error(msg);
+        }
+        
+
+    }
+
     private setPreference(pref: DashboardPreference, value: string | undefined): void {
 
         if (!this.available) {
@@ -540,7 +653,7 @@ export class Dashboard {
 
         let success = true;
         try {
-            await backend.updateUser({ dashboardSettings: { preferences: data } });
+            await backend.updateUser({ dashboardSettings: { preferences: data, jobTemplatePrefs: this.templatePreferences ?? [] } });
         } catch (reason) {
             success = false;
             console.error("Error posting user preferences", reason)
diff --git a/Engine/Source/Programs/Horde/HordeDashboard/src/components/NewBuild.tsx b/Engine/Source/Programs/Horde/HordeDashboard/src/components/NewBuild.tsx
--- a/Engine/Source/Programs/Horde/HordeDashboard/src/components/NewBuild.tsx
+++ b/Engine/Source/Programs/Horde/HordeDashboard/src/components/NewBuild.tsx
@@ -12,6 +12,7 @@ import { ErrorHandler } from "../components/ErrorHandler";
 import { hordeClasses, modeColors } from '../styles/Styles';
 import { useQuery } from './JobDetailCommon';
 import { JobDetailsV2 } from './jobDetailsV2/JobDetailsViewCommon';
+import dashboard from '../backend/Dashboard';
 
 let toolTipId = 0;
 
@@ -787,6 +788,12 @@ export const NewBuild: React.FC<{ streamId: string; show: boolean; onClose: (new
             console.error(`Stream default preflight template cannot be found for stream ${stream.fullname} : stream defaultPreflightTemplate ${stream.defaultPreflight?.templateId}, will use first template in list`);
          }
       }
+      if (!t) {
+         const pref = dashboard.getJobTemplatePref(stream.id, templates.map(t => t.id));
+         if (pref) {
+            t = templates.find(t => stream.id === pref.streamId && t.id === pref.templateId)
+         }
+      }
 
       // default to sane template when all are shown
       if (!t && (showAllTemplates) && stream.tabs.length > 0) {
@@ -1375,6 +1382,8 @@ export const NewBuild: React.FC<{ streamId: string; show: boolean; onClose: (new
             console.log("Debug Job Submit");
             console.log(JSON.stringify(data));
 
+            await dashboard.registerJobTemplatePref(data.streamId, data.templateId, template.hash, data.arguments ?? [])
+
          } else {
 
 
@@ -1385,17 +1394,22 @@ export const NewBuild: React.FC<{ streamId: string; show: boolean; onClose: (new
 
             let redirected = false;
 
-            await backend.createJob(data).then(async (data) => {
-               console.log(`Job created: ${JSON.stringify(data)}`);
+            await backend.createJob(data).then(async (response) => {
+               console.log(`Job created: ${JSON.stringify(response)}`);
                console.log("Updating notifications")
                try {
-                  await backend.updateNotification({ slack: true }, "job", data.id);
+                  await backend.updateNotification({ slack: true }, "job", response.id);
                } catch (reason) {
                   console.log(`Error on updating notifications: ${reason}`);
                }
+               try {
+                  await dashboard.registerJobTemplatePref(data.streamId, data.templateId, template.hash, data.arguments ?? [])
+               } catch (reason) {
+                  console.log(`Error on regisyering job template pref`);
+               }               
 
                redirected = true;
-               onClose(data.id);
+               onClose(response.id);
             }).catch(reason => {
                // "Not Found" is generally a permissions error
                errorReason = reason ? reason : "Unknown";
