Horde: Fix incorrect path being set for lease cleanup script.

diff --git a/Engine/Source/Programs/Horde/Horde.Agent/Execution/JobExecutor.cs b/Engine/Source/Programs/Horde/Horde.Agent/Execution/JobExecutor.cs
--- a/Engine/Source/Programs/Horde/Horde.Agent/Execution/JobExecutor.cs
+++ b/Engine/Source/Programs/Horde/Horde.Agent/Execution/JobExecutor.cs
@@ -1352,7 +1352,7 @@ namespace Horde.Agent.Execution
 
 			// Pass the location of the cleanup script to the job
 			FileReference leaseCleanupScript = GetLeaseCleanupScript(workspaceDir);
-			newEnvVars["UE_HORDE_LEASE_CLEANUP"] = cleanupScript.FullName;
+			newEnvVars["UE_HORDE_LEASE_CLEANUP"] = leaseCleanupScript.FullName;
 
 			// Set up the shared working dir
 			newEnvVars["UE_HORDE_SHARED_DIR"] = DirectoryReference.Combine(_session.WorkingDir, "Saved").FullName;
