Horde: Fix static analyzer warning.

diff --git a/Engine/Source/Programs/Horde/Horde.Server.Tests/Fleet/AwsAutoScalingLifecycleTest.cs b/Engine/Source/Programs/Horde/Horde.Server.Tests/Fleet/AwsAutoScalingLifecycleTest.cs
--- a/Engine/Source/Programs/Horde/Horde.Server.Tests/Fleet/AwsAutoScalingLifecycleTest.cs
+++ b/Engine/Source/Programs/Horde/Horde.Server.Tests/Fleet/AwsAutoScalingLifecycleTest.cs
@@ -63,7 +63,17 @@ public class AwsAutoScalingLifecycleServiceTest : TestSetup
 	private AwsAutoScalingLifecycleService _asgLifecycleService = default!;
 	private Mock<IAmazonAutoScaling> _asgMock = default!;
 	private readonly List<CompleteLifecycleActionRequest> _lifecycleUpdates = new();
-	
+
+	protected override void Dispose(bool disposing)
+	{
+		if (disposing)
+		{
+			_asgLifecycleService?.Dispose();
+		}
+
+		base.Dispose(disposing);
+	}
+
 	[TestInitialize]
 	public async Task Setup()
 	{
