Horde: Remove a bundle-specific method from the IStorageClient interface.

diff --git a/Engine/Source/Programs/Horde/Horde.Server.Tests/BlobStoreTests.cs b/Engine/Source/Programs/Horde/Horde.Server.Tests/BlobStoreTests.cs
--- a/Engine/Source/Programs/Horde/Horde.Server.Tests/BlobStoreTests.cs
+++ b/Engine/Source/Programs/Horde/Horde.Server.Tests/BlobStoreTests.cs
@@ -124,13 +124,6 @@ namespace Horde.Server.Tests
 				NodeHandle refTarget = await store.ReadRefTargetAsync(refName);
 				Assert.AreEqual(locator3, refTarget.GetLocator().Blob);
 			}
-
-			RefName refName2 = new RefName("hello2");
-
-			NodeHandle refTargetId2 = await store.WriteRefAsync(refName2, CreateTestBundle(input3, new BlobLocator[] { locator1, locator2 }), 0);
-			Blob refTarget2 = await ReadBlobAsync(store, refTargetId2.GetLocator().Blob);
-			Assert.IsTrue(refTarget2.Data.Span.SequenceEqual(input3));
-			Assert.IsTrue(refTarget2.References.SequenceEqual(new BlobLocator[] { locator1, locator2 }));
 		}
 
 		[TestMethod]
diff --git a/Engine/Source/Programs/Horde/Horde.Server/Commands/Bundles/BundleWriteTest.cs b/Engine/Source/Programs/Horde/Horde.Server/Commands/Bundles/BundleWriteTest.cs
--- a/Engine/Source/Programs/Horde/Horde.Server/Commands/Bundles/BundleWriteTest.cs
+++ b/Engine/Source/Programs/Horde/Horde.Server/Commands/Bundles/BundleWriteTest.cs
@@ -35,7 +35,6 @@ namespace Horde.Server.Commands.Bundles
 			public IAsyncEnumerable<NodeHandle> FindNodesAsync(Utf8String alias, CancellationToken cancellationToken = default) => throw new NotImplementedException();
 			public Task<NodeHandle?> TryReadRefTargetAsync(RefName name, RefCacheTime cacheTime = default, CancellationToken cancellationToken = default) => throw new NotImplementedException();
 			public Task<BlobLocator> WriteBlobAsync(Stream stream, Utf8String prefix = default, CancellationToken cancellationToken = default) => Task.FromResult(BlobLocator.Create(HostId.Empty));
-			public Task<NodeHandle> WriteRefAsync(RefName name, Bundle bundle, int exportIdx, Utf8String prefix = default, RefOptions? options = null, CancellationToken cancellationToken = default) => throw new NotImplementedException();
 			public Task WriteRefTargetAsync(RefName name, NodeHandle target, RefOptions? options = null, CancellationToken cancellationToken = default) => Task.CompletedTask;
 
 			public IStorageWriter CreateWriter(RefName refName = default, TreeOptions? options = null) => new TreeWriter(this, _reader, options);
diff --git a/Engine/Source/Programs/Shared/EpicGames.Horde/Storage/BundleStorageClient.cs b/Engine/Source/Programs/Shared/EpicGames.Horde/Storage/BundleStorageClient.cs
--- a/Engine/Source/Programs/Shared/EpicGames.Horde/Storage/BundleStorageClient.cs
+++ b/Engine/Source/Programs/Shared/EpicGames.Horde/Storage/BundleStorageClient.cs
@@ -73,15 +73,6 @@ namespace EpicGames.Horde.Storage
 		/// <inheritdoc/>
 		public abstract Task<NodeHandle?> TryReadRefTargetAsync(RefName name, RefCacheTime cacheTime = default, CancellationToken cancellationToken = default);
 
-		/// <inheritdoc/>
-		public virtual async Task<NodeHandle> WriteRefAsync(RefName name, Bundle bundle, int exportIdx = 0, Utf8String prefix = default, RefOptions? options = null, CancellationToken cancellationToken = default)
-		{
-			BlobLocator locator = await this.WriteBundleAsync(bundle, prefix, cancellationToken);
-			NodeHandle target = new FlushedNodeHandle(TreeReader, bundle.Header.Exports[exportIdx].Hash, new NodeLocator(locator, exportIdx));
-			await WriteRefTargetAsync(name, target, options, cancellationToken);
-			return target;
-		}
-
 		/// <inheritdoc/>
 		public abstract Task WriteRefTargetAsync(RefName name, NodeHandle target, RefOptions? options = null, CancellationToken cancellationToken = default);
 
diff --git a/Engine/Source/Programs/Shared/EpicGames.Horde/Storage/IStorageClient.cs b/Engine/Source/Programs/Shared/EpicGames.Horde/Storage/IStorageClient.cs
--- a/Engine/Source/Programs/Shared/EpicGames.Horde/Storage/IStorageClient.cs
+++ b/Engine/Source/Programs/Shared/EpicGames.Horde/Storage/IStorageClient.cs
@@ -160,18 +160,6 @@ namespace EpicGames.Horde.Storage
 		/// <returns>Node pointed to by the ref</returns>
 		Task<NodeHandle?> TryReadRefTargetAsync(RefName name, RefCacheTime cacheTime = default, CancellationToken cancellationToken = default);
 
-		/// <summary>
-		/// Writes a new ref to the store which points to a new blob
-		/// </summary>
-		/// <param name="name">Ref to write</param>
-		/// <param name="bundle">The bundle to write</param>
-		/// <param name="exportIdx">Index of the export in the bundle to be the root of the tree</param>
-		/// <param name="prefix">Prefix for blob names.</param>
-		/// <param name="options">Options for the new ref</param>
-		/// <param name="cancellationToken">Cancellation token for the operation</param>
-		/// <returns>Unique identifier for the blob</returns>
-		Task<NodeHandle> WriteRefAsync(RefName name, Bundle bundle, int exportIdx, Utf8String prefix = default, RefOptions? options = null, CancellationToken cancellationToken = default);
-
 		/// <summary>
 		/// Writes a new ref to the store
 		/// </summary>
