Horde: Listen for AWS auto-scaling events via SQS

The previous HTTP callback-based method is removed in favor of this.

diff --git a/Engine/Source/Programs/Horde/Horde.Server.Tests/Fleet/AwsAutoScalingLifecycleTest.cs b/Engine/Source/Programs/Horde/Horde.Server.Tests/Fleet/AwsAutoScalingLifecycleTest.cs
--- a/Engine/Source/Programs/Horde/Horde.Server.Tests/Fleet/AwsAutoScalingLifecycleTest.cs
+++ b/Engine/Source/Programs/Horde/Horde.Server.Tests/Fleet/AwsAutoScalingLifecycleTest.cs
@@ -19,56 +19,21 @@ using Moq;
 
 namespace Horde.Server.Tests.Fleet;
 
-[TestClass]
-public class LifecycleHookMessageTest
-{
-	[TestMethod]
-	public void DeserializeLifecycleHookRequest()
-	{
-		string rawText = @"
-{
-	""version"": ""0"",
-	""id"": ""d2b0a72e-03a4-11ee-be56-0242ac120002"",
-	""detail-type"": ""EC2 Instance-terminate Lifecycle Action"",
-	""source"": ""aws.autoscaling"",
-	""account"": ""111111111111"",
-	""time"": ""2023-06-05T10:48:08Z"",
-	""region"": ""us-east-1"",
-	""resources"": [""arn:aws:autoscaling:us-east-1:111111111111:autoScalingGroup:d816d436-03a4-11ee-be56-0242ac120002:autoScalingGroupName/my-asg-name""],
-	""detail"": {
-		""LifecycleActionToken"": ""f191f274-03a4-11ee-be56-0242ac120002"",
-		""AutoScalingGroupName"": ""my-asg-name"",
-		""LifecycleHookName"": ""my-hook-name"",
-		""EC2InstanceId"": ""i-11234567890abcdef"",
-		""LifecycleTransition"": ""autoscaling:EC2_INSTANCE_TERMINATING"",
-		""Origin"": ""AutoScalingGroup"",
-		""Destination"": ""EC2""
-	}
-}
-";
-
-		LifecycleActionRequest? req = JsonSerializer.Deserialize<LifecycleActionRequest>(rawText);
-		Assert.AreEqual("f191f274-03a4-11ee-be56-0242ac120002", req!.Event.LifecycleActionToken);
-		Assert.AreEqual("my-asg-name", req!.Event.AutoScalingGroupName);
-		Assert.AreEqual("my-hook-name", req!.Event.LifecycleHookName);
-		Assert.AreEqual("i-11234567890abcdef", req!.Event.Ec2InstanceId);
-		Assert.AreEqual("autoscaling:EC2_INSTANCE_TERMINATING", req!.Event.LifecycleTransition);
-		Assert.AreEqual("AutoScalingGroup", req!.Event.Origin);
-	}
-}
-
 [TestClass]
 public class AwsAutoScalingLifecycleServiceTest : TestSetup
 {
 	private AwsAutoScalingLifecycleService _asgLifecycleService = default!;
 	private Mock<IAmazonAutoScaling> _asgMock = default!;
 	private readonly List<CompleteLifecycleActionRequest> _lifecycleUpdates = new();
+	private readonly FakeAmazonSqs _fakeSqs = new();
+	private readonly string _queueUrl = "https://sqs.us-east-1.amazonaws.com/123456789/MyQueue";
 
 	protected override void Dispose(bool disposing)
 	{
 		if (disposing)
 		{
 			_asgLifecycleService?.Dispose();
+			_fakeSqs?.Dispose();
 		}
 
 		base.Dispose(disposing);
@@ -90,21 +55,23 @@ public class AwsAutoScalingLifecycleServiceTest : TestSetup
 			.Setup(x => x.CompleteLifecycleActionAsync(It.IsAny<CompleteLifecycleActionRequest>(), It.IsAny<CancellationToken>()))
 			.Returns(OnCompleteLifecycleActionAsync);
 		
-		_asgLifecycleService = new AwsAutoScalingLifecycleService(AgentService, GetRedisServiceSingleton(), AgentCollection, Clock, ServiceProvider, Tracer, logger);
-		_asgLifecycleService.SetAmazonAutoScalingClient(_asgMock.Object);
+		ServerSettings ss = new () { AwsAutoScalingQueueUrl = _queueUrl };
+		_asgLifecycleService = new AwsAutoScalingLifecycleService(AgentService, GetRedisServiceSingleton(), AgentCollection, Clock, new TestOptionsMonitor<ServerSettings>(ss), ServiceProvider, Tracer, logger);
+		_asgLifecycleService.SetAmazonClientsTesting(_asgMock.Object, _fakeSqs);
 		await _asgLifecycleService.StartAsync(CancellationToken.None);
 	}
-
+	
 	[TestMethod]
 	public async Task TerminationRequested_WithAgentInService_MarkedForShutdown()
 	{
 		// Arrange
 		LifecycleActionEvent lae = new() { Ec2InstanceId = "i-1234", LifecycleActionToken = "action-token-test", Origin = "AutoScalingGroup" };
 		IAgent agent = await CreateAgentAsync(new PoolId("pool1"), properties: new() { KnownPropertyNames.AwsInstanceId + "=" + lae.Ec2InstanceId });
-
-		// Act
-		await AwsAsgLifecycleService.InitiateTerminationAsync(lae, CancellationToken.None);
 		
+		// Act
+		await _fakeSqs.SendMessageAsync(_queueUrl, JsonSerializer.Serialize(lae));
+		await _asgLifecycleService.ReceiveLifecycleEventsAsync(CancellationToken.None);
+
 		// Assert
 		agent = (await AgentService.GetAgentAsync(agent.Id))!;
 		Assert.IsTrue(agent.RequestShutdown);
@@ -166,6 +133,34 @@ public class AwsAutoScalingLifecycleServiceTest : TestSetup
 		// Assert
 		AssertLifecycleUpdate(lae, AwsAutoScalingLifecycleService.ActionAbandon, _lifecycleUpdates[0]);
 	}
+	
+	[TestMethod]
+	public void DeserializeLifecycleActionEvent()
+	{
+		string rawText = @"
+{
+	""Origin"": ""AutoScalingGroup"",
+	""LifecycleHookName"": ""my-hook-name"",
+	""Destination"": ""EC2"",
+	""AccountId"": ""123456789"",
+	""RequestId"": ""5ab929ac-06d9-11ee-be56-0242ac120002"",
+	""LifecycleTransition"": ""autoscaling:EC2_INSTANCE_TERMINATING"",
+	""AutoScalingGroupName"": ""my-asg-name"",
+	""Service"": ""AWS Auto Scaling"",
+	""Time"": ""2023-06-09T11:11:20.345Z"",
+	""EC2InstanceId"": ""i-123456789"",
+	""LifecycleActionToken"": ""5ec3537e-06d9-11ee-be56-0242ac120002""
+}
+";
+
+		LifecycleActionEvent? ev = JsonSerializer.Deserialize<LifecycleActionEvent>(rawText);
+		Assert.AreEqual("5ec3537e-06d9-11ee-be56-0242ac120002", ev!.LifecycleActionToken);
+		Assert.AreEqual("my-asg-name", ev.AutoScalingGroupName);
+		Assert.AreEqual("my-hook-name", ev.LifecycleHookName);
+		Assert.AreEqual("i-123456789", ev.Ec2InstanceId);
+		Assert.AreEqual("autoscaling:EC2_INSTANCE_TERMINATING", ev.LifecycleTransition);
+		Assert.AreEqual("AutoScalingGroup", ev.Origin);
+	}
 
 	private static void AssertLifecycleUpdate(LifecycleActionEvent expectedEvent, string expectedResult, CompleteLifecycleActionRequest actual)
 	{
diff --git a/Engine/Source/Programs/Horde/Horde.Server.Tests/Fleet/FakeAmazonSqs.cs b/Engine/Source/Programs/Horde/Horde.Server.Tests/Fleet/FakeAmazonSqs.cs
new file mode 100644
--- /dev/null
+++ b/Engine/Source/Programs/Horde/Horde.Server.Tests/Fleet/FakeAmazonSqs.cs
@@ -0,0 +1,246 @@
+// Copyright Epic Games, Inc. All Rights Reserved.
+
+using System;
+using System.Collections.Concurrent;
+using System.Collections.Generic;
+using System.Linq;
+using System.Threading;
+using System.Threading.Tasks;
+using Amazon.Runtime;
+using Amazon.SQS;
+using Amazon.SQS.Model;
+
+namespace Horde.Server.Tests.Fleet;
+
+#pragma warning disable CA1054
+
+/// <summary>
+/// A fake implementation of IAmazonSQS to aid testing
+/// Not a fully valid implementation but enough to pass tests in Horde
+/// </summary>
+public sealed class FakeAmazonSqs : IAmazonSQS
+{
+	private readonly ConcurrentDictionary<string, List<string>> _queues = new ();
+
+	public void Dispose()
+	{
+	}
+
+	public Task<SendMessageResponse> SendMessageAsync(string queueUrl, string messageBody, CancellationToken cancellationToken = default)
+	{
+		if (!_queues.TryGetValue(queueUrl, out List<string>? messages))
+		{
+			messages = new();
+			_queues[queueUrl] = messages;
+		}
+		
+		messages.Add(messageBody);
+		return Task.FromResult(new SendMessageResponse() { MessageId = "unset" });
+	}
+	
+	public Task<ReceiveMessageResponse> ReceiveMessageAsync(ReceiveMessageRequest request, CancellationToken cancellationToken = default)
+	{
+		if (_queues.TryGetValue(request.QueueUrl, out List<string>? messages))
+		{
+			return Task.FromResult(new ReceiveMessageResponse
+			{
+				Messages = messages.Select(body => new Message() { Body = body, MessageId = "unset" }).ToList()
+			});
+		}
+
+		return Task.FromResult(new ReceiveMessageResponse());
+	}
+	
+	public Task<DeleteMessageResponse> DeleteMessageAsync(DeleteMessageRequest request, CancellationToken cancellationToken = default)
+	{
+		return Task.FromResult(new DeleteMessageResponse());
+	}
+	
+#region Not implemented
+
+public ISQSPaginatorFactory Paginators { get; } = null!;
+	public IClientConfig Config { get; } = null!;
+	
+	public Task<Dictionary<string, string>> GetAttributesAsync(string queueUrl)
+	{
+		throw new NotImplementedException();
+	}
+
+	public Task SetAttributesAsync(string queueUrl, Dictionary<string, string> attributes)
+	{
+		throw new NotImplementedException();
+	}
+
+	public Task<string> AuthorizeS3ToSendMessageAsync(string queueUrl, string bucket)
+	{
+		throw new NotImplementedException();
+	}
+
+	public Task<AddPermissionResponse> AddPermissionAsync(string queueUrl, string label, List<string> awsAccountIds, List<string> actions, CancellationToken cancellationToken = default)
+	{
+		throw new NotImplementedException();
+	}
+
+	public Task<AddPermissionResponse> AddPermissionAsync(AddPermissionRequest request, CancellationToken cancellationToken = default)
+	{
+		throw new NotImplementedException();
+	}
+
+	public Task<ChangeMessageVisibilityResponse> ChangeMessageVisibilityAsync(string queueUrl, string receiptHandle, int visibilityTimeout, CancellationToken cancellationToken = default)
+	{
+		throw new NotImplementedException();
+	}
+
+	public Task<ChangeMessageVisibilityResponse> ChangeMessageVisibilityAsync(ChangeMessageVisibilityRequest request, CancellationToken cancellationToken = default)
+	{
+		throw new NotImplementedException();
+	}
+
+	public Task<ChangeMessageVisibilityBatchResponse> ChangeMessageVisibilityBatchAsync(string queueUrl, List<ChangeMessageVisibilityBatchRequestEntry> entries, CancellationToken cancellationToken = default)
+	{
+		throw new NotImplementedException();
+	}
+
+	public Task<ChangeMessageVisibilityBatchResponse> ChangeMessageVisibilityBatchAsync(ChangeMessageVisibilityBatchRequest request, CancellationToken cancellationToken = default)
+	{
+		throw new NotImplementedException();
+	}
+
+	public Task<CreateQueueResponse> CreateQueueAsync(string queueName, CancellationToken cancellationToken = default)
+	{
+		throw new NotImplementedException();
+	}
+
+	public Task<CreateQueueResponse> CreateQueueAsync(CreateQueueRequest request, CancellationToken cancellationToken = default)
+	{
+		throw new NotImplementedException();
+	}
+
+	public Task<DeleteMessageResponse> DeleteMessageAsync(string queueUrl, string receiptHandle, CancellationToken cancellationToken = default)
+	{
+		throw new NotImplementedException();
+	}
+
+	public Task<DeleteMessageBatchResponse> DeleteMessageBatchAsync(string queueUrl, List<DeleteMessageBatchRequestEntry> entries, CancellationToken cancellationToken = default)
+	{
+		throw new NotImplementedException();
+	}
+
+	public Task<DeleteMessageBatchResponse> DeleteMessageBatchAsync(DeleteMessageBatchRequest request, CancellationToken cancellationToken = default)
+	{
+		throw new NotImplementedException();
+	}
+
+	public Task<DeleteQueueResponse> DeleteQueueAsync(string queueUrl, CancellationToken cancellationToken = default)
+	{
+		throw new NotImplementedException();
+	}
+
+	public Task<DeleteQueueResponse> DeleteQueueAsync(DeleteQueueRequest request, CancellationToken cancellationToken = default)
+	{
+		throw new NotImplementedException();
+	}
+
+	public Task<GetQueueAttributesResponse> GetQueueAttributesAsync(string queueUrl, List<string> attributeNames, CancellationToken cancellationToken = default)
+	{
+		throw new NotImplementedException();
+	}
+
+	public Task<GetQueueAttributesResponse> GetQueueAttributesAsync(GetQueueAttributesRequest request, CancellationToken cancellationToken = default)
+	{
+		throw new NotImplementedException();
+	}
+
+	public Task<GetQueueUrlResponse> GetQueueUrlAsync(string queueName, CancellationToken cancellationToken = default)
+	{
+		throw new NotImplementedException();
+	}
+
+	public Task<GetQueueUrlResponse> GetQueueUrlAsync(GetQueueUrlRequest request, CancellationToken cancellationToken = default)
+	{
+		throw new NotImplementedException();
+	}
+
+	public Task<ListDeadLetterSourceQueuesResponse> ListDeadLetterSourceQueuesAsync(ListDeadLetterSourceQueuesRequest request, CancellationToken cancellationToken = default)
+	{
+		throw new NotImplementedException();
+	}
+
+	public Task<ListQueuesResponse> ListQueuesAsync(string queueNamePrefix, CancellationToken cancellationToken = default)
+	{
+		throw new NotImplementedException();
+	}
+
+	public Task<ListQueuesResponse> ListQueuesAsync(ListQueuesRequest request, CancellationToken cancellationToken = default)
+	{
+		throw new NotImplementedException();
+	}
+
+	public Task<ListQueueTagsResponse> ListQueueTagsAsync(ListQueueTagsRequest request, CancellationToken cancellationToken = default)
+	{
+		throw new NotImplementedException();
+	}
+
+	public Task<PurgeQueueResponse> PurgeQueueAsync(string queueUrl, CancellationToken cancellationToken = default)
+	{
+		throw new NotImplementedException();
+	}
+
+	public Task<PurgeQueueResponse> PurgeQueueAsync(PurgeQueueRequest request, CancellationToken cancellationToken = default)
+	{
+		throw new NotImplementedException();
+	}
+
+	public Task<ReceiveMessageResponse> ReceiveMessageAsync(string queueUrl, CancellationToken cancellationToken = default)
+	{
+		throw new NotImplementedException();
+	}
+
+	public Task<RemovePermissionResponse> RemovePermissionAsync(string queueUrl, string label, CancellationToken cancellationToken = default)
+	{
+		throw new NotImplementedException();
+	}
+
+	public Task<RemovePermissionResponse> RemovePermissionAsync(RemovePermissionRequest request, CancellationToken cancellationToken = default)
+	{
+		throw new NotImplementedException();
+	}
+
+	public Task<SendMessageResponse> SendMessageAsync(SendMessageRequest request, CancellationToken cancellationToken = default)
+	{
+		throw new NotImplementedException();
+	}
+
+	public Task<SendMessageBatchResponse> SendMessageBatchAsync(string queueUrl, List<SendMessageBatchRequestEntry> entries, CancellationToken cancellationToken = default)
+	{
+		throw new NotImplementedException();
+	}
+
+	public Task<SendMessageBatchResponse> SendMessageBatchAsync(SendMessageBatchRequest request, CancellationToken cancellationToken = default)
+	{
+		throw new NotImplementedException();
+	}
+
+	public Task<SetQueueAttributesResponse> SetQueueAttributesAsync(string queueUrl, Dictionary<string, string> attributes, CancellationToken cancellationToken = default)
+	{
+		throw new NotImplementedException();
+	}
+
+	public Task<SetQueueAttributesResponse> SetQueueAttributesAsync(SetQueueAttributesRequest request, CancellationToken cancellationToken = default)
+	{
+		throw new NotImplementedException();
+	}
+
+	public Task<TagQueueResponse> TagQueueAsync(TagQueueRequest request, CancellationToken cancellationToken = default)
+	{
+		throw new NotImplementedException();
+	}
+
+	public Task<UntagQueueResponse> UntagQueueAsync(UntagQueueRequest request, CancellationToken cancellationToken = default)
+	{
+		throw new NotImplementedException();
+	}
+#endregion Not implemented	
+}
+
+#pragma warning restore CA1054
\ No newline at end of file
diff --git a/Engine/Source/Programs/Horde/Horde.Server.Tests/TestSetup.cs b/Engine/Source/Programs/Horde/Horde.Server.Tests/TestSetup.cs
--- a/Engine/Source/Programs/Horde/Horde.Server.Tests/TestSetup.cs
+++ b/Engine/Source/Programs/Horde/Horde.Server.Tests/TestSetup.cs
@@ -100,7 +100,6 @@ namespace Horde.Server.Tests
 		public AclService AclService => ServiceProvider.GetRequiredService<AclService>();
 		public FleetService FleetService => ServiceProvider.GetRequiredService<FleetService>();
 		public AgentService AgentService => ServiceProvider.GetRequiredService<AgentService>();
-		public AwsAutoScalingLifecycleService AwsAsgLifecycleService => ServiceProvider.GetRequiredService<AwsAutoScalingLifecycleService>();
 		public ICommitService CommitService => ServiceProvider.GetRequiredService<ICommitService>();
 		public GlobalsService GlobalsService => ServiceProvider.GetRequiredService<GlobalsService>();
 		public MongoService MongoService => ServiceProvider.GetRequiredService<MongoService>();
@@ -134,7 +133,6 @@ namespace Horde.Server.Tests
 		public DashboardController DashboardController => GetDashboardController();
 		public TestDataController TestDataController => GetTestDataController();
 		public BisectTasksController BisectTasksController => GetBisectTasksController();
-		public AwsAutoScalingLifecycleController AwsAutoScalingLifecycleController => GetAwsAutoScalingLifecycleController();
 
 		public OpenTelemetry.Trace.Tracer Tracer => ServiceProvider.GetRequiredService<OpenTelemetry.Trace.Tracer>();
 		public Meter Meter => ServiceProvider.GetRequiredService<Meter>();
@@ -318,13 +316,6 @@ namespace Horde.Server.Tests
 			dataCtrl.ControllerContext = GetControllerContext();
 			return dataCtrl;
 		}
-		
-		private AwsAutoScalingLifecycleController GetAwsAutoScalingLifecycleController()
-		{
-			AwsAutoScalingLifecycleController ctrl = new (AwsAsgLifecycleService);
-			ctrl.ControllerContext = GetControllerContext();
-			return ctrl;
-		}
 
 		private BisectTasksController GetBisectTasksController()
 		{
diff --git a/Engine/Source/Programs/Horde/Horde.Server/Agents/Fleet/AwsAutoScalingLifecycle.cs b/Engine/Source/Programs/Horde/Horde.Server/Agents/Fleet/AwsAutoScalingLifecycle.cs
--- a/Engine/Source/Programs/Horde/Horde.Server/Agents/Fleet/AwsAutoScalingLifecycle.cs
+++ b/Engine/Source/Programs/Horde/Horde.Server/Agents/Fleet/AwsAutoScalingLifecycle.cs
@@ -10,14 +10,16 @@ using System.Threading;
 using System.Threading.Tasks;
 using Amazon.AutoScaling;
 using Amazon.AutoScaling.Model;
+using Amazon.Runtime;
+using Amazon.SQS;
+using Amazon.SQS.Model;
+using EpicGames.Core;
 using Horde.Server.Server;
-using Horde.Server.Utilities;
 using HordeCommon;
-using Microsoft.AspNetCore.Authorization;
-using Microsoft.AspNetCore.Mvc;
 using Microsoft.Extensions.DependencyInjection;
 using Microsoft.Extensions.Hosting;
 using Microsoft.Extensions.Logging;
+using Microsoft.Extensions.Options;
 using OpenTelemetry.Trace;
 using StackExchange.Redis;
 
@@ -63,9 +65,12 @@ public sealed class AwsAutoScalingLifecycleService : IHostedService, IDisposable
 	private readonly Tracer _tracer;
 	private readonly ILogger<AwsAutoScalingLifecycleService> _logger;
 	private readonly ITicker _updateLifecyclesTicker;
+	private readonly BackgroundTask _lifecycleEventListenerTask;
+	private readonly string? _sqsQueueUrl;
 
 #pragma warning disable CA2213 // Disposable fields should be disposed
 	private IAmazonAutoScaling? _awsAutoScaling;
+	private IAmazonSQS? _awsSqs;
 #pragma warning restore CA2213
 	
 	/// <summary>
@@ -76,6 +81,7 @@ public sealed class AwsAutoScalingLifecycleService : IHostedService, IDisposable
 		RedisService redisService,
 		IAgentCollection agents,
 		IClock clock,
+		IOptionsMonitor<ServerSettings> settings,
 		IServiceProvider serviceProvider,
 		Tracer tracer,
 		ILogger<AwsAutoScalingLifecycleService> logger)
@@ -87,32 +93,122 @@ public sealed class AwsAutoScalingLifecycleService : IHostedService, IDisposable
 		_tracer = tracer;
 		_logger = logger;
 		_awsAutoScaling = serviceProvider.GetService<IAmazonAutoScaling>();
+		_awsSqs = serviceProvider.GetService<IAmazonSQS>();
 
 		string tickerName = $"{nameof(AwsAutoScalingLifecycleService)}.{nameof(UpdateLifecyclesAsync)}";
 		_updateLifecyclesTicker = clock.AddSharedTicker(tickerName, LifecycleUpdaterInterval, UpdateLifecyclesAsync, logger);
+		_lifecycleEventListenerTask = new BackgroundTask(ListenForLifecycleEventsAsync);
+		_sqsQueueUrl = settings.CurrentValue.AwsAutoScalingQueueUrl;
 	}
 
-	internal void SetAmazonAutoScalingClient(IAmazonAutoScaling awsAutoScaling)
+	internal void SetAmazonClientsTesting(IAmazonAutoScaling awsAutoScaling, IAmazonSQS awsSqs)
 	{
 		_awsAutoScaling = awsAutoScaling;
+		_awsSqs = awsSqs;
 	}
 	
 	/// <inheritdoc/>
 	public async Task StartAsync(CancellationToken cancellationToken)
 	{
-		await _updateLifecyclesTicker.StartAsync();
+		if (_sqsQueueUrl != null)
+		{
+			await _updateLifecyclesTicker.StartAsync();
+			_lifecycleEventListenerTask.Start();	
+		}
 	}
 
 	/// <inheritdoc/>
 	public async Task StopAsync(CancellationToken cancellationToken)
 	{
+		await _lifecycleEventListenerTask.StopAsync();
 		await _updateLifecyclesTicker.StopAsync();
-	} 
+	}
 
 	/// <inheritdoc/>
 	public void Dispose()
 	{
 		_updateLifecyclesTicker.Dispose();
+		_lifecycleEventListenerTask.Dispose();
+	}
+
+	/// <summary>
+	/// Continuously request and receive messages from SQS
+	/// </summary>
+	/// <param name="cancellationToken">Cancellation token</param>
+	private async Task ListenForLifecycleEventsAsync(CancellationToken cancellationToken)
+	{
+		_logger.LogInformation("Listening for lifecycle events...");
+		while (!cancellationToken.IsCancellationRequested)
+		{
+			try
+			{
+				await ReceiveLifecycleEventsAsync(cancellationToken);
+			}
+			catch (OperationCanceledException) when (cancellationToken.IsCancellationRequested)
+			{
+				break;
+			}
+			catch (Exception ex)
+			{
+				_logger.LogError(ex, "Exception while receiving lifecycle events: {Message}", ex.Message);
+				await Task.Delay(TimeSpan.FromSeconds(15), cancellationToken);
+			}
+			
+			await Task.Delay(TimeSpan.FromSeconds(1), cancellationToken);
+		}
+	}
+	
+	/// <summary>
+	/// Receive any queued messages from SQS.
+	/// This is a one-time operation, if more messages are expected, call this method again.
+	/// </summary>
+	/// <param name="cancellationToken"></param>
+	/// <exception cref="JsonException"></exception>
+	internal async Task ReceiveLifecycleEventsAsync(CancellationToken cancellationToken)
+	{
+		ReceiveMessageRequest request = new ()
+		{
+			QueueUrl = _sqsQueueUrl,
+			MaxNumberOfMessages = 10,
+			VisibilityTimeout = 10,
+			WaitTimeSeconds = 10,
+		};
+		
+		ReceiveMessageResponse response = await GetSqs().ReceiveMessageAsync(request, cancellationToken);
+		foreach (Message message in response.Messages)
+		{
+			try
+			{
+				LifecycleActionEvent ev = JsonSerializer.Deserialize<LifecycleActionEvent>(message.Body) ?? throw new JsonException("Value is null");
+				await InitiateTerminationAsync(ev, cancellationToken);
+			}
+			catch (JsonException je)
+			{
+				_logger.LogError(je, "Failed deserializing message. Body={Body}", message.Body);
+			}
+			finally
+			{
+				await DeleteMessageAsync(message, cancellationToken);
+			}
+		}
+	}
+
+	private async Task DeleteMessageAsync(Message message, CancellationToken cancellationToken)
+	{
+		try
+		{
+			DeleteMessageRequest deleteRequest = new () { QueueUrl = _sqsQueueUrl, ReceiptHandle = message.ReceiptHandle };
+			await GetSqs().DeleteMessageAsync(deleteRequest, cancellationToken);
+		}
+		catch (AmazonServiceException e)
+		{
+			_logger.LogError(e, "Failed to delete message {MessageId} from SQS queue after processing it", message.MessageId);
+		}
+	}
+
+	private IAmazonSQS GetSqs()
+	{
+		return _awsSqs ?? throw new Exception("AWS SQS client is not set. Make sure AWS is configured in settings.");;
 	}
 	
 	/// <summary>
@@ -257,27 +353,6 @@ public sealed class AwsAutoScalingLifecycleService : IHostedService, IDisposable
 	}
 }
 
-/// <summary>
-/// An AWS EventBridge event containing a lifecycle hook message for an auto-scaling group
-/// </summary>
-public class LifecycleActionRequest
-{
-	/// <summary>
-	/// Event
-	/// </summary>
-	[JsonPropertyName("detail")]
-	public LifecycleActionEvent Event { get; set; }
-
-	/// <summary>
-	/// Constructor
-	/// </summary>
-	/// <param name="event"></param>
-	public LifecycleActionRequest(LifecycleActionEvent @event)
-	{
-		Event = @event;
-	}
-}
-
 /// <summary>
 /// Lifecycle action event from AWS auto-scaling group
 /// The property names are explicitly set to highlight these are not controlled by Horde but sent from AWS API.
@@ -314,37 +389,3 @@ public class LifecycleActionEvent
 	/// </summary>
 	[JsonPropertyName("Origin")] public string Origin { get; set; } = "";
 }
-
-/// <summary>
-/// Controller handling callbacks for AWS auto-scaling
-/// </summary>
-[ApiController]
-[Authorize]
-[Route("[controller]")]
-public class AwsAutoScalingLifecycleController : HordeControllerBase
-{
-	private AwsAutoScalingLifecycleService _lifecycleService;
-
-	/// <summary>
-	/// Constructor
-	/// </summary>
-	public AwsAutoScalingLifecycleController(AwsAutoScalingLifecycleService lifecycleService)
-	{
-		_lifecycleService = lifecycleService;
-	}
-	
-	/// <summary>
-	/// Callback handling auto-scaling notification generated by the ASG
-	/// A Lambda function should be configured first to receive the notification from AWS EventBridge.
-	/// Then in turn, call the Horde server on this endpoint. 
-	/// </summary>
-	/// <param name="request">Request coming from</param>
-	/// <returns>Human readable message</returns>
-	[HttpPost]
-	[Route("/api/v1/aws/asg-lifecycle-hook")]
-	public async Task<ActionResult> LifecycleHookAsync([FromBody] LifecycleActionRequest request)
-	{
-		await _lifecycleService.InitiateTerminationAsync(request.Event, CancellationToken.None);
-		return new JsonResult(new { message = "Handled" });
-	}
-}
\ No newline at end of file
diff --git a/Engine/Source/Programs/Horde/Horde.Server/Horde.Server.csproj b/Engine/Source/Programs/Horde/Horde.Server/Horde.Server.csproj
--- a/Engine/Source/Programs/Horde/Horde.Server/Horde.Server.csproj
+++ b/Engine/Source/Programs/Horde/Horde.Server/Horde.Server.csproj
@@ -24,6 +24,7 @@
     <PackageReference Include="AWSSDK.EC2" Version="3.7.100" />
     <PackageReference Include="AWSSDK.Extensions.NETCore.Setup" Version="3.7.2" />
     <PackageReference Include="AWSSDK.S3" Version="3.7.100" />
+    <PackageReference Include="AWSSDK.SQS" Version="3.7.100" />
     <PackageReference Include="AWSSDK.SecurityToken" Version="3.7.100" />
     <PackageReference Include="Crc32.NET" Version="1.2.0" />
     <PackageReference Include="Grpc.AspNetCore" Version="2.49.0" />
diff --git a/Engine/Source/Programs/Horde/Horde.Server/ServerSettings.cs b/Engine/Source/Programs/Horde/Horde.Server/ServerSettings.cs
--- a/Engine/Source/Programs/Horde/Horde.Server/ServerSettings.cs
+++ b/Engine/Source/Programs/Horde/Horde.Server/ServerSettings.cs
@@ -634,6 +634,12 @@ namespace Horde.Server
 		/// Config for the fleet manager (serialized JSON)
 		/// </summary>
 		public string? FleetManagerV2Config { get; set; }
+		
+		/// <summary>
+		/// AWS SQS queue URL where lifecycle events from EC2 auto-scaling are received
+		/// <see cref="AwsAutoScalingLifecycleService" />
+		/// </summary>
+		public string? AwsAutoScalingQueueUrl { get; set; }
 
 		/// <summary>
 		/// Whether to run scheduled jobs.
