Horde: Format submit errors better on steps
#rnx
#skipci

diff --git a/Engine/Source/Programs/Horde/HordeDashboard/src/components/AutoSubmit.tsx b/Engine/Source/Programs/Horde/HordeDashboard/src/components/AutoSubmit.tsx
--- a/Engine/Source/Programs/Horde/HordeDashboard/src/components/AutoSubmit.tsx
+++ b/Engine/Source/Programs/Horde/HordeDashboard/src/components/AutoSubmit.tsx
@@ -29,7 +29,7 @@ export const AutosubmitInfo: React.FC<{ jobDetails: JobDetailsV2 }> = observer((
    }
 
    if (job.autoSubmitMessage) {
-      return <Stack style={{ paddingTop: 24 }}><Text>{`Unable to submit change: ${job.autoSubmitMessage}`}</Text></Stack>;
+      return <Stack style={{ paddingTop: 24, whiteSpace:"pre" }}><Text>{`Unable to submit change: ${job.autoSubmitMessage}`}</Text></Stack>;
    }
 
    if (job.state !== JobState.Complete) {
diff --git a/Engine/Source/Programs/Horde/HordeDashboard/src/components/jobDetailsV2/JobDetailViewSummary.tsx b/Engine/Source/Programs/Horde/HordeDashboard/src/components/jobDetailsV2/JobDetailViewSummary.tsx
--- a/Engine/Source/Programs/Horde/HordeDashboard/src/components/jobDetailsV2/JobDetailViewSummary.tsx
+++ b/Engine/Source/Programs/Horde/HordeDashboard/src/components/jobDetailsV2/JobDetailViewSummary.tsx
@@ -73,7 +73,7 @@ export const AutosubmitInfo: React.FC<{ jobDetails: JobDetailsV2 }> = observer((
    }
 
    if (jobData.autoSubmitMessage) {
-      return <Stack style={{ paddingTop: 24 }}><Text>{`Unable to submit change: ${jobData.autoSubmitMessage}`}</Text></Stack>;
+      return <Stack style={{ paddingTop: 24, whiteSpace:"pre" }}><Text>{`Unable to submit change: ${jobData.autoSubmitMessage}`}</Text></Stack>;
    }
 
    if (jobData.state !== JobState.Complete) {
