Horde: Add AWS SQS for dependency injection

diff --git a/Engine/Source/Programs/Horde/Horde.Server/Startup.cs b/Engine/Source/Programs/Horde/Horde.Server/Startup.cs
--- a/Engine/Source/Programs/Horde/Horde.Server/Startup.cs
+++ b/Engine/Source/Programs/Horde/Horde.Server/Startup.cs
@@ -18,6 +18,7 @@ using Amazon.AutoScaling;
 using Amazon.CloudWatch;
 using Amazon.EC2;
 using Amazon.Extensions.NETCore.Setup;
+using Amazon.SQS;
 using EpicGames.AspNet;
 using EpicGames.Core;
 using EpicGames.Horde.Storage;
@@ -524,6 +525,7 @@ namespace Horde.Server
 				
 				services.AddAWSService<IAmazonCloudWatch>();
 				services.AddAWSService<IAmazonAutoScaling>();
+				services.AddAWSService<IAmazonSQS>();
 				services.AddAWSService<IAmazonEC2>();
 			}
 
