Horde: Better formatting on 500 exception in error handler modal
#skipci
#rnx

diff --git a/Engine/Source/Programs/Horde/HordeDashboard/src/components/ErrorHandler.tsx b/Engine/Source/Programs/Horde/HordeDashboard/src/components/ErrorHandler.tsx
--- a/Engine/Source/Programs/Horde/HordeDashboard/src/components/ErrorHandler.tsx
+++ b/Engine/Source/Programs/Horde/HordeDashboard/src/components/ErrorHandler.tsx
@@ -30,7 +30,7 @@ export type ErrorInfo = {
    url?: string;
 
    // REST mode
-   mode?: "GET" | "PUT" | "DELETE" | "POST" | string;
+   mode?: "GET" | "PUT" | "DELETE" | "POST" | "PATCH" | string;
 
    // custom title override
    title?: string;
@@ -204,7 +204,7 @@ export const ErrorDialog: React.FC = observer(() => {
       return <div />;
    }
 
-   const message = ErrorHandler.message;
+   let message = ErrorHandler.message;
    const unauthorized = ErrorHandler.unauthorized();
 
    let title = error.title;
@@ -220,9 +220,23 @@ export const ErrorDialog: React.FC = observer(() => {
    const helpEmail = dashboard.helpEmail ?? "Not Configured";
    const helpSlack = dashboard.helpSlack ?? "Not Configured";
 
+   let fontSize = 13;
+
+   if (message?.length) {
+      try {
+         const exception = JSON.parse(message);
+         if (exception?.exception?.message && exception?.exception?.trace) {
+            fontSize = 11
+            message = `${exception?.exception?.message}\n${exception?.exception?.trace}`
+         }
+      } catch (error) {
+
+      }
+   }
+
    return (
       <Dialog
-         minWidth={800}
+         minWidth={1380}
          hidden={!error}
          onDismiss={() => ErrorHandler.clear()}
          modalProps={{
@@ -231,7 +245,7 @@ export const ErrorDialog: React.FC = observer(() => {
          dialogContentProps={dialogContentProps}>
 
          <Stack tokens={{ childrenGap: 12 }}>
-            <Text styles={{ root: { whiteSpace: "pre-wrap", fontSize: 13, color: "#EC4C47" } }}>{message}</Text>
+            <Text styles={{ root: { whiteSpace: "pre", fontSize: fontSize, color: "#EC4C47" } }}>{message}</Text>
             {!!error.url && <Text variant="medium">{`URL: ${error.url}`}</Text>}
             <Stack>
                <Stack tokens={{ childrenGap: 18 }} style={{ paddingTop: 12 }}>
