Horde: Don't add bisection updates for zero length changes
#rnx

diff --git a/Engine/Source/Programs/Horde/Horde.Server/Jobs/Bisect/BisectTaskCollection.cs b/Engine/Source/Programs/Horde/Horde.Server/Jobs/Bisect/BisectTaskCollection.cs
--- a/Engine/Source/Programs/Horde/Horde.Server/Jobs/Bisect/BisectTaskCollection.cs
+++ b/Engine/Source/Programs/Horde/Horde.Server/Jobs/Bisect/BisectTaskCollection.cs
@@ -170,19 +170,21 @@ namespace Horde.Server.Jobs.Bisect
 			{
 				update = update.Set(x => x.State, options.State.Value);
 			}
-			if (options.IncludeChanges != null)
+		
+			if (options.IncludeChanges != null && options.IncludeChanges.Count > 0)
 			{
 				update = update.PullAll(x => x.IgnoreChanges, options.IncludeChanges); 
 			}
-			if (options.ExcludeChanges != null)
+			else if (options.ExcludeChanges != null && options.ExcludeChanges.Count > 0)
 			{
 				update = update.AddToSetEach(x => x.IgnoreChanges, options.ExcludeChanges);
 			}
-			if (options.IncludeJobs != null)
+
+			if (options.IncludeJobs != null && options.IncludeJobs.Count > 0)
 			{
 				update = update.PullAll(x => x.IgnoreJobs, options.IncludeJobs);
 			}
-			if (options.ExcludeJobs != null)
+			else if (options.ExcludeJobs != null && options.ExcludeJobs.Count > 0)
 			{
 				update = update.AddToSetEach(x => x.IgnoreJobs, options.ExcludeJobs);
 			}
