Horde: Ignore #robomerge-owner tags containing @ characters.

diff --git a/Engine/Source/Programs/Horde/Horde.Server/Issues/IssueService.cs b/Engine/Source/Programs/Horde/Horde.Server/Issues/IssueService.cs
--- a/Engine/Source/Programs/Horde/Horde.Server/Issues/IssueService.cs
+++ b/Engine/Source/Programs/Horde/Horde.Server/Issues/IssueService.cs
@@ -1344,25 +1344,6 @@ namespace Horde.Server.Issues
 			}
 		}
 
-		/// <summary>
-		/// Attempts to parse the Robomerge owner from this commit information
-		/// </summary>
-		/// <param name="description">Description text to parse</param>
-		/// <returns>The Robomerge owner, or null if no #ROBOMERGE-OWNER tag was present</returns>
-		static string? ParseRobomergeOwner(string description)
-		{
-			// #ROBOMERGE-OWNER: ben.marsh
-			Match match = Regex.Match(description, @"^#ROBOMERGE-OWNER:\s*([^\s]+)", RegexOptions.Multiline);
-			if (match.Success)
-			{
-				return match.Groups[1].Value;
-			}
-			else
-			{
-				return null;
-			}
-		}
-
 		/// <summary>
 		/// Finds the handler for a given issue
 		/// </summary>
diff --git a/Engine/Source/Programs/Horde/Horde.Server/Perforce/PerforceService.cs b/Engine/Source/Programs/Horde/Horde.Server/Perforce/PerforceService.cs
--- a/Engine/Source/Programs/Horde/Horde.Server/Perforce/PerforceService.cs
+++ b/Engine/Source/Programs/Horde/Horde.Server/Perforce/PerforceService.cs
@@ -730,7 +730,7 @@ namespace Horde.Server.Perforce
 		static string? ParseRobomergeOwner(string description)
 		{
 			// #ROBOMERGE-OWNER: ben.marsh
-			Match match = Regex.Match(description, @"^#ROBOMERGE-OWNER:\s*([^\s]+)", RegexOptions.Multiline);
+			Match match = Regex.Match(description, @"^#ROBOMERGE-OWNER:\s*([^@\s]+)", RegexOptions.Multiline);
 			if (match.Success)
 			{
 				return match.Groups[1].Value;
