Horde: Include a hash in node locators.

diff --git a/Engine/Source/Programs/Horde/Horde.Agent/Utility/JsonRpcLogSink.cs b/Engine/Source/Programs/Horde/Horde.Agent/Utility/JsonRpcLogSink.cs
--- a/Engine/Source/Programs/Horde/Horde.Agent/Utility/JsonRpcLogSink.cs
+++ b/Engine/Source/Programs/Horde/Horde.Agent/Utility/JsonRpcLogSink.cs
@@ -244,7 +244,7 @@ namespace Horde.Agent.Utility
 			UpdateLogRequest request = new UpdateLogRequest();
 			request.LogId = _logId;
 			request.LineCount = lineCount;
-			request.Target = new HashedNodeLocator(target.Hash, target.GetLocator()).ToString();
+			request.Target = target.GetLocator().ToString();
 			request.Complete = complete;
 			await _connection.InvokeAsync((LogRpcClient client) => client.UpdateLogAsync(request, cancellationToken: cancellationToken), cancellationToken);
 		}
diff --git a/Engine/Source/Programs/Horde/Horde.Server.Tests/BlobStoreTests.cs b/Engine/Source/Programs/Horde/Horde.Server.Tests/BlobStoreTests.cs
--- a/Engine/Source/Programs/Horde/Horde.Server.Tests/BlobStoreTests.cs
+++ b/Engine/Source/Programs/Horde/Horde.Server.Tests/BlobStoreTests.cs
@@ -120,7 +120,7 @@ namespace Horde.Server.Tests
 			for(int idx = 0; idx < 2; idx++)
 			{
 				RefName refName = new RefName("hello");
-				await store.WriteRefTargetAsync(refName, new HashedNodeLocator(bundle3.Header.Exports[0].Hash, locator3, 0));
+				await store.WriteRefTargetAsync(refName, new NodeLocator(bundle3.Header.Exports[0].Hash, locator3, 0));
 				BlobHandle refTarget = await store.ReadRefTargetAsync(refName);
 				Assert.AreEqual(locator3, refTarget.GetLocator().Blob);
 			}
@@ -133,7 +133,7 @@ namespace Horde.Server.Tests
 
 			Bundle bundle1 = CreateTestBundle(new byte[] { 1, 2, 3 }, Array.Empty<BlobLocator>());
 			BlobLocator locator1 = await store.WriteBundleAsync(bundle1);
-			HashedNodeLocator target = new HashedNodeLocator(bundle1.Header.Exports[0].Hash, locator1, 0);
+			NodeLocator target = new NodeLocator(bundle1.Header.Exports[0].Hash, locator1, 0);
 
 			await store.WriteRefTargetAsync("test-ref-1", target);
 			await store.WriteRefTargetAsync("test-ref-2", target, new RefOptions { Lifetime = TimeSpan.FromMinutes(30.0), Extend = true });
@@ -162,14 +162,14 @@ namespace Horde.Server.Tests
 			Assert.AreEqual(default, await TryReadRefTargetAsync(store, "test-ref-3"));
 		}
 
-		static async Task<HashedNodeLocator> TryReadRefTargetAsync(IStorageClient store, RefName name)
+		static async Task<NodeLocator> TryReadRefTargetAsync(IStorageClient store, RefName name)
 		{
 			BlobHandle? handle = await store.TryReadRefTargetAsync(name);
 			if (handle == null)
 			{
 				return default;
 			}
-			return new HashedNodeLocator(handle.Hash, handle.GetLocator());
+			return handle.GetLocator();
 		}
 	}
 }
diff --git a/Engine/Source/Programs/Horde/Horde.Server.Tests/GcServiceTests.cs b/Engine/Source/Programs/Horde/Horde.Server.Tests/GcServiceTests.cs
--- a/Engine/Source/Programs/Horde/Horde.Server.Tests/GcServiceTests.cs
+++ b/Engine/Source/Programs/Horde/Horde.Server.Tests/GcServiceTests.cs
@@ -43,7 +43,7 @@ namespace Horde.Server.Tests
 				int blobIdx = (int)(random.NextDouble() * blobs.Length);
 				if (roots.Add(blobs[blobIdx]))
 				{
-					await store.WriteRefTargetAsync(new RefName($"ref-{idx}"), new HashedNodeLocator(IoHash.Zero, blobs[blobIdx], 0));
+					await store.WriteRefTargetAsync(new RefName($"ref-{idx}"), new NodeLocator(IoHash.Zero, blobs[blobIdx], 0));
 				}
 			}
 
diff --git a/Engine/Source/Programs/Horde/Horde.Server.Tests/StorageServiceTests.cs b/Engine/Source/Programs/Horde/Horde.Server.Tests/StorageServiceTests.cs
--- a/Engine/Source/Programs/Horde/Horde.Server.Tests/StorageServiceTests.cs
+++ b/Engine/Source/Programs/Horde/Horde.Server.Tests/StorageServiceTests.cs
@@ -44,20 +44,20 @@ namespace Horde.Server.Tests
 			Bundle bundle = new Bundle(header, Array.Empty<ReadOnlyMemory<byte>>());
 			BlobLocator locator = await client.WriteBundleAsync(bundle);
 
-			await client.AddAliasAsync("foo", new HashedNodeLocator(hash1, locator, 0));
-			await client.AddAliasAsync("foo", new HashedNodeLocator(hash1, locator, 1));
-			await client.AddAliasAsync("bar", new HashedNodeLocator(hash2, locator, 2));
+			await client.AddAliasAsync("foo", new NodeLocator(hash1, locator, 0));
+			await client.AddAliasAsync("foo", new NodeLocator(hash1, locator, 1));
+			await client.AddAliasAsync("bar", new NodeLocator(hash2, locator, 2));
 
 			List<BlobHandle> handles;
 			
 			handles = await client.FindNodesAsync("foo").ToListAsync();
 			Assert.AreEqual(2, handles.Count);
-			Assert.AreEqual(new NodeLocator(locator, 0), handles[0].GetLocator());
-			Assert.AreEqual(new NodeLocator(locator, 1), handles[1].GetLocator());
+			Assert.AreEqual(new NodeLocator(hash1, locator, 0), handles[0].GetLocator());
+			Assert.AreEqual(new NodeLocator(hash1, locator, 1), handles[1].GetLocator());
 
 			handles = await client.FindNodesAsync("bar").ToListAsync();
 			Assert.AreEqual(1, handles.Count);
-			Assert.AreEqual(new NodeLocator(locator, 2), handles[0].GetLocator());
+			Assert.AreEqual(new NodeLocator(hash2, locator, 2), handles[0].GetLocator());
 		}
 	}
 }
diff --git a/Engine/Source/Programs/Horde/Horde.Server/Logs/LogRpcService.cs b/Engine/Source/Programs/Horde/Horde.Server/Logs/LogRpcService.cs
--- a/Engine/Source/Programs/Horde/Horde.Server/Logs/LogRpcService.cs
+++ b/Engine/Source/Programs/Horde/Horde.Server/Logs/LogRpcService.cs
@@ -52,7 +52,7 @@ namespace Horde.Server.Logs
 
 			IStorageClientImpl store = await _storageService.GetClientAsync(Namespace.Logs, context.CancellationToken);
 			_logger.LogInformation("Updating {LogId} to node {RefTarget} (lines: {LineCount}, complete: {Complete})", request.LogId, request.Target, request.LineCount, request.Complete);
-			await store.WriteRefTargetAsync(new RefName(request.LogId), HashedNodeLocator.Parse(request.Target));
+			await store.WriteRefTargetAsync(new RefName(request.LogId), NodeLocator.Parse(request.Target));
 
 			await _logFileCollection.UpdateLineCountAsync(logFile, request.LineCount, request.Complete, CancellationToken.None);
 
diff --git a/Engine/Source/Programs/Horde/Horde.Server/Storage/StorageController.cs b/Engine/Source/Programs/Horde/Horde.Server/Storage/StorageController.cs
--- a/Engine/Source/Programs/Horde/Horde.Server/Storage/StorageController.cs
+++ b/Engine/Source/Programs/Horde/Horde.Server/Storage/StorageController.cs
@@ -364,7 +364,7 @@ namespace Horde.Server.Storage
 			}
 
 			IStorageClientImpl client = await _storageService.GetClientAsync(namespaceId, cancellationToken);
-			HashedNodeLocator target = new HashedNodeLocator(request.Hash, request.Blob, request.ExportIdx);
+			NodeLocator target = new NodeLocator(request.Hash, request.Blob, request.ExportIdx);
 			await client.WriteRefTargetAsync(refName, target, request.Options, cancellationToken);
 
 			return Ok();
@@ -573,7 +573,7 @@ namespace Horde.Server.Storage
 
 			object content;
 
-			BlobData nodeData = await reader.ReadNodeDataAsync(new NodeLocator(locator, exportIdx), cancellationToken);
+			BlobData nodeData = await reader.ReadNodeDataAsync(new NodeLocator(export.Hash, locator, exportIdx), cancellationToken);
 
 			Node node = Node.Deserialize(nodeData);
 			switch (node)
diff --git a/Engine/Source/Programs/Horde/Horde.Server/Storage/StorageService.cs b/Engine/Source/Programs/Horde/Horde.Server/Storage/StorageService.cs
--- a/Engine/Source/Programs/Horde/Horde.Server/Storage/StorageService.cs
+++ b/Engine/Source/Programs/Horde/Horde.Server/Storage/StorageService.cs
@@ -64,13 +64,13 @@ namespace Horde.Server.Storage
 		bool SupportsRedirects { get; }
 
 		/// <inheritdoc cref="IStorageClient.AddAliasAsync(Utf8String, BlobHandle, CancellationToken)"/>
-		Task AddAliasAsync(Utf8String name, HashedNodeLocator target, CancellationToken cancellationToken = default);
+		Task AddAliasAsync(Utf8String name, NodeLocator target, CancellationToken cancellationToken = default);
 
 		/// <inheritdoc cref="IStorageClient.RemoveAliasAsync(Utf8String, BlobHandle, CancellationToken)"/>
-		Task RemoveAliasAsync(Utf8String name, HashedNodeLocator target, CancellationToken cancellationToken = default); 
+		Task RemoveAliasAsync(Utf8String name, NodeLocator target, CancellationToken cancellationToken = default); 
 
 		/// <inheritdoc cref="IStorageClient.WriteRefTargetAsync(RefName, BlobHandle, RefOptions?, CancellationToken)"/>
-		Task WriteRefTargetAsync(RefName name, HashedNodeLocator handle, RefOptions? options = null, CancellationToken cancellationToken = default);
+		Task WriteRefTargetAsync(RefName name, NodeLocator handle, RefOptions? options = null, CancellationToken cancellationToken = default);
 
 		/// <summary>
 		/// Gets a redirect for a read request
@@ -224,21 +224,21 @@ namespace Horde.Server.Storage
 			#region Nodes
 
 			/// <inheritdoc/>
-			public override Task AddAliasAsync(Utf8String name, BlobHandle handle, CancellationToken cancellationToken = default) => _outer.AddAliasAsync(NamespaceId, name, new HashedNodeLocator(handle.Hash, handle.GetLocator()), cancellationToken);
+			public override Task AddAliasAsync(Utf8String name, BlobHandle handle, CancellationToken cancellationToken = default) => _outer.AddAliasAsync(NamespaceId, name, handle.GetLocator(), cancellationToken);
 
 			/// <inheritdoc/>
-			public Task AddAliasAsync(Utf8String name, HashedNodeLocator locator, CancellationToken cancellationToken = default) => _outer.AddAliasAsync(NamespaceId, name, locator, cancellationToken);
+			public Task AddAliasAsync(Utf8String name, NodeLocator locator, CancellationToken cancellationToken = default) => _outer.AddAliasAsync(NamespaceId, name, locator, cancellationToken);
 
 			/// <inheritdoc/>
-			public override Task RemoveAliasAsync(Utf8String name, BlobHandle handle, CancellationToken cancellationToken = default) => _outer.RemoveAliasAsync(NamespaceId, name, new HashedNodeLocator(handle.Hash, handle.GetLocator()), cancellationToken);
+			public override Task RemoveAliasAsync(Utf8String name, BlobHandle handle, CancellationToken cancellationToken = default) => _outer.RemoveAliasAsync(NamespaceId, name, handle.GetLocator(), cancellationToken);
 
 			/// <inheritdoc/>
-			public Task RemoveAliasAsync(Utf8String name, HashedNodeLocator locator, CancellationToken cancellationToken = default) => _outer.RemoveAliasAsync(NamespaceId, name, locator, cancellationToken);
+			public Task RemoveAliasAsync(Utf8String name, NodeLocator locator, CancellationToken cancellationToken = default) => _outer.RemoveAliasAsync(NamespaceId, name, locator, cancellationToken);
 
 			/// <inheritdoc/>
 			public override async IAsyncEnumerable<BlobHandle> FindNodesAsync(Utf8String alias, [EnumeratorCancellation] CancellationToken cancellationToken = default)
 			{
-				await foreach (HashedNodeLocator locator in _outer.FindNodesAsync(NamespaceId, alias, cancellationToken))
+				await foreach (NodeLocator locator in _outer.FindNodesAsync(NamespaceId, alias, cancellationToken))
 				{
 					yield return new FlushedNodeHandle(TreeReader, locator);
 				}
@@ -251,7 +251,7 @@ namespace Horde.Server.Storage
 			/// <inheritdoc/>
 			public override async Task<BlobHandle?> TryReadRefTargetAsync(RefName name, RefCacheTime cacheTime = default, CancellationToken cancellationToken = default)
 			{
-				HashedNodeLocator? locator = await _outer.TryReadRefTargetAsync(NamespaceId, name, cacheTime, cancellationToken);
+				NodeLocator? locator = await _outer.TryReadRefTargetAsync(NamespaceId, name, cacheTime, cancellationToken);
 				if (locator == null)
 				{
 					return null;
@@ -263,11 +263,11 @@ namespace Horde.Server.Storage
 			public override async Task WriteRefTargetAsync(RefName name, BlobHandle target, RefOptions? options = null, CancellationToken cancellationToken = default)
 			{
 				NodeLocator locator = await target.FlushAsync(cancellationToken);
-				await _outer.WriteRefTargetAsync(NamespaceId, name, new HashedNodeLocator(target.Hash, locator), options, cancellationToken);
+				await _outer.WriteRefTargetAsync(NamespaceId, name, locator, options, cancellationToken);
 			}
 
 			/// <inheritdoc/>
-			public Task WriteRefTargetAsync(RefName name, HashedNodeLocator target, RefOptions? options = null, CancellationToken cancellationToken = default) => _outer.WriteRefTargetAsync(NamespaceId, name, target, options, cancellationToken);
+			public Task WriteRefTargetAsync(RefName name, NodeLocator target, RefOptions? options = null, CancellationToken cancellationToken = default) => _outer.WriteRefTargetAsync(NamespaceId, name, target, options, cancellationToken);
 
 			/// <inheritdoc/>
 			public override Task DeleteRefAsync(RefName name, CancellationToken cancellationToken = default) => _outer.DeleteRefAsync(NamespaceId, name, cancellationToken);
@@ -398,7 +398,7 @@ namespace Horde.Server.Storage
 			public TimeSpan? Lifetime { get; set; }
 
 			[BsonIgnore]
-			public HashedNodeLocator Target => new HashedNodeLocator(Hash, new NodeLocator(BlobLocator, ExportIdx));
+			public NodeLocator Target => new NodeLocator(Hash, BlobLocator, ExportIdx);
 
 			public RefInfo()
 			{
@@ -406,14 +406,14 @@ namespace Horde.Server.Storage
 				BlobLocator = BlobLocator.Empty;
 			}
 
-			public RefInfo(NamespaceId namespaceId, RefName name, HashedNodeLocator target, ObjectId blobInfoId)
+			public RefInfo(NamespaceId namespaceId, RefName name, NodeLocator target, ObjectId blobInfoId)
 			{
 				NamespaceId = namespaceId;
 				Name = name;
 				Hash = target.Hash;
-				BlobLocator = target.Locator.Blob;
+				BlobLocator = target.Blob;
 				BlobInfoId = blobInfoId;
-				ExportIdx = target.Locator.ExportIdx;
+				ExportIdx = target.ExportIdx;
 			}
 
 			public bool HasExpired(DateTime utcNow) => ExpiresAtUtc.HasValue && utcNow >= ExpiresAtUtc.Value;
@@ -707,7 +707,7 @@ namespace Horde.Server.Storage
 		/// <param name="target">Target node for the alias</param>
 		/// <param name="cancellationToken">Cancellation token for the operation</param>
 		/// <returns>Sequence of thandles</returns>
-		async Task AddAliasAsync(NamespaceId namespaceId, Utf8String alias, HashedNodeLocator target, CancellationToken cancellationToken = default)
+		async Task AddAliasAsync(NamespaceId namespaceId, Utf8String alias, NodeLocator target, CancellationToken cancellationToken = default)
 		{
 			BlobInfo? blobInfo = await _blobCollection.Find(x => x.NamespaceId == namespaceId && x.BlobId == target.Blob.BlobId).FirstOrDefaultAsync(cancellationToken);
 			if (blobInfo == null)
@@ -733,7 +733,7 @@ namespace Horde.Server.Storage
 		/// <param name="target">Target node for the alias</param>
 		/// <param name="cancellationToken">Cancellation token for the operation</param>
 		/// <returns>Sequence of thandles</returns>
-		Task RemoveAliasAsync(NamespaceId namespaceId, Utf8String alias, HashedNodeLocator target, CancellationToken cancellationToken = default)
+		Task RemoveAliasAsync(NamespaceId namespaceId, Utf8String alias, NodeLocator target, CancellationToken cancellationToken = default)
 		{
 			throw new NotSupportedException();
 		}
@@ -745,7 +745,7 @@ namespace Horde.Server.Storage
 		/// <param name="alias">Alias for the node</param>
 		/// <param name="cancellationToken">Cancellation token for the operation</param>
 		/// <returns>Sequence of thandles</returns>
-		async IAsyncEnumerable<HashedNodeLocator> FindNodesAsync(NamespaceId namespaceId, Utf8String alias, [EnumeratorCancellation] CancellationToken cancellationToken = default)
+		async IAsyncEnumerable<NodeLocator> FindNodesAsync(NamespaceId namespaceId, Utf8String alias, [EnumeratorCancellation] CancellationToken cancellationToken = default)
 		{
 			await foreach (BlobInfo blobInfo in _blobCollection.Find(x => x.NamespaceId == namespaceId && x.Exports!.Any(y => y.Alias == alias)).ToAsyncEnumerable(cancellationToken))
 			{
@@ -755,7 +755,7 @@ namespace Horde.Server.Storage
 					{
 						if (exportInfo.Alias == alias)
 						{
-							yield return new HashedNodeLocator(exportInfo.Hash, new NodeLocator(blobInfo.Locator, exportInfo.Index));
+							yield return new NodeLocator(exportInfo.Hash, blobInfo.Locator, exportInfo.Index);
 						}
 					}
 				}
@@ -847,7 +847,7 @@ namespace Horde.Server.Storage
 		}
 
 		/// <inheritdoc/>
-		async Task<HashedNodeLocator?> TryReadRefTargetAsync(NamespaceId namespaceId, RefName name, RefCacheTime cacheTime = default, CancellationToken cancellationToken = default)
+		async Task<NodeLocator?> TryReadRefTargetAsync(NamespaceId namespaceId, RefName name, RefCacheTime cacheTime = default, CancellationToken cancellationToken = default)
 		{
 			RefCacheValue entry;
 			if (!_cache.TryGetValue(name, out entry) || RefCacheTime.IsStaleCacheEntry(entry.Time, cacheTime))
@@ -879,7 +879,7 @@ namespace Horde.Server.Storage
 		}
 
 		/// <inheritdoc/>
-		async Task WriteRefTargetAsync(NamespaceId namespaceId, RefName name, HashedNodeLocator target, RefOptions? options = null, CancellationToken cancellationToken = default)
+		async Task WriteRefTargetAsync(NamespaceId namespaceId, RefName name, NodeLocator target, RefOptions? options = null, CancellationToken cancellationToken = default)
 		{
 			BlobInfo? newBlobInfo = await _blobCollection.Find(x => x.NamespaceId == namespaceId && x.BlobId == target.Blob.BlobId).FirstOrDefaultAsync(cancellationToken);
 			if (newBlobInfo == null)
diff --git a/Engine/Source/Programs/Horde/Horde.Server/Tools/IToolCollection.cs b/Engine/Source/Programs/Horde/Horde.Server/Tools/IToolCollection.cs
--- a/Engine/Source/Programs/Horde/Horde.Server/Tools/IToolCollection.cs
+++ b/Engine/Source/Programs/Horde/Horde.Server/Tools/IToolCollection.cs
@@ -38,11 +38,11 @@ namespace Horde.Server.Tools
 		/// </summary>
 		/// <param name="tool">The tool to update</param>
 		/// <param name="options">Options for the new deployment</param>
-		/// <param name="hashedLocator">Handle to the root node containing the tool data</param>
+		/// <param name="locator">Handle to the root node containing the tool data</param>
 		/// <param name="globalConfig">The current configuration</param>
 		/// <param name="cancellationToken">Cancellation token for the operation</param>
 		/// <returns>Updated tool document, or null if it does not exist</returns>
-		Task<ITool?> CreateDeploymentAsync(ITool tool, ToolDeploymentConfig options, HashedNodeLocator hashedLocator, GlobalConfig globalConfig, CancellationToken cancellationToken);
+		Task<ITool?> CreateDeploymentAsync(ITool tool, ToolDeploymentConfig options, NodeLocator locator, GlobalConfig globalConfig, CancellationToken cancellationToken);
 
 		/// <summary>
 		/// Updates the state of the current deployment
diff --git a/Engine/Source/Programs/Horde/Horde.Server/Tools/ToolCollection.cs b/Engine/Source/Programs/Horde/Horde.Server/Tools/ToolCollection.cs
--- a/Engine/Source/Programs/Horde/Horde.Server/Tools/ToolCollection.cs
+++ b/Engine/Source/Programs/Horde/Horde.Server/Tools/ToolCollection.cs
@@ -263,7 +263,7 @@ namespace Horde.Server.Tools
 			BlobHandle handle = nodeRef.Handle;
 			await client.WriteRefTargetAsync(refName, handle, cancellationToken: cancellationToken);
 
-			return await CreateDeploymentAsync(tool, options, new HashedNodeLocator(handle.Hash, handle.GetLocator()), globalConfig, cancellationToken);
+			return await CreateDeploymentAsync(tool, options, handle.GetLocator(), globalConfig, cancellationToken);
 		}
 
 		/// <summary>
@@ -271,17 +271,17 @@ namespace Horde.Server.Tools
 		/// </summary>
 		/// <param name="tool">The tool to update</param>
 		/// <param name="options">Options for the new deployment</param>
-		/// <param name="handle">Handle for the tool data</param>
+		/// <param name="locator">Locator for the tool data</param>
 		/// <param name="globalConfig">The current configuration</param>
 		/// <param name="cancellationToken">Cancellation token for the operation</param>
 		/// <returns>Updated tool document, or null if it does not exist</returns>
-		public async Task<ITool?> CreateDeploymentAsync(ITool tool, ToolDeploymentConfig options, HashedNodeLocator handle, GlobalConfig globalConfig, CancellationToken cancellationToken)
+		public async Task<ITool?> CreateDeploymentAsync(ITool tool, ToolDeploymentConfig options, NodeLocator locator, GlobalConfig globalConfig, CancellationToken cancellationToken)
 		{
 			ToolDeploymentId deploymentId = ToolDeploymentId.GenerateNewId();
 			RefName refName = new RefName($"{tool.Id}/{deploymentId}");
 
 			IStorageClientImpl client = await _storageService.GetClientAsync(Namespace.Tools, cancellationToken);
-			await client.WriteRefTargetAsync(refName, handle, cancellationToken: cancellationToken);
+			await client.WriteRefTargetAsync(refName, locator, cancellationToken: cancellationToken);
 
 			return await CreateDeploymentInternalAsync(tool, deploymentId, options, refName, globalConfig, cancellationToken);
 		}
diff --git a/Engine/Source/Programs/Horde/Horde.Server/Tools/ToolsController.cs b/Engine/Source/Programs/Horde/Horde.Server/Tools/ToolsController.cs
--- a/Engine/Source/Programs/Horde/Horde.Server/Tools/ToolsController.cs
+++ b/Engine/Source/Programs/Horde/Horde.Server/Tools/ToolsController.cs
@@ -304,7 +304,7 @@ namespace Horde.Server.Tools
 
 			ToolDeploymentConfig options = new ToolDeploymentConfig { Version = request.Version, Duration = TimeSpan.FromMinutes(request.Duration ?? 0.0), CreatePaused = request.CreatePaused ?? false };
 
-			tool = await _toolCollection.CreateDeploymentAsync(tool, options, HashedNodeLocator.Parse(request.Node), _globalConfig.Value, cancellationToken);
+			tool = await _toolCollection.CreateDeploymentAsync(tool, options, NodeLocator.Parse(request.Node), _globalConfig.Value, cancellationToken);
 			if (tool == null)
 			{
 				return NotFound(id);
diff --git a/Engine/Source/Programs/Horde/Horde/Commands/Bundles/BundleDump.cs b/Engine/Source/Programs/Horde/Horde/Commands/Bundles/BundleDump.cs
--- a/Engine/Source/Programs/Horde/Horde/Commands/Bundles/BundleDump.cs
+++ b/Engine/Source/Programs/Horde/Horde/Commands/Bundles/BundleDump.cs
@@ -67,7 +67,7 @@ namespace Horde.Commands.Bundles
 					for(int referenceIdx = 0; referenceIdx < export.References.Count; referenceIdx++)
 					{
 						BundleExportRef exportRef = export.References[referenceIdx];
-						NodeLocator nodeLocator = new NodeLocator(header.Imports[exportRef.ImportIdx], exportRef.NodeIdx);
+						NodeLocator nodeLocator = new NodeLocator(exportRef.Hash, header.Imports[exportRef.ImportIdx], exportRef.NodeIdx);
 						logger.LogInformation("            REF {RefIdx,-3} -> {Node}", referenceIdx, nodeLocator);
 					}
 				}
diff --git a/Engine/Source/Programs/Shared/EpicGames.Horde/Storage/Backends/FileStorageClient.cs b/Engine/Source/Programs/Shared/EpicGames.Horde/Storage/Backends/FileStorageClient.cs
--- a/Engine/Source/Programs/Shared/EpicGames.Horde/Storage/Backends/FileStorageClient.cs
+++ b/Engine/Source/Programs/Shared/EpicGames.Horde/Storage/Backends/FileStorageClient.cs
@@ -42,7 +42,7 @@ namespace EpicGames.Horde.Storage.Backends
 		public async ValueTask<BlobHandle> ReadRefAsync(FileReference file)
 		{
 			string text = await FileReference.ReadAllTextAsync(file);
-			return new FlushedNodeHandle(TreeReader, HashedNodeLocator.Parse(text));
+			return new FlushedNodeHandle(TreeReader, NodeLocator.Parse(text));
 		}
 
 		FileReference GetRefFile(RefName name) => FileReference.Combine(_rootDir, name.ToString() + ".ref");
@@ -127,7 +127,7 @@ namespace EpicGames.Horde.Storage.Backends
 
 			_logger.LogInformation("Reading {File}", file);
 			string text = await FileReference.ReadAllTextAsync(file, cancellationToken);
-			return new FlushedNodeHandle(TreeReader, HashedNodeLocator.Parse(text));
+			return new FlushedNodeHandle(TreeReader, NodeLocator.Parse(text));
 		}
 
 		/// <inheritdoc/>
diff --git a/Engine/Source/Programs/Shared/EpicGames.Horde/Storage/Backends/HttpStorageClient.cs b/Engine/Source/Programs/Shared/EpicGames.Horde/Storage/Backends/HttpStorageClient.cs
--- a/Engine/Source/Programs/Shared/EpicGames.Horde/Storage/Backends/HttpStorageClient.cs
+++ b/Engine/Source/Programs/Shared/EpicGames.Horde/Storage/Backends/HttpStorageClient.cs
@@ -220,7 +220,7 @@ namespace EpicGames.Horde.Storage.Backends
 						FindNodesResponse? message = await response.Content.ReadFromJsonAsync<FindNodesResponse>(cancellationToken: cancellationToken);
 						foreach (FindNodeResponse node in message!.Nodes)
 						{
-							yield return new FlushedNodeHandle(TreeReader, node.Hash, new NodeLocator(node.Blob, node.ExportIdx));
+							yield return new FlushedNodeHandle(TreeReader, new NodeLocator(node.Hash, node.Blob, node.ExportIdx));
 						}
 					}
 				}
@@ -276,7 +276,7 @@ namespace EpicGames.Horde.Storage.Backends
 							response.EnsureSuccessStatusCode();
 							ReadRefResponse? data = await response.Content.ReadFromJsonAsync<ReadRefResponse>(cancellationToken: cancellationToken);
 							_logger.LogDebug("Read ref {RefName} -> {Blob}#{ExportIdx}", name, data!.Blob, data!.ExportIdx);
-							return new FlushedNodeHandle(TreeReader, data.Hash, new NodeLocator(data!.Blob, data!.ExportIdx));
+							return new FlushedNodeHandle(TreeReader, new NodeLocator(data.Hash, data!.Blob, data!.ExportIdx));
 						}
 					}
 				}
diff --git a/Engine/Source/Programs/Shared/EpicGames.Horde/Storage/Backends/MemoryStorageClient.cs b/Engine/Source/Programs/Shared/EpicGames.Horde/Storage/Backends/MemoryStorageClient.cs
--- a/Engine/Source/Programs/Shared/EpicGames.Horde/Storage/Backends/MemoryStorageClient.cs
+++ b/Engine/Source/Programs/Shared/EpicGames.Horde/Storage/Backends/MemoryStorageClient.cs
@@ -30,7 +30,7 @@ namespace EpicGames.Horde.Storage.Backends
 		/// <summary>
 		/// Map of ref name to ref data
 		/// </summary>
-		readonly ConcurrentDictionary<RefName, HashedNodeLocator> _refs = new ConcurrentDictionary<RefName, HashedNodeLocator>();
+		readonly ConcurrentDictionary<RefName, NodeLocator> _refs = new ConcurrentDictionary<RefName, NodeLocator>();
 
 		/// <summary>
 		/// Content addressed data lookup
@@ -41,7 +41,7 @@ namespace EpicGames.Horde.Storage.Backends
 		public IReadOnlyDictionary<BlobLocator, Bundle> Blobs => _blobs;
 
 		/// <inheritdoc cref="_refs"/>
-		public IReadOnlyDictionary<RefName, HashedNodeLocator> Refs => _refs;
+		public IReadOnlyDictionary<RefName, NodeLocator> Refs => _refs;
 
 		/// <summary>
 		/// Constructor
@@ -118,7 +118,7 @@ namespace EpicGames.Horde.Storage.Backends
 		/// <inheritdoc/>
 		public override Task<BlobHandle?> TryReadRefTargetAsync(RefName name, RefCacheTime cacheTime = default, CancellationToken cancellationToken = default)
 		{
-			HashedNodeLocator hashedLocator;
+			NodeLocator hashedLocator;
 			if (_refs.TryGetValue(name, out hashedLocator))
 			{
 				return Task.FromResult<BlobHandle?>(new FlushedNodeHandle(TreeReader, hashedLocator)); 
@@ -132,8 +132,7 @@ namespace EpicGames.Horde.Storage.Backends
 		/// <inheritdoc/>
 		public override async Task WriteRefTargetAsync(RefName name, BlobHandle target, RefOptions? options = null, CancellationToken cancellationToken = default)
 		{
-			NodeLocator locator = await target.FlushAsync(cancellationToken);
-			_refs[name] = new HashedNodeLocator(target.Hash, locator);
+			_refs[name] = await target.FlushAsync(cancellationToken);
 		}
 
 		#endregion
diff --git a/Engine/Source/Programs/Shared/EpicGames.Horde/Storage/BlobHandle.cs b/Engine/Source/Programs/Shared/EpicGames.Horde/Storage/BlobHandle.cs
--- a/Engine/Source/Programs/Shared/EpicGames.Horde/Storage/BlobHandle.cs
+++ b/Engine/Source/Programs/Shared/EpicGames.Horde/Storage/BlobHandle.cs
@@ -99,7 +99,7 @@ namespace EpicGames.Horde.Storage
 		public abstract ValueTask<NodeLocator> FlushAsync(CancellationToken cancellationToken = default);
 
 		/// <inheritdoc/>
-		public override string ToString() => HasLocator() ? new HashedNodeLocator(Hash, GetLocator()).ToString() : Hash.ToString();
+		public override string ToString() => HasLocator()? GetLocator().ToString() : Hash.ToString();
 	}
 
 	/// <summary>
diff --git a/Engine/Source/Programs/Shared/EpicGames.Horde/Storage/BundleReader.cs b/Engine/Source/Programs/Shared/EpicGames.Horde/Storage/BundleReader.cs
--- a/Engine/Source/Programs/Shared/EpicGames.Horde/Storage/BundleReader.cs
+++ b/Engine/Source/Programs/Shared/EpicGames.Horde/Storage/BundleReader.cs
@@ -557,7 +557,7 @@ namespace EpicGames.Horde.Storage
 					importBlob = bundleInfo.Header.Imports[reference.ImportIdx];
 				}
 				Debug.Assert(importBlob.IsValid());
-				refs.Add(new FlushedNodeHandle(this, reference.Hash, new NodeLocator(importBlob, reference.NodeIdx)));
+				refs.Add(new FlushedNodeHandle(this, new NodeLocator(reference.Hash, importBlob, reference.NodeIdx)));
 			}
 
 			ReadOnlyMemory<byte> nodeData = ReadOnlyMemory<byte>.Empty;
diff --git a/Engine/Source/Programs/Shared/EpicGames.Horde/Storage/BundleWriter.cs b/Engine/Source/Programs/Shared/EpicGames.Horde/Storage/BundleWriter.cs
--- a/Engine/Source/Programs/Shared/EpicGames.Horde/Storage/BundleWriter.cs
+++ b/Engine/Source/Programs/Shared/EpicGames.Horde/Storage/BundleWriter.cs
@@ -62,16 +62,8 @@ namespace EpicGames.Horde.Storage
 		/// <summary>
 		/// Constructor
 		/// </summary>
-		public FlushedNodeHandle(BundleReader reader, HashedNodeLocator hashedLocator)
-			: this(reader, hashedLocator.Hash, hashedLocator.Locator)
-		{
-		}
-
-		/// <summary>
-		/// Constructor
-		/// </summary>
-		public FlushedNodeHandle(BundleReader reader, IoHash hash, NodeLocator locator)
-			: base(hash)
+		public FlushedNodeHandle(BundleReader reader, NodeLocator locator)
+			: base(locator.Hash)
 		{
 			_reader = reader;
 			_locator = locator;
@@ -547,7 +539,7 @@ namespace EpicGames.Horde.Storage
 
 					for (int idx = 0; idx < _queue.Count; idx++)
 					{
-						NodeLocator nodeLocator = new NodeLocator(locator, idx);
+						NodeLocator nodeLocator = new NodeLocator(_queue[idx].Key.Hash, locator, idx);
 						traceLogger?.LogInformation("Updated pending node {Hash} with locator {Locator}", _queue[idx].Key.Hash, nodeLocator);
 						_queue[idx].MarkAsWritten(nodeLocator);
 					}
diff --git a/Engine/Source/Programs/Shared/EpicGames.Horde/Storage/NodeLocator.cs b/Engine/Source/Programs/Shared/EpicGames.Horde/Storage/NodeLocator.cs
--- a/Engine/Source/Programs/Shared/EpicGames.Horde/Storage/NodeLocator.cs
+++ b/Engine/Source/Programs/Shared/EpicGames.Horde/Storage/NodeLocator.cs
@@ -4,6 +4,7 @@ using System;
 using System.ComponentModel;
 using System.Diagnostics.CodeAnalysis;
 using System.Globalization;
+using System.Text;
 using System.Text.Json;
 using System.Text.Json.Serialization;
 using EpicGames.Core;
@@ -19,6 +20,11 @@ namespace EpicGames.Horde.Storage
 	[CbConverter(typeof(NodeLocatorCbConverter))]
 	public struct NodeLocator : IEquatable<NodeLocator>
 	{
+		/// <summary>
+		/// Hash of the referenced node
+		/// </summary>
+		public IoHash Hash { get; }
+
 		/// <summary>
 		/// Location of the blob containing this node
 		/// </summary>
@@ -32,8 +38,9 @@ namespace EpicGames.Horde.Storage
 		/// <summary>
 		/// Constructor
 		/// </summary>
-		public NodeLocator(BlobLocator blob, int exportIdx)
+		public NodeLocator(IoHash hash, BlobLocator blob, int exportIdx)
 		{
+			Hash = hash;
 			Blob = blob;
 			ExportIdx = exportIdx;
 		}
@@ -50,6 +57,25 @@ namespace EpicGames.Horde.Storage
 		/// <returns></returns>
 		public static NodeLocator Parse(ReadOnlySpan<char> text)
 		{
+			int hashLength = IoHash.NumBytes * 2;
+
+			IoHash hash;
+			if (text.Length == hashLength && IoHash.TryParse(text, out hash))
+			{
+				return new NodeLocator(hash, default, 0);
+			}
+
+			if (text[hashLength] == '@')
+			{
+				hash = IoHash.Parse(text.Slice(0, hashLength));
+				text = text.Slice(hashLength + 1);
+			}
+			else
+			{
+				// For legacy locators that don't have a hash, hash the path instead. This is obviously incorrect, but we never validate hashes and will serve as a stable id during migration.
+				hash = IoHash.Compute(Encoding.UTF8.GetBytes(text.ToString()));
+			}
+
 			int hashIdx = text.IndexOf('#');
 			if (hashIdx == -1)
 			{
@@ -58,7 +84,7 @@ namespace EpicGames.Horde.Storage
 
 			int exportIdx = Int32.Parse(text.Slice(hashIdx + 1), NumberStyles.None, CultureInfo.InvariantCulture);
 			BlobLocator blobLocator = new BlobLocator(new Utf8String(text.Slice(0, hashIdx)));
-			return new NodeLocator(blobLocator, exportIdx);
+			return new NodeLocator(hash, blobLocator, exportIdx);
 		}
 
 		/// <inheritdoc/>
@@ -71,7 +97,7 @@ namespace EpicGames.Horde.Storage
 		public override int GetHashCode() => HashCode.Combine(Blob, ExportIdx);
 
 		/// <inheritdoc/>
-		public override string ToString() => $"{Blob}#{ExportIdx}";
+		public override string ToString() => $"{Hash}@{Blob}#{ExportIdx}";
 
 		/// <inheritdoc/>
 		public static bool operator ==(NodeLocator left, NodeLocator right) => left.Equals(right);
@@ -80,63 +106,6 @@ namespace EpicGames.Horde.Storage
 		public static bool operator !=(NodeLocator left, NodeLocator right) => !left.Equals(right);
 	}
 
-	/// <summary>
-	/// Combination of a NodeLocator with Hash
-	/// </summary>
-	/// <param name="Hash">Hash of the target node</param>
-	/// <param name="Locator">Locator for the node</param>
-	public readonly record struct HashedNodeLocator(IoHash Hash, NodeLocator Locator)
-	{
-		/// <inheritdoc cref="NodeLocator.Blob"/>
-		public BlobLocator Blob => Locator.Blob;
-
-		/// <inheritdoc cref="NodeLocator.ExportIdx"/>
-		public int ExportIdx => Locator.ExportIdx;
-
-		/// <summary>
-		/// Constructor
-		/// </summary>
-		public HashedNodeLocator(IoHash hash, BlobLocator blobLocator, int exportIdx) 
-			: this(hash, new NodeLocator(blobLocator, exportIdx))
-		{
-		}
-
-		/// <summary>
-		/// Parse a node handle value from a string
-		/// </summary>
-		/// <param name="text">Text to parse</param>
-		/// <returns>Parsed node handle</returns>
-		public static HashedNodeLocator Parse(string text)
-		{
-			int hashLength = IoHash.NumBytes * 2;
-			if (text.Length == hashLength)
-			{
-				return new HashedNodeLocator(IoHash.Parse(text), default);
-			}
-			else if (text[hashLength] == '@')
-			{
-				return new HashedNodeLocator(IoHash.Parse(text.Substring(0, hashLength)), NodeLocator.Parse(text.Substring(hashLength + 1)));
-			}
-			else
-			{
-				throw new FormatException("Invalid NodeHandle value");
-			}
-		}
-
-		/// <inheritdoc/>
-		public override string ToString()
-		{
-			if (Locator.IsValid())
-			{
-				return $"{Hash}@{Locator}";
-			}
-			else
-			{
-				return $"{Hash}";
-			}
-		}
-	}
-
 	/// <summary>
 	/// Type converter for <see cref="NodeLocator"/> to and from JSON
 	/// </summary>
diff --git a/Engine/Source/Programs/Shared/EpicGames.IoHash/IoHash.cs b/Engine/Source/Programs/Shared/EpicGames.IoHash/IoHash.cs
--- a/Engine/Source/Programs/Shared/EpicGames.IoHash/IoHash.cs
+++ b/Engine/Source/Programs/Shared/EpicGames.IoHash/IoHash.cs
@@ -189,7 +189,7 @@ namespace EpicGames.Core
 		/// <param name="text"></param>
 		/// <param name="hash">Receives the hash on success</param>
 		/// <returns></returns>
-		public static bool TryParse(string text, out IoHash hash)
+		public static bool TryParse(ReadOnlySpan<char> text, out IoHash hash)
 		{
 			byte[]? bytes;
 			if (StringUtils.TryParseHexString(text, out bytes) && bytes.Length == IoHash.NumBytes)
