Horde: Handle all exceptions raised from terminating EC2 instances via ASG events

diff --git a/Engine/Source/Programs/Horde/Horde.Server/Agents/Fleet/AwsAutoScalingLifecycle.cs b/Engine/Source/Programs/Horde/Horde.Server/Agents/Fleet/AwsAutoScalingLifecycle.cs
--- a/Engine/Source/Programs/Horde/Horde.Server/Agents/Fleet/AwsAutoScalingLifecycle.cs
+++ b/Engine/Source/Programs/Horde/Horde.Server/Agents/Fleet/AwsAutoScalingLifecycle.cs
@@ -182,9 +182,9 @@ public sealed class AwsAutoScalingLifecycleService : IHostedService, IDisposable
 				LifecycleActionEvent ev = JsonSerializer.Deserialize<LifecycleActionEvent>(message.Body) ?? throw new JsonException("Value is null");
 				await InitiateTerminationAsync(ev, cancellationToken);
 			}
-			catch (JsonException je)
+			catch (Exception e)
 			{
-				_logger.LogError(je, "Failed deserializing message. Body={Body}", message.Body);
+				_logger.LogError(e, "Failed initiating termination of EC2 instance. Body={Body}", message.Body);
 			}
 			finally
 			{
