Horde: Add slack thread to step health summary
#skipci
#rnx

diff --git a/Engine/Source/Programs/Horde/HordeDashboard/src/components/jobDetailsV2/JobDetailHealthV2.tsx b/Engine/Source/Programs/Horde/HordeDashboard/src/components/jobDetailsV2/JobDetailHealthV2.tsx
--- a/Engine/Source/Programs/Horde/HordeDashboard/src/components/jobDetailsV2/JobDetailHealthV2.tsx
+++ b/Engine/Source/Programs/Horde/HordeDashboard/src/components/jobDetailsV2/JobDetailHealthV2.tsx
@@ -253,10 +253,14 @@ export const HealthPanel: React.FC<{ jobDetails: JobDetailsV2 }> = observer(({ j
 
       if (column.name === "Quarantine") {
 
-         if (!issue.quarantinedByUserInfo) {
+         if (!issue.quarantinedByUserInfo && !issue.workflowThreadUrl) {
             return null;
          }
 
+         if (issue.workflowThreadUrl) {
+            return <a onClick={(e) => e.stopPropagation()} href={issue.workflowThreadUrl} target="_blank" rel="noreferrer"><Stack style={{ height: "100%" }} horizontalAlign="center" disableShrink={true}  verticalAlign="center" >Slack Thread</Stack></a>
+         }
+
          return <Stack horizontalAlign="start" disableShrink={true} verticalAlign="center" ><Text variant={textSize}>Quarantined</Text></Stack>;
       }
 
