Horde: Adding user job template preference settings
#rnx

diff --git a/Engine/Source/Programs/Horde/Horde.Server.Tests/JobServiceTests.cs b/Engine/Source/Programs/Horde/Horde.Server.Tests/JobServiceTests.cs
--- a/Engine/Source/Programs/Horde/Horde.Server.Tests/JobServiceTests.cs
+++ b/Engine/Source/Programs/Horde/Horde.Server.Tests/JobServiceTests.cs
@@ -1,13 +1,11 @@
 // Copyright Epic Games, Inc. All Rights Reserved.
 
-using System;
 using System.Collections.Generic;
-using System.Runtime.InteropServices;
+using System.Linq;
 using System.Text;
 using System.Threading;
 using System.Threading.Tasks;
 using EpicGames.Core;
-using EpicGames.Horde.Common;
 using Horde.Server.Agents;
 using Horde.Server.Agents.Pools;
 using Horde.Server.Jobs;
@@ -19,7 +17,6 @@ using Horde.Server.Projects;
 using Horde.Server.Server;
 using Horde.Server.Streams;
 using Horde.Server.Users;
-using Horde.Server.Utilities;
 using HordeCommon;
 using Microsoft.CodeAnalysis;
 using Microsoft.VisualStudio.TestTools.UnitTesting;
@@ -144,6 +141,41 @@ namespace Horde.Server.Tests
 			Assert.IsNull(differentTplHash.AbortedByUserId);
 			Assert.IsNull(differentUserName.AbortedByUserId);
 			Assert.IsNull(differentArgs.AbortedByUserId);
+
+			//IUser user = await UserCollection.FindOrAddUserByLoginAsync("elvis");
+			//IUserSettings settings = await UserCollection.GetSettingsAsync(user!.Id);
+		}
+
+		[TestMethod]
+		public async Task TestJobTemplateSettings()
+		{
+			// Scenario: User creates a number of preflights
+			// Expected: The users job template settings are populated and updated in place
+
+			Fixture fixture = await CreateFixtureAsync();
+			await CreatePreflightJob(fixture, "tpl-ref-1", "tpl-hash-1", "elvis", 1000, new[] { "-Target=targeta" });
+			await CreatePreflightJob(fixture, "tpl-ref-1", "tpl-hash-1", "julia", 1000, new[] { "-Target=targeta" });
+			await CreatePreflightJob(fixture, "tpl-ref-other", "tpl-hash-1", "elvis", 1001, new[] { "-Target=targetb" });
+			await CreatePreflightJob(fixture, "tpl-ref-1", "tpl-hash-1", "elvis", 1002, new[] { "-Target=targetc" });
+			await CreatePreflightJob(fixture, "tpl-ref-1", "tpl-hash-1", "elvis", 1003, new[] { "-Target=targetd" });
+			await CreatePreflightJob(fixture, "tpl-ref-other", "tpl-hash-1", "julia", 1004, new[] { "-Target=targete" });
+			await CreatePreflightJob(fixture, "tpl-ref-other", "tpl-hash-1", "elvis", 1004, new[] { "-Target=targete" });
+
+			IUser user = await UserCollection.FindOrAddUserByLoginAsync("elvis")!;
+			IUserSettings settings = await UserCollection.GetSettingsAsync(user!.Id)!;
+
+			Assert.IsNotNull(settings.JobTemplateSettings);
+			Assert.AreEqual(settings.JobTemplateSettings!.Count, 2);
+
+			TemplateId templateRef1 = new TemplateId("tpl-ref-1");			
+			IUserJobTemplateSettings? templateSettings = settings.JobTemplateSettings.FirstOrDefault(x => x.TemplateId ==templateRef1);
+			Assert.IsNotNull(templateSettings);
+			Assert.AreEqual(templateSettings.Arguments[0], "-Target=targetd" );
+
+			TemplateId templateRefOther = new TemplateId("tpl-ref-other");
+			templateSettings = settings.JobTemplateSettings.FirstOrDefault(x => x.TemplateId == templateRefOther);
+			Assert.IsNotNull(templateSettings);
+			Assert.AreEqual(templateSettings.Arguments[0], "-Target=targete");
 		}
 
 		private async Task<IJob> CreatePreflightJob(Fixture fixture, string templateRefId, string templateHash, string startedByUserName, int preflightChange, string[] arguments)
diff --git a/Engine/Source/Programs/Horde/Horde.Server.Tests/TestSetup.cs b/Engine/Source/Programs/Horde/Horde.Server.Tests/TestSetup.cs
--- a/Engine/Source/Programs/Horde/Horde.Server.Tests/TestSetup.cs
+++ b/Engine/Source/Programs/Horde/Horde.Server.Tests/TestSetup.cs
@@ -226,7 +226,7 @@ namespace Horde.Server.Tests
 			services.AddSingleton<ITemplateCollection, TemplateCollection>();
 			services.AddSingleton<IToolCollection, ToolCollection>();
 			services.AddSingleton<IUgsMetadataCollection, UgsMetadataCollection>();
-			services.AddSingleton<IUserCollection, UserCollectionV1>();
+			services.AddSingleton<IUserCollection, UserCollectionV2>();
 			services.AddSingleton<IDeviceCollection, DeviceCollection>();
 			services.AddSingleton<IDashboardPreviewCollection, DashboardPreviewCollection>();
 
diff --git a/Engine/Source/Programs/Horde/Horde.Server/Jobs/JobService.cs b/Engine/Source/Programs/Horde/Horde.Server/Jobs/JobService.cs
--- a/Engine/Source/Programs/Horde/Horde.Server/Jobs/JobService.cs
+++ b/Engine/Source/Programs/Horde/Horde.Server/Jobs/JobService.cs
@@ -222,7 +222,7 @@ namespace Horde.Server.Jobs
 
 			if (options.StartedByUserId != null)
 			{
-				await _userCollection.UpdateSettingsAsync(options.StartedByUserId.Value, addPinnedJobIds: new[] { newJob.Id });
+				await _userCollection.UpdateSettingsAsync(options.StartedByUserId.Value, addPinnedJobIds: new[] { newJob.Id }, templateOptions: new UpdateUserJobTemplateOptions { StreamId = streamConfig.Id, TemplateId = templateRefId, TemplateHash = templateHash.ToString(), Arguments = options.Arguments });
 			}
 
 			await AbortAnyDuplicateJobs(newJob);
diff --git a/Engine/Source/Programs/Horde/Horde.Server/Users/IUser.cs b/Engine/Source/Programs/Horde/Horde.Server/Users/IUser.cs
--- a/Engine/Source/Programs/Horde/Horde.Server/Users/IUser.cs
+++ b/Engine/Source/Programs/Horde/Horde.Server/Users/IUser.cs
@@ -1,8 +1,9 @@
 // Copyright Epic Games, Inc. All Rights Reserved.
 
+using System;
 using System.Collections.Generic;
 using Horde.Server.Jobs;
-using Horde.Server.Utilities;
+using Horde.Server.Streams;
 using MongoDB.Bson;
 
 namespace Horde.Server.Users
@@ -60,6 +61,64 @@ namespace Horde.Server.Users
 		public IReadOnlyList<IUserClaim> Claims { get; }
 	}
 
+	/// <summary>
+	/// 
+	/// </summary>
+	public interface IUserJobTemplateSettings
+	{
+		/// <summary>
+		/// The stream the job was run in
+		/// </summary>
+		public StreamId StreamId { get; }
+
+		/// <summary>
+		/// The template id of the job
+		/// </summary>
+		public TemplateId TemplateId { get; }
+
+		/// <summary>
+		/// The hash of the template definition
+		/// </summary>
+		public string TemplateHash { get; }
+
+		/// <summary>
+		/// The arguments defined when creating the job
+		/// </summary>
+		public IReadOnlyList<string> Arguments { get; }
+
+		/// <summary>
+		/// The last time the job template was used
+		/// </summary>
+		public DateTime UpdateTimeUtc { get; }
+	}
+
+	/// <summary>
+	/// Settings for updating job template user preference
+	/// </summary>
+	public class UpdateUserJobTemplateOptions
+	{
+		/// <summary>
+		/// The stream the job was run in
+		/// </summary>
+		public StreamId StreamId { get; set; }
+
+		/// <summary>
+		/// The template id of the job
+		/// </summary>
+		public TemplateId TemplateId { get; set; }
+
+		/// <summary>
+		/// The hash of the template definition
+		/// </summary>
+		public string TemplateHash { get; set; } = String.Empty;
+
+		/// <summary>
+		/// The arguments defined when creating the job
+		/// </summary>
+		public IReadOnlyList<string> Arguments { get; set; } = new List<string>();
+	}
+
+
 	/// <summary>
 	/// User settings document
 	/// </summary>
@@ -84,5 +143,11 @@ namespace Horde.Server.Users
 		/// List of pinned jobs
 		/// </summary>
 		public IReadOnlyList<JobId> PinnedJobIds { get; }
+
+		/// <summary>
+		/// List of job template preferences
+		/// </summary>
+		public IReadOnlyList<IUserJobTemplateSettings>? JobTemplateSettings { get; }
+
 	}
 }
diff --git a/Engine/Source/Programs/Horde/Horde.Server/Users/IUserCollection.cs b/Engine/Source/Programs/Horde/Horde.Server/Users/IUserCollection.cs
--- a/Engine/Source/Programs/Horde/Horde.Server/Users/IUserCollection.cs
+++ b/Engine/Source/Programs/Horde/Horde.Server/Users/IUserCollection.cs
@@ -91,8 +91,9 @@ namespace Horde.Server.Users
 		/// <param name="dashboardSettings">Opaque settings object for the dashboard</param>
 		/// <param name="addPinnedJobIds"></param>
 		/// <param name="removePinnedJobIds"></param>
+		/// <param name="templateOptions"></param>
 		/// <returns>Updated user object</returns>
-		Task UpdateSettingsAsync(UserId userId, bool? enableExperimentalFeatures = null, BsonValue? dashboardSettings = null, IEnumerable<JobId>? addPinnedJobIds = null, IEnumerable<JobId>? removePinnedJobIds = null);
+		Task UpdateSettingsAsync(UserId userId, bool? enableExperimentalFeatures = null, BsonValue? dashboardSettings = null, IEnumerable<JobId>? addPinnedJobIds = null, IEnumerable<JobId>? removePinnedJobIds = null, UpdateUserJobTemplateOptions? templateOptions = null);
 	}
 
 	/// <summary>
diff --git a/Engine/Source/Programs/Horde/Horde.Server/Users/UserCollection.cs b/Engine/Source/Programs/Horde/Horde.Server/Users/UserCollection.cs
--- a/Engine/Source/Programs/Horde/Horde.Server/Users/UserCollection.cs
+++ b/Engine/Source/Programs/Horde/Horde.Server/Users/UserCollection.cs
@@ -42,6 +42,8 @@ namespace Horde.Server.Users
 
 			UserId IUserSettings.UserId => Id;
 			IReadOnlyList<JobId> IUserSettings.PinnedJobIds => PinnedJobIds;
+
+			IReadOnlyList<IUserJobTemplateSettings>? IUserSettings.JobTemplateSettings => null;
 		}
 
 		class ClaimDocument : IUserClaim
@@ -155,7 +157,7 @@ namespace Horde.Server.Users
 		}
 
 		/// <inheritdoc/>
-		public async Task UpdateSettingsAsync(UserId userId, bool? enableExperimentalFeatures, BsonValue? dashboardSettings = null, IEnumerable<JobId>? addPinnedJobIds = null, IEnumerable<JobId>? removePinnedJobIds = null)
+		public async Task UpdateSettingsAsync(UserId userId, bool? enableExperimentalFeatures, BsonValue? dashboardSettings = null, IEnumerable<JobId>? addPinnedJobIds = null, IEnumerable<JobId>? removePinnedJobIds = null, UpdateUserJobTemplateOptions? templateOptions = null)
 		{
 			if (addPinnedJobIds != null)
 			{
diff --git a/Engine/Source/Programs/Horde/Horde.Server/Users/UserCollectionV2.cs b/Engine/Source/Programs/Horde/Horde.Server/Users/UserCollectionV2.cs
--- a/Engine/Source/Programs/Horde/Horde.Server/Users/UserCollectionV2.cs
+++ b/Engine/Source/Programs/Horde/Horde.Server/Users/UserCollectionV2.cs
@@ -6,6 +6,7 @@ using System.Linq;
 using System.Threading.Tasks;
 using Horde.Server.Jobs;
 using Horde.Server.Server;
+using Horde.Server.Streams;
 using Horde.Server.Utilities;
 using Microsoft.Extensions.Caching.Memory;
 using Microsoft.Extensions.Logging;
@@ -83,6 +84,42 @@ namespace Horde.Server.Users
 			}
 		}
 
+		class JobTemplateSettingsDocument : IUserJobTemplateSettings
+		{
+			public StreamId StreamId { get; set; }
+
+			public TemplateId TemplateId { get; set; } 
+
+			public string TemplateHash { get; set; } = String.Empty;
+
+			public List<string> Arguments { get; set; } = new List<string>();
+			IReadOnlyList<string> IUserJobTemplateSettings.Arguments => Arguments;
+
+			public DateTime UpdateTimeUtc { get; set; }
+
+			[BsonConstructor]
+			private JobTemplateSettingsDocument()
+			{
+
+			}
+
+			/// <summary>
+			/// Constructor
+			/// </summary>
+			/// <param name="streamId"></param>
+			/// <param name="templateId"></param>
+			/// <param name="templateHash"></param>
+			/// <param name="arguments"></param>
+			public JobTemplateSettingsDocument(StreamId streamId, TemplateId templateId, string templateHash, List<string> arguments)
+			{
+				StreamId = streamId;
+				TemplateId = templateId;
+				TemplateHash = templateHash;
+				Arguments = arguments;
+				UpdateTimeUtc= DateTime.UtcNow;
+			}
+		}
+
 		class UserSettingsDocument : IUserSettings
 		{
 			public UserId Id { get; set; }
@@ -96,6 +133,9 @@ namespace Horde.Server.Users
 			UserId IUserSettings.UserId => Id;
 			IReadOnlyList<JobId> IUserSettings.PinnedJobIds => PinnedJobIds;
 
+			public List<JobTemplateSettingsDocument> JobTemplateSettings { get; set; } = new List<JobTemplateSettingsDocument>();
+			IReadOnlyList<IUserJobTemplateSettings>? IUserSettings.JobTemplateSettings => JobTemplateSettings;
+
 			[BsonConstructor]
 			private UserSettingsDocument()
 			{
@@ -312,7 +352,7 @@ namespace Horde.Server.Users
 		}
 
 		/// <inheritdoc/>
-		public async Task UpdateSettingsAsync(UserId userId, bool? enableExperimentalFeatures = null, BsonValue? dashboardSettings = null, IEnumerable<JobId>? addPinnedJobIds = null, IEnumerable<JobId>? removePinnedJobIds = null)
+		public async Task UpdateSettingsAsync(UserId userId, bool? enableExperimentalFeatures = null, BsonValue? dashboardSettings = null, IEnumerable<JobId>? addPinnedJobIds = null, IEnumerable<JobId>? removePinnedJobIds = null, UpdateUserJobTemplateOptions? templateOptions = null)
 		{
 			List<UpdateDefinition<UserSettingsDocument>> updates = new List<UpdateDefinition<UserSettingsDocument>>();
 			if (enableExperimentalFeatures != null)
@@ -331,6 +371,18 @@ namespace Horde.Server.Users
 			{
 				updates.Add(Builders<UserSettingsDocument>.Update.PullAll(x => x.PinnedJobIds, removePinnedJobIds));
 			}
+
+			if (templateOptions != null)
+			{
+				JobTemplateSettingsDocument doc = new JobTemplateSettingsDocument(templateOptions.StreamId, templateOptions.TemplateId, templateOptions.TemplateHash, templateOptions.Arguments.ToList());
+				FilterDefinition<UserSettingsDocument> filter = Builders<UserSettingsDocument>.Filter.Eq(x => x.Id, userId) & Builders<UserSettingsDocument>.Filter.ElemMatch(x => x.JobTemplateSettings, t => t.StreamId == templateOptions.StreamId && t.TemplateId == templateOptions.TemplateId);
+				UpdateResult result = await _userSettings.UpdateOneAsync(filter, Builders<UserSettingsDocument>.Update.Set(x => x.JobTemplateSettings[-1], doc));
+				if (result.ModifiedCount == 0)
+				{
+					updates.Add(Builders<UserSettingsDocument>.Update.PushEach(x => x.JobTemplateSettings, new[] { doc }, -100));
+				}
+			}
+
 			if (updates.Count > 0)
 			{
 				await _userSettings.UpdateOneAsync<UserSettingsDocument>(x => x.Id == userId, Builders<UserSettingsDocument>.Update.Combine(updates), new UpdateOptions { IsUpsert = true });
