Horde: Adding job template preferences to user response
#rnx

diff --git a/Engine/Source/Programs/Horde/Horde.Server/Users/UserMessages.cs b/Engine/Source/Programs/Horde/Horde.Server/Users/UserMessages.cs
--- a/Engine/Source/Programs/Horde/Horde.Server/Users/UserMessages.cs
+++ b/Engine/Source/Programs/Horde/Horde.Server/Users/UserMessages.cs
@@ -9,6 +9,7 @@ using EpicGames.Core;
 using Horde.Server.Acls;
 using Horde.Server.Agents;
 using Horde.Server.Agents.Pools;
+using Horde.Server.Jobs;
 using Horde.Server.Server;
 using Horde.Server.Server.Notices;
 using MongoDB.Bson;
@@ -75,6 +76,11 @@ namespace Horde.Server.Users
 		/// </summary>
 		public GetDashboardFeaturesResponse? DashboardFeatures { get; set; }
 
+		/// <summary>
+		/// User job template preferences
+		/// </summary>
+		public List<GetJobTemplateSettingsResponse>? JobTemplateSettings { get; set; }
+
 		/// <summary>
 		/// List of pinned job ids
 		/// </summary>
@@ -102,10 +108,62 @@ namespace Horde.Server.Users
 				
 				DashboardSettings = BsonTypeMapper.MapToDotNetValue(settings.DashboardSettings);
 				PinnedJobIds = settings.PinnedJobIds.ConvertAll(x => x.ToString());
+
+				if (settings.JobTemplateSettings != null && settings.JobTemplateSettings.Count > 0)
+				{
+					JobTemplateSettings = new List<GetJobTemplateSettingsResponse>();
+					for (int i = 0; i < settings.JobTemplateSettings.Count; i++)
+					{
+						JobTemplateSettings.Add(new GetJobTemplateSettingsResponse(settings.JobTemplateSettings[i]));
+					}
+				}
 			}
 		}
 	}
 
+	/// <summary>
+	/// Job template settings for the current user
+	/// </summary>
+	public class GetJobTemplateSettingsResponse
+	{
+		/// <summary>
+		/// The stream the job was run in
+		/// </summary>
+		public string StreamId { get; set; }
+
+		/// <summary>
+		/// The template id of the job
+		/// </summary>
+		public string TemplateId { get; set; }
+
+		/// <summary>
+		/// The hash of the template definition
+		/// </summary>
+		public string TemplateHash { get; set; }
+
+		/// <summary>
+		/// The arguments defined when creating the job
+		/// </summary>
+		public List<string> Arguments { get; set; }
+
+		/// <summary>
+		/// The last update time of the job template
+		/// </summary>
+		public DateTimeOffset UpdateTimeUtc { get; set; }
+
+		/// <summary>
+		/// Constructor
+		/// </summary>
+		public GetJobTemplateSettingsResponse(IUserJobTemplateSettings settings)
+		{
+			StreamId = settings.StreamId.ToString();
+			TemplateId = settings.TemplateId.ToString();
+			TemplateHash = settings.TemplateHash.ToString();
+			Arguments = settings.Arguments.ToList();
+			UpdateTimeUtc = new DateTimeOffset(settings.UpdateTimeUtc);
+		}
+	}
+
 	/// <summary>
 	/// Settings for whether various features should be enabled on the dashboard
 	/// </summary>
