Horde: Small cleanup to some vestigial test code
#skipci
#rnx

diff --git a/Engine/Source/Programs/Horde/Horde.Server.Tests/JobServiceTests.cs b/Engine/Source/Programs/Horde/Horde.Server.Tests/JobServiceTests.cs
--- a/Engine/Source/Programs/Horde/Horde.Server.Tests/JobServiceTests.cs
+++ b/Engine/Source/Programs/Horde/Horde.Server.Tests/JobServiceTests.cs
@@ -141,9 +141,6 @@ namespace Horde.Server.Tests
 			Assert.IsNull(differentTplHash.AbortedByUserId);
 			Assert.IsNull(differentUserName.AbortedByUserId);
 			Assert.IsNull(differentArgs.AbortedByUserId);
-
-			//IUser user = await UserCollection.FindOrAddUserByLoginAsync("elvis");
-			//IUserSettings settings = await UserCollection.GetSettingsAsync(user!.Id);
 		}
 
 		[TestMethod]
