Horde: Adding user job template preferences to dashboard
#skipci
#rnx

diff --git a/Engine/Source/Programs/Horde/HordeDashboard/src/backend/Api.ts b/Engine/Source/Programs/Horde/HordeDashboard/src/backend/Api.ts
--- a/Engine/Source/Programs/Horde/HordeDashboard/src/backend/Api.ts
+++ b/Engine/Source/Programs/Horde/HordeDashboard/src/backend/Api.ts
@@ -3143,6 +3143,26 @@ export type GetDashboardFeaturesResponse = {
 	showTests?: boolean;
 }
 
+/// Job template settings for the current user
+export type GetJobTemplateSettingsResponse = {
+
+	/// The stream the job was run in	
+	streamId: string;
+
+	/// The template id of the job	
+	templateId: string;
+
+	/// The hash of the template definition	
+	templateHash: string;
+
+	/// The arguments defined when creating the job	
+	arguments: string[];
+
+	/// The last update time of the job template	
+	updateTimeUtc: Date | string;
+
+}
+
 /**  Response describing the current user */
 export type GetUserResponse = {
 
@@ -3181,6 +3201,9 @@ export type GetUserResponse = {
 	/// </summary>
 	dashboardFeatures?: GetDashboardFeaturesResponse;
 
+	// array of user job templates settings
+	jobTemplateSettings?: GetJobTemplateSettingsResponse[];
+
 	/** List of pinned job ids */
 	pinnedJobIds?: string[];
 
diff --git a/Engine/Source/Programs/Horde/HordeDashboard/src/backend/Dashboard.ts b/Engine/Source/Programs/Horde/HordeDashboard/src/backend/Dashboard.ts
--- a/Engine/Source/Programs/Horde/HordeDashboard/src/backend/Dashboard.ts
+++ b/Engine/Source/Programs/Horde/HordeDashboard/src/backend/Dashboard.ts
@@ -3,7 +3,7 @@
 import { getTheme } from "@fluentui/react";
 import { action, makeObservable, observable } from 'mobx';
 import backend from '.';
-import { DashboardPreference, GetDashboardConfigResponse, GetUserResponse, UserClaim } from './Api';
+import { DashboardPreference, GetDashboardConfigResponse, GetJobTemplateSettingsResponse, GetUserResponse, UserClaim } from './Api';
 
 const theme = getTheme();
 
@@ -114,23 +114,23 @@ export class Dashboard {
     }
 
     get preview(): boolean {
-        
-        try {            
+
+        try {
             return !!(process.env.REACT_APP_HORDE_DEBUG_PREVIEW) || window?.location?.hostname?.indexOf("preview") !== -1;
         } catch (reason) {
             console.error(reason);
         }
-        return false;        
+        return false;
     }
 
     get development(): boolean {
 
-        try {            
+        try {
             return window?.location?.hostname?.indexOf("devtools-dev") !== -1;
         } catch (reason) {
             console.error(reason);
         }
-        return false;        
+        return false;
     }
 
     get browser(): WebBrowser {
@@ -178,6 +178,31 @@ export class Dashboard {
         return this.data.pinnedJobIds ?? [];
     }
 
+    get lastJobTemplateSettings(): GetJobTemplateSettingsResponse | undefined {
+
+        try {            
+            if (!this.data.jobTemplateSettings?.length) {
+                return undefined;
+            }
+
+            const sorted = this.data.jobTemplateSettings.sort((a, b) => {
+                return new Date(b.updateTimeUtc).getTime() - new Date(a.updateTimeUtc).getTime();
+            });
+
+            return sorted[0];
+        } catch (error) {
+            console.error(error);
+        }
+
+        return undefined;
+    }
+
+    set jobTemplateSettings(settings: GetJobTemplateSettingsResponse[]) {
+        if (this.data) {
+            this.data.jobTemplateSettings = settings;
+        }
+    }
+
     get experimentalFeatures(): boolean {
 
         return this.data.enableExperimentalFeatures!;
@@ -281,9 +306,9 @@ export class Dashboard {
             return true;
         }
         */
-        
+
         const value = this.preferences.get(DashboardPreference.LocalCache) !== 'false';
-        
+
         if (value && !this.hasLoggedLocalCache) {
             this.hasLoggedLocalCache = true;
             console.log("Local graph and template caching is enabled")
diff --git a/Engine/Source/Programs/Horde/HordeDashboard/src/components/ErrorHandler.tsx b/Engine/Source/Programs/Horde/HordeDashboard/src/components/ErrorHandler.tsx
--- a/Engine/Source/Programs/Horde/HordeDashboard/src/components/ErrorHandler.tsx
+++ b/Engine/Source/Programs/Horde/HordeDashboard/src/components/ErrorHandler.tsx
@@ -222,10 +222,12 @@ export const ErrorDialog: React.FC = observer(() => {
 
    let fontSize = 13;
 
+   let width = 800;
    if (message?.length) {
       try {
          const exception = JSON.parse(message);
          if (exception?.exception?.message && exception?.exception?.trace) {
+            width = 1380;
             fontSize = 11
             message = `${exception?.exception?.message}\n${exception?.exception?.trace}`
          }
@@ -236,7 +238,7 @@ export const ErrorDialog: React.FC = observer(() => {
 
    return (
       <Dialog
-         minWidth={1380}
+         minWidth={width}
          hidden={!error}
          onDismiss={() => ErrorHandler.clear()}
          modalProps={{
diff --git a/Engine/Source/Programs/Horde/HordeDashboard/src/components/NewBuild.tsx b/Engine/Source/Programs/Horde/HordeDashboard/src/components/NewBuild.tsx
--- a/Engine/Source/Programs/Horde/HordeDashboard/src/components/NewBuild.tsx
+++ b/Engine/Source/Programs/Horde/HordeDashboard/src/components/NewBuild.tsx
@@ -12,6 +12,7 @@ import { ErrorHandler } from "../components/ErrorHandler";
 import { hordeClasses, modeColors } from '../styles/Styles';
 import { useQuery } from './JobDetailCommon';
 import { JobDetailsV2 } from './jobDetailsV2/JobDetailsViewCommon';
+import dashboard from '../backend/Dashboard';
 
 let toolTipId = 0;
 
@@ -788,6 +789,13 @@ export const NewBuild: React.FC<{ streamId: string; show: boolean; onClose: (new
          }
       }
 
+      if (!t) {
+         const pref = dashboard.lastJobTemplateSettings;
+         if (pref) {
+            t = templates.find(t => stream.id === pref.streamId && t.id === pref.templateId)
+         }
+      }      
+
       // default to sane template when all are shown
       if (!t && (showAllTemplates) && stream.tabs.length > 0) {
          const stab = stream.tabs[0] as JobsTabData;
@@ -1390,6 +1398,11 @@ export const NewBuild: React.FC<{ streamId: string; show: boolean; onClose: (new
                console.log("Updating notifications")
                try {
                   await backend.updateNotification({ slack: true }, "job", data.id);
+                  const user = await backend.getCurrentUser();
+                  if (user.jobTemplateSettings) {
+                     dashboard.jobTemplateSettings = user.jobTemplateSettings;
+                  }
+                  
                } catch (reason) {
                   console.log(`Error on updating notifications: ${reason}`);
                }
