Horde: Fix last template selection for new build dialog
#rnx
#skipci

diff --git a/Engine/Source/Programs/Horde/HordeDashboard/src/backend/Dashboard.ts b/Engine/Source/Programs/Horde/HordeDashboard/src/backend/Dashboard.ts
--- a/Engine/Source/Programs/Horde/HordeDashboard/src/backend/Dashboard.ts
+++ b/Engine/Source/Programs/Horde/HordeDashboard/src/backend/Dashboard.ts
@@ -178,14 +178,17 @@ export class Dashboard {
         return this.data.pinnedJobIds ?? [];
     }
 
-    get lastJobTemplateSettings(): GetJobTemplateSettingsResponse | undefined {
+    getLastJobTemplateSettings(streamId: string, templateIds: string[]): GetJobTemplateSettingsResponse | undefined {
 
-        try {            
-            if (!this.data.jobTemplateSettings?.length) {
+        try {
+            
+            const streamTemplates = this.data?.jobTemplateSettings?.filter(t => t.streamId === streamId && templateIds.indexOf(t.templateId) !== -1)
+
+            if (!streamTemplates?.length) {
                 return undefined;
             }
 
-            const sorted = this.data.jobTemplateSettings.sort((a, b) => {
+            const sorted = streamTemplates.sort((a, b) => {
                 return new Date(b.updateTimeUtc).getTime() - new Date(a.updateTimeUtc).getTime();
             });
 
diff --git a/Engine/Source/Programs/Horde/HordeDashboard/src/components/NewBuild.tsx b/Engine/Source/Programs/Horde/HordeDashboard/src/components/NewBuild.tsx
--- a/Engine/Source/Programs/Horde/HordeDashboard/src/components/NewBuild.tsx
+++ b/Engine/Source/Programs/Horde/HordeDashboard/src/components/NewBuild.tsx
@@ -790,12 +790,11 @@ export const NewBuild: React.FC<{ streamId: string; show: boolean; onClose: (new
       }
 
       if (!t) {
-         const pref = dashboard.lastJobTemplateSettings;
+         const pref = dashboard.getLastJobTemplateSettings(stream.id, templates.map(t => t.id));
          if (pref) {
             t = templates.find(t => stream.id === pref.streamId && t.id === pref.templateId)
          }
-      }      
-
+      }
       // default to sane template when all are shown
       if (!t && (showAllTemplates) && stream.tabs.length > 0) {
          const stab = stream.tabs[0] as JobsTabData;
