Horde: Move long pool creation/update argument lists into dedicated classes.

diff --git a/Engine/Source/Programs/Horde/Horde.Server.Tests/Agents/Pools/PoolUpdateServiceTest.cs b/Engine/Source/Programs/Horde/Horde.Server.Tests/Agents/Pools/PoolUpdateServiceTest.cs
--- a/Engine/Source/Programs/Horde/Horde.Server.Tests/Agents/Pools/PoolUpdateServiceTest.cs
+++ b/Engine/Source/Programs/Horde/Horde.Server.Tests/Agents/Pools/PoolUpdateServiceTest.cs
@@ -27,7 +27,7 @@ namespace Horde.Server.Tests.Agents.Pools
 		[TestInitialize]
 		public async Task Setup()
 		{
-			_pool = await PoolService.CreatePoolAsync("testPool", null, true);
+			_pool = await PoolService.CreatePoolAsync("testPool", new AddPoolOptions { EnableAutoscaling = true });
 			_enabledAgent = await CreateAgentAsync(_pool, true);
 			_disabledAgent = await CreateAgentAsync(_pool, false);
 			_disabledAgentBeyondGracePeriod = await CreateAgentAsync(_pool, enabled: false, adjustClockBy: -TimeSpan.FromHours(9));
@@ -68,7 +68,7 @@ namespace Horde.Server.Tests.Agents.Pools
 		{
 			// Arrange
 			// Explicitly set the grace period for the pool to be longer than the default of 8 hours
-			await PoolCollection.TryUpdateAsync(_pool, shutdownIfDisabledGracePeriod: TimeSpan.FromHours(24));
+			await PoolCollection.TryUpdateAsync(_pool, new UpdatePoolOptions { ShutdownIfDisabledGracePeriod = TimeSpan.FromHours(24) });
 			
 			// Act
 			await _pus.ShutdownDisabledAgentsAsync(CancellationToken.None);
@@ -84,7 +84,7 @@ namespace Horde.Server.Tests.Agents.Pools
 		public async Task ShutdownDisabledAgents_WithAutoScalingOff_DoesNotRequestShutdown()
 		{
 			// Arrange
-			await PoolCollection.TryUpdateAsync(_pool, newEnableAutoscaling: false);
+			await PoolCollection.TryUpdateAsync(_pool, new UpdatePoolOptions { EnableAutoscaling = false });
 			
 			// Act
 			await _pus.ShutdownDisabledAgentsAsync(CancellationToken.None);
diff --git a/Engine/Source/Programs/Horde/Horde.Server.Tests/Fleet/AwsRecyclingFleetManagerTest.cs b/Engine/Source/Programs/Horde/Horde.Server.Tests/Fleet/AwsRecyclingFleetManagerTest.cs
--- a/Engine/Source/Programs/Horde/Horde.Server.Tests/Fleet/AwsRecyclingFleetManagerTest.cs
+++ b/Engine/Source/Programs/Horde/Horde.Server.Tests/Fleet/AwsRecyclingFleetManagerTest.cs
@@ -217,7 +217,7 @@ namespace Horde.Server.Tests.Fleet
 				});
 
 			ILogger<AwsRecyclingFleetManager> logger = loggerFactory.CreateLogger<AwsRecyclingFleetManager>();
-			IPool pool = await PoolService.CreatePoolAsync("testPool", null, true, 0, 0, sizeStrategy: PoolSizeStrategy.NoOp);
+			IPool pool = await PoolService.CreatePoolAsync("testPool", new AddPoolOptions { EnableAutoscaling = true, MinAgents = 0, NumReserveAgents = 0, SizeStrategy = PoolSizeStrategy.NoOp });
 			AwsRecyclingFleetManager manager = new (ec2, AgentCollection, Meter, _clock, settings, Tracer, logger);
 			return await manager.ExpandPoolAsync(pool, new List<IAgent>(), numRequestedInstances, CancellationToken.None);
 		}
diff --git a/Engine/Source/Programs/Horde/Horde.Server.Tests/Fleet/AwsReuseFleetManagerTest.cs b/Engine/Source/Programs/Horde/Horde.Server.Tests/Fleet/AwsReuseFleetManagerTest.cs
--- a/Engine/Source/Programs/Horde/Horde.Server.Tests/Fleet/AwsReuseFleetManagerTest.cs
+++ b/Engine/Source/Programs/Horde/Horde.Server.Tests/Fleet/AwsReuseFleetManagerTest.cs
@@ -45,7 +45,7 @@ namespace Horde.Server.Tests.Fleet
 		private async Task ExpandPoolAsync(IAmazonEC2 ec2, int numRequestedInstances, AwsReuseFleetManagerSettings settings)
 		{
 			ILogger<AwsReuseFleetManager> logger = ServiceProvider.GetRequiredService<ILogger<AwsReuseFleetManager>>();
-			IPool pool = await PoolService.CreatePoolAsync("testPool", null, true, 0, 0, sizeStrategy: PoolSizeStrategy.NoOp);
+			IPool pool = await PoolService.CreatePoolAsync("testPool", new AddPoolOptions { EnableAutoscaling = true, MinAgents = 0, NumReserveAgents = 0, SizeStrategy = PoolSizeStrategy.NoOp});
 			AwsReuseFleetManager manager = new (ec2, AgentCollection, settings, Tracer, logger);
 			await manager.ExpandPoolAsync(pool, new List<IAgent>(), numRequestedInstances, CancellationToken.None);
 		}
diff --git a/Engine/Source/Programs/Horde/Horde.Server.Tests/Fleet/FleetServiceTest.cs b/Engine/Source/Programs/Horde/Horde.Server.Tests/Fleet/FleetServiceTest.cs
--- a/Engine/Source/Programs/Horde/Horde.Server.Tests/Fleet/FleetServiceTest.cs
+++ b/Engine/Source/Programs/Horde/Horde.Server.Tests/Fleet/FleetServiceTest.cs
@@ -104,7 +104,7 @@ namespace Horde.Server.Tests.Fleet
 		public async Task OnlyEnabledAgentsAreAutoScaled()
 		{
 			using FleetService service = GetFleetService(_fleetManagerSpy);
-			IPool pool = await PoolService.CreatePoolAsync("testPool", null, true, 0, 0, sizeStrategy: PoolSizeStrategy.LeaseUtilization);
+			IPool pool = await PoolService.CreatePoolAsync("testPool", new AddPoolOptions { EnableAutoscaling = true, MinAgents = 0, NumReserveAgents = 0, SizeStrategy = PoolSizeStrategy.LeaseUtilization });
 			await CreateAgentAsync(pool, true);
 			await CreateAgentAsync(pool, true);
 			await CreateAgentAsync(pool, false);
@@ -120,7 +120,7 @@ namespace Horde.Server.Tests.Fleet
 		public async Task ScaleOutWithPendingShutdowns()
 		{
 			using FleetService service = GetFleetService(new FakeFleetManager());
-			IPool pool = await PoolService.CreatePoolAsync("testPool", null, true, 0, 0, sizeStrategy: PoolSizeStrategy.LeaseUtilization);
+			IPool pool = await PoolService.CreatePoolAsync("testPool", new AddPoolOptions { EnableAutoscaling = true, MinAgents = 0, NumReserveAgents = 0, SizeStrategy = PoolSizeStrategy.LeaseUtilization });
 			IAgent agent1 = await CreateAgentAsync(pool, true);
 			IAgent agent2 = await CreateAgentAsync(pool, true, requestShutdown: true);
 
@@ -138,7 +138,7 @@ namespace Horde.Server.Tests.Fleet
 		public async Task ScaleOutWithPendingShutdownsWithStoppedAgents()
 		{
 			using FleetService service = GetFleetService(new FakeFleetManager());
-			IPool pool = await PoolService.CreatePoolAsync("testPool", null, true, 0, 0, sizeStrategy: PoolSizeStrategy.LeaseUtilization);
+			IPool pool = await PoolService.CreatePoolAsync("testPool", new AddPoolOptions { EnableAutoscaling = true, MinAgents = 0, NumReserveAgents = 0, SizeStrategy = PoolSizeStrategy.LeaseUtilization });
 			IAgent agent1 = await CreateAgentAsync(pool, true);
 			IAgent agent2 = await CreateAgentAsync(pool, false);
 			IAgent agent3 = await CreateAgentAsync(pool, true, requestShutdown: true);
@@ -157,7 +157,7 @@ namespace Horde.Server.Tests.Fleet
 		public async Task ScaleOutCooldown()
 		{
 			using FleetService service = GetFleetService(_fleetManagerSpy);
-			IPool pool = await PoolService.CreatePoolAsync("testPool", null, true, 0, 0, sizeStrategy: PoolSizeStrategy.NoOp);
+			IPool pool = await PoolService.CreatePoolAsync("testPool", new AddPoolOptions { EnableAutoscaling = true, MinAgents = 0, NumReserveAgents = 0, SizeStrategy = PoolSizeStrategy.NoOp });
 
 			// First scale-out will succeed
 			ScaleResult result = await service.ScalePoolAsync(pool, new List<IAgent>(), new PoolSizeResult(0, 1), CancellationToken.None);
@@ -180,7 +180,7 @@ namespace Horde.Server.Tests.Fleet
 		public async Task ScaleInCooldown()
 		{
 			using FleetService service = GetFleetService(_fleetManagerSpy);
-			IPool pool = await PoolService.CreatePoolAsync("testPool", null, true, 0, 0, sizeStrategy: PoolSizeStrategy.NoOp);
+			IPool pool = await PoolService.CreatePoolAsync("testPool", new AddPoolOptions { EnableAutoscaling = true, MinAgents = 0, NumReserveAgents = 0, SizeStrategy = PoolSizeStrategy.NoOp });
 			IAgent agent1 = await CreateAgentAsync(pool);
 			IAgent agent2 = await CreateAgentAsync(pool);
 
@@ -205,7 +205,7 @@ namespace Horde.Server.Tests.Fleet
 		public async Task ScaleOutDuringDowntime()
 		{
 			using FleetService service = GetFleetService(_fleetManagerSpy, true);
-			IPool pool = await PoolService.CreatePoolAsync("testPool", null, true, 0, 0, sizeStrategy: PoolSizeStrategy.NoOp);
+			IPool pool = await PoolService.CreatePoolAsync("testPool", new AddPoolOptions { EnableAutoscaling = true, MinAgents = 0, NumReserveAgents = 0, SizeStrategy = PoolSizeStrategy.NoOp });
 
 			ScaleResult result = await service.ScalePoolAsync(pool, new List<IAgent>(), new PoolSizeResult(0, 1), CancellationToken.None);
 			Assert.AreEqual(new ScaleResult(FleetManagerOutcome.NoOp, 0, 0), result);
@@ -216,7 +216,7 @@ namespace Horde.Server.Tests.Fleet
 		public async Task ScaleInDuringDowntime()
 		{
 			using FleetService service = GetFleetService(_fleetManagerSpy, true);
-			IPool pool = await PoolService.CreatePoolAsync("testPool", null, true, 0, 0, sizeStrategy: PoolSizeStrategy.NoOp);
+			IPool pool = await PoolService.CreatePoolAsync("testPool", new AddPoolOptions { EnableAutoscaling = true, MinAgents = 0, NumReserveAgents = 0, SizeStrategy = PoolSizeStrategy.NoOp });
 			IAgent agent1 = await CreateAgentAsync(pool);
 			IAgent agent2 = await CreateAgentAsync(pool);
 
@@ -246,10 +246,10 @@ namespace Horde.Server.Tests.Fleet
 		[TestMethod]
 		public async Task CreateJobQueueFromLegacySettings()
 		{
-			IPool pool1 = await PoolService.CreatePoolAsync("test1", sizeStrategy: PoolSizeStrategy.JobQueue);
+			IPool pool1 = await PoolService.CreatePoolAsync("test1", new AddPoolOptions { SizeStrategy = PoolSizeStrategy.JobQueue });
 			Assert.AreEqual(typeof(JobQueueStrategy), FleetService.CreatePoolSizeStrategy(pool1).GetType());
 			
-			IPool pool2 = await PoolService.CreatePoolAsync("test2", sizeStrategy: PoolSizeStrategy.JobQueue, jobQueueSettings: new JobQueueSettings(22, 33));
+			IPool pool2 = await PoolService.CreatePoolAsync("test2", new AddPoolOptions { SizeStrategy = PoolSizeStrategy.JobQueue, JobQueueSettings = new JobQueueSettings(22, 33) });
 			IPoolSizeStrategy s = FleetService.CreatePoolSizeStrategy(pool2);
 			Assert.AreEqual(typeof(JobQueueStrategy), s.GetType());
 			Assert.AreEqual(22.0, ((JobQueueStrategy)s).Settings.ScaleOutFactor);
@@ -259,7 +259,7 @@ namespace Horde.Server.Tests.Fleet
 		[TestMethod]
 		public async Task CreateLeaseUtilizationFromLegacySettings()
 		{
-			IPool pool = await PoolService.CreatePoolAsync("test1", sizeStrategy: PoolSizeStrategy.LeaseUtilization);
+			IPool pool = await PoolService.CreatePoolAsync("test1", new AddPoolOptions { SizeStrategy = PoolSizeStrategy.LeaseUtilization });
 			Assert.AreEqual(typeof(LeaseUtilizationStrategy), FleetService.CreatePoolSizeStrategy(pool).GetType());
 		}
 
@@ -412,8 +412,8 @@ namespace Horde.Server.Tests.Fleet
 		
 		private async Task<IPoolSizeStrategy> CreateStrategy(params PoolSizeStrategyInfo[] infos)
 		{
-			IPool pool = await PoolService.CreatePoolAsync("testPool-" + s_poolCount++, null, true, 0, 0, sizeStrategy: PoolSizeStrategy.NoOp);
-			await PoolService.UpdatePoolAsync(pool, newSizeStrategies: infos.ToList());
+			IPool pool = await PoolService.CreatePoolAsync("testPool-" + s_poolCount++, new AddPoolOptions { EnableAutoscaling = true, MinAgents = 0, NumReserveAgents = 0, SizeStrategy = PoolSizeStrategy.NoOp });
+			await PoolService.UpdatePoolAsync(pool, new UpdatePoolOptions { SizeStrategies = infos.ToList() });
 			return FleetService.CreatePoolSizeStrategy(pool);
 		}
 	}
@@ -490,7 +490,7 @@ namespace Horde.Server.Tests.Fleet
 		
 		private async Task<IFleetManager> CreateFleetManager(params FleetManagerInfo[] infos)
 		{
-			IPool pool = await PoolService.CreatePoolAsync("testPool-" + s_poolCount++, null, true, 0, 0, fleetManagers: infos.ToList(), sizeStrategy: PoolSizeStrategy.NoOp);
+			IPool pool = await PoolService.CreatePoolAsync("testPool-" + s_poolCount++, new AddPoolOptions { EnableAutoscaling = true, MinAgents = 0, NumReserveAgents = 0, FleetManagers = infos.ToList(), SizeStrategy = PoolSizeStrategy.NoOp });
 			return FleetService.CreateFleetManager(pool);
 		}
 	}
diff --git a/Engine/Source/Programs/Horde/Horde.Server.Tests/Fleet/JobQueueStrategyTest.cs b/Engine/Source/Programs/Horde/Horde.Server.Tests/Fleet/JobQueueStrategyTest.cs
--- a/Engine/Source/Programs/Horde/Horde.Server.Tests/Fleet/JobQueueStrategyTest.cs
+++ b/Engine/Source/Programs/Horde/Horde.Server.Tests/Fleet/JobQueueStrategyTest.cs
@@ -98,7 +98,7 @@ namespace Horde.Server.Tests.Fleet
 		/// <param name="numBatchesReady">Num of job batches that should be in state waiting</param>
 		private async Task<(JobQueueStrategy, PoolSizeResult, IPool, List<IAgent> agents)> SetUpJobsAsync(int numBatchesRunning, int numBatchesReady, int numAgents = 8)
 		{
-			IPool pool = await PoolService.CreatePoolAsync("bogusPool1", null, true, 0, 0);
+			IPool pool = await PoolService.CreatePoolAsync("bogusPool1", new AddPoolOptions { EnableAutoscaling = true, MinAgents = 0, NumReserveAgents = 0 });
 			List<IAgent> agents = new();
 			for (int i = 0; i < numAgents; i++)
 			{
diff --git a/Engine/Source/Programs/Horde/Horde.Server.Tests/Fleet/LeaseUtilizationStrategyTest.cs b/Engine/Source/Programs/Horde/Horde.Server.Tests/Fleet/LeaseUtilizationStrategyTest.cs
--- a/Engine/Source/Programs/Horde/Horde.Server.Tests/Fleet/LeaseUtilizationStrategyTest.cs
+++ b/Engine/Source/Programs/Horde/Horde.Server.Tests/Fleet/LeaseUtilizationStrategyTest.cs
@@ -38,7 +38,7 @@ namespace Horde.Server.Tests.Fleet
 		[TestMethod]
 		public async Task UtilizationFull()
 		{
-			IPool pool = await PoolService.CreatePoolAsync("AutoscalePool1", null, true, 0, 0);
+			IPool pool = await PoolService.CreatePoolAsync("AutoscalePool1", new AddPoolOptions { EnableAutoscaling = true, MinAgents = 0, NumReserveAgents = 0 });
 			CreateAgents(pool);
 			
 			await AddPlaceholderLease(_agent1, pool, Clock.UtcNow - TimeSpan.FromMinutes(120), TimeSpan.FromMinutes(120));
@@ -51,7 +51,7 @@ namespace Horde.Server.Tests.Fleet
 		[TestMethod]
 		public async Task UtilizationHalf()
 		{
-			IPool pool = await PoolService.CreatePoolAsync("AutoscalePool1", null, true, 0, 0);
+			IPool pool = await PoolService.CreatePoolAsync("AutoscalePool1", new AddPoolOptions { EnableAutoscaling = true, MinAgents = 0, NumReserveAgents = 0 });
 			CreateAgents(pool);
 			
 			await AddPlaceholderLease(_agent1, pool, Clock.UtcNow - TimeSpan.FromMinutes(120), TimeSpan.FromMinutes(120));
@@ -62,7 +62,7 @@ namespace Horde.Server.Tests.Fleet
 		[TestMethod]
 		public async Task UtilizationZero()
 		{
-			IPool pool = await PoolService.CreatePoolAsync("AutoscalePool1", null, true, 0, 0);
+			IPool pool = await PoolService.CreatePoolAsync("AutoscalePool1", new AddPoolOptions { EnableAutoscaling = true, MinAgents = 0, NumReserveAgents = 0 });
 			CreateAgents(pool);
 
 			await AssertPoolSizeAsync(pool, 0);
@@ -71,7 +71,7 @@ namespace Horde.Server.Tests.Fleet
 		[TestMethod]
 		public async Task ReserveAgents()
 		{
-			IPool pool = await PoolService.CreatePoolAsync("AutoscalePool1", null, true, 0, 5);
+			IPool pool = await PoolService.CreatePoolAsync("AutoscalePool1", new AddPoolOptions { EnableAutoscaling = true, MinAgents = 0, NumReserveAgents = 5 });
 			CreateAgents(pool);
 			
 			await AddPlaceholderLease(_agent1, pool, Clock.UtcNow - TimeSpan.FromMinutes(120), TimeSpan.FromMinutes(120));
@@ -86,7 +86,7 @@ namespace Horde.Server.Tests.Fleet
 		[TestMethod]
 		public async Task MinAgents()
 		{
-			IPool pool = await PoolService.CreatePoolAsync("AutoscalePool1", null, true, 2, 0);
+			IPool pool = await PoolService.CreatePoolAsync("AutoscalePool1", new AddPoolOptions { EnableAutoscaling = true, MinAgents = 2, NumReserveAgents = 0 });
 			CreateAgents(pool);
 			
 			// Even with no utilization, expect at least the min number of agents
diff --git a/Engine/Source/Programs/Horde/Horde.Server.Tests/JobTaskSourceTests.cs b/Engine/Source/Programs/Horde/Horde.Server.Tests/JobTaskSourceTests.cs
--- a/Engine/Source/Programs/Horde/Horde.Server.Tests/JobTaskSourceTests.cs
+++ b/Engine/Source/Programs/Horde/Horde.Server.Tests/JobTaskSourceTests.cs
@@ -160,7 +160,7 @@ namespace Horde.Server.Tests
 		private async Task<Fixture> SetupPoolWithAgentAsync(bool isPoolAutoScaled, bool shouldCreateAgent, bool isAgentEnabled)
 		{
 			Fixture fixture = await CreateFixtureAsync();
-			IPool pool = await PoolService.CreatePoolAsync(Fixture.PoolName, null, isPoolAutoScaled, 0, 0);
+			IPool pool = await PoolService.CreatePoolAsync(Fixture.PoolName, new AddPoolOptions { EnableAutoscaling = isPoolAutoScaled, MinAgents = 0, NumReserveAgents = 0 });
 
 			if (shouldCreateAgent)
 			{
diff --git a/Engine/Source/Programs/Horde/Horde.Server.Tests/NotificationServiceTests.cs b/Engine/Source/Programs/Horde/Horde.Server.Tests/NotificationServiceTests.cs
--- a/Engine/Source/Programs/Horde/Horde.Server.Tests/NotificationServiceTests.cs
+++ b/Engine/Source/Programs/Horde/Horde.Server.Tests/NotificationServiceTests.cs
@@ -103,7 +103,7 @@ namespace Horde.Server.Tests
 			await service._ticker.StartAsync();
 
 			Fixture fixture = await CreateFixtureAsync();
-			IPool pool = await PoolService.CreatePoolAsync("BogusPool", properties: new Dictionary<string, string>());
+			IPool pool = await PoolService.CreatePoolAsync("BogusPool", new AddPoolOptions { Properties = new Dictionary<string, string>() });
 
 			Assert.AreEqual(0, fakeSink.JobScheduledNotifications.Count);
 			service.NotifyJobScheduled(pool, false, fixture.Job1, fixture.Graph, SubResourceId.Random());
@@ -115,7 +115,7 @@ namespace Horde.Server.Tests
 			Assert.AreEqual(2, fakeSink.JobScheduledNotifications.Count);
 			Assert.AreEqual(1, fakeSink.JobScheduledCallCount);
 		}
-		
+
 		[TestMethod]
 		public async Task JobScheduledNotificationsAreDeduplicated()
 		{
@@ -123,7 +123,7 @@ namespace Horde.Server.Tests
 			NotificationService service = (NotificationService)ServiceProvider.GetRequiredService<INotificationService>();
 			await service._ticker.StartAsync();
 			Fixture fixture = await CreateFixtureAsync();
-			IPool pool = await PoolService.CreatePoolAsync("BogusPool", properties: new Dictionary<string, string>());
+			IPool pool = await PoolService.CreatePoolAsync("BogusPool", new AddPoolOptions { Properties = new Dictionary<string, string>()});
 
 			service.NotifyJobScheduled(pool, false, fixture.Job1, fixture.Graph, SubResourceId.Random());
 			service.NotifyJobScheduled(pool, false, fixture.Job1, fixture.Graph, SubResourceId.Random());
diff --git a/Engine/Source/Programs/Horde/Horde.Server.Tests/PoolServiceTests.cs b/Engine/Source/Programs/Horde/Horde.Server.Tests/PoolServiceTests.cs
--- a/Engine/Source/Programs/Horde/Horde.Server.Tests/PoolServiceTests.cs
+++ b/Engine/Source/Programs/Horde/Horde.Server.Tests/PoolServiceTests.cs
@@ -25,7 +25,7 @@ namespace Horde.Server.Tests
 
         private async Task<IPool> CreatePoolFixture(string name)
         {
-            return await PoolService.CreatePoolAsync(name, properties: _fixtureProps);
+            return await PoolService.CreatePoolAsync(name, new AddPoolOptions { Properties = _fixtureProps });
         }
 
         [TestMethod]
@@ -79,7 +79,7 @@ namespace Horde.Server.Tests
                 {"cookies", "yumyum"},
             };
 
-            IPool? updatedPool = await PoolService.UpdatePoolAsync(pool, $"update-pool-new-name-{uniqueSuffix}", newProperties: updatedProps);
+            IPool? updatedPool = await PoolService.UpdatePoolAsync(pool, new UpdatePoolOptions { Name = $"update-pool-new-name-{uniqueSuffix}", Properties = updatedProps });
 			Assert.IsNotNull(updatedPool);
             Assert.AreEqual(pool.Id, updatedPool!.Id);
             Assert.AreEqual($"update-pool-new-name-{uniqueSuffix}", updatedPool.Name);
@@ -92,7 +92,7 @@ namespace Horde.Server.Tests
         public async Task UpdatePoolCollectionTest()
         {
 	        IPool pool = await CreatePoolFixture("update-pool-2");
-	        await PoolCollection.TryUpdateAsync(pool, lastScaleUpTime: DateTime.UtcNow);
+	        await PoolCollection.TryUpdateAsync(pool, new UpdatePoolOptions { LastScaleUpTime = DateTime.UtcNow });
         }
     }
 }
\ No newline at end of file
diff --git a/Engine/Source/Programs/Horde/Horde.Server.Tests/PoolsControllerTest.cs b/Engine/Source/Programs/Horde/Horde.Server.Tests/PoolsControllerTest.cs
--- a/Engine/Source/Programs/Horde/Horde.Server.Tests/PoolsControllerTest.cs
+++ b/Engine/Source/Programs/Horde/Horde.Server.Tests/PoolsControllerTest.cs
@@ -16,7 +16,7 @@ namespace Horde.Server.Tests
         [TestMethod]
         public async Task GetPoolsTest()
         {
-	        IPool pool1 = await PoolService.CreatePoolAsync("Pool1", properties: new() { {"foo", "bar"}, {"lorem", "ipsum"} });
+	        IPool pool1 = await PoolService.CreatePoolAsync("Pool1", new AddPoolOptions { Properties = new Dictionary<string, string>() { { "foo", "bar" }, { "lorem", "ipsum" } } });
 	        ActionResult<List<object>> rawResult = await PoolsController.GetPoolsAsync();
 	        Assert.AreEqual(1, rawResult.Value!.Count);
 	        GetPoolResponse response = (rawResult.Value![0] as GetPoolResponse)!;
@@ -62,7 +62,7 @@ namespace Horde.Server.Tests
         [TestMethod]
         public async Task UpdatePoolTest()
         {
-	        IPool pool1 = await PoolService.CreatePoolAsync("Pool1", properties: new() { {"foo", "bar"}, {"lorem", "ipsum"} });
+	        IPool pool1 = await PoolService.CreatePoolAsync("Pool1", new AddPoolOptions { Properties = new Dictionary<string, string>() { { "foo", "bar" }, { "lorem", "ipsum" } } });
 
 	        await PoolsController.UpdatePoolAsync(pool1.Id.ToString(), new UpdatePoolRequest
 	        {
diff --git a/Engine/Source/Programs/Horde/Horde.Server/Agents/Fleet/FleetService.cs b/Engine/Source/Programs/Horde/Horde.Server/Agents/Fleet/FleetService.cs
--- a/Engine/Source/Programs/Horde/Horde.Server/Agents/Fleet/FleetService.cs
+++ b/Engine/Source/Programs/Horde/Horde.Server/Agents/Fleet/FleetService.cs
@@ -312,11 +312,15 @@ namespace Horde.Server.Agents.Fleet
 
 			await _poolCollection.TryUpdateAsync(
 				pool,
-				lastScaleUpTime: scaleOutTime,
-				lastScaleDownTime: scaleInTime,
-				lastScaleResult: isCooldownResult ? null : result,
-				lastAgentCount: currentAgentCount,
-				lastDesiredAgentCount: desiredAgentCount);
+				new UpdatePoolOptions
+				{
+					LastScaleUpTime = scaleOutTime,
+					LastScaleDownTime = scaleInTime,
+					LastScaleResult = isCooldownResult ? null : result,
+					LastAgentCount = currentAgentCount,
+					LastDesiredAgentCount = desiredAgentCount 
+				}
+				);
 
 			return result;
 		}
diff --git a/Engine/Source/Programs/Horde/Horde.Server/Agents/Pools/IPoolCollection.cs b/Engine/Source/Programs/Horde/Horde.Server/Agents/Pools/IPoolCollection.cs
--- a/Engine/Source/Programs/Horde/Horde.Server/Agents/Pools/IPoolCollection.cs
+++ b/Engine/Source/Programs/Horde/Horde.Server/Agents/Pools/IPoolCollection.cs
@@ -3,45 +3,11 @@
 using System;
 using System.Collections.Generic;
 using System.Threading.Tasks;
-using Amazon.Util.Internal;
 using EpicGames.Horde.Common;
 using Horde.Server.Agents.Fleet;
-using Horde.Server.Utilities;
 
 namespace Horde.Server.Agents.Pools
 {
-	/// <summary>
-	/// Exception when conflicting definitions for a pool are encountered
-	/// </summary>
-	public sealed class PoolConflictException : Exception
-	{
-		/// <summary>
-		/// Pool with conflicting definitions
-		/// </summary>
-		public PoolId PoolId { get; }
-
-		/// <summary>
-		/// The first revision string
-		/// </summary>
-		public string PrevRevision { get; }
-
-		/// <summary>
-		/// The second revision string
-		/// </summary>
-		public string NextRevision { get; }
-
-		/// <summary>
-		/// Constructor
-		/// </summary>
-		public PoolConflictException(PoolId poolId, string prevRevision, string nextRevision)
-			: base($"Duplicate definitions for {poolId} - first in {prevRevision}, now {nextRevision}")
-		{
-			PoolId = poolId;
-			PrevRevision = prevRevision;
-			NextRevision = nextRevision;
-		}
-	}
-
 	/// <summary>
 	/// Collection of pool documents
 	/// </summary>
@@ -52,38 +18,9 @@ namespace Horde.Server.Agents.Pools
 		/// </summary>
 		/// <param name="id">Unique id for the new pool</param>
 		/// <param name="name">Name of the new pool</param>
-		/// <param name="condition">Condition for agents to be included in this pool</param>
-		/// <param name="enableAutoscaling">Whether to enable autoscaling for this pool</param>
-		/// <param name="minAgents">Minimum number of agents in the pool</param>
-		/// <param name="numReserveAgents">Minimum number of idle agents to maintain</param>
-		/// <param name="conformInterval">Interval between conforms. Set to zero to disable.</param>
-		/// <param name="scaleOutCooldown">Cooldown time between scale-out events</param>
-		/// <param name="scaleInCooldown">Cooldown time between scale-in events</param>
-		/// <param name="sizeStrategies">Pool sizing strategies</param>
-		/// <param name="fleetManagers">Fleet managers</param>
-		/// <param name="sizeStrategy">Pool sizing strategy</param>
-		/// <param name="leaseUtilizationSettings">Settings for lease utilization strategy</param>
-		/// <param name="jobQueueSettings">Settings for job queue strategy</param>
-		/// <param name="computeQueueAwsMetricSettings">Settings for compute queue AWS metric strategy</param>
-		/// <param name="properties">Properties for the pool</param>
+		/// <param name="options">Options for the new pool</param>
 		/// <returns>The new pool document</returns>
-		Task<IPool> AddAsync(
-			PoolId id,
-			string name,
-			Condition? condition = null,
-			bool? enableAutoscaling = null,
-			int? minAgents = null,
-			int? numReserveAgents = null,
-			TimeSpan? conformInterval = null,
-			TimeSpan? scaleOutCooldown = null,
-			TimeSpan? scaleInCooldown = null,
-			List<PoolSizeStrategyInfo>? sizeStrategies = null,
-			List<FleetManagerInfo>? fleetManagers = null,
-			PoolSizeStrategy? sizeStrategy = null,
-			LeaseUtilizationSettings? leaseUtilizationSettings = null,
-			JobQueueSettings? jobQueueSettings = null,
-			ComputeQueueAwsMetricSettings? computeQueueAwsMetricSettings = null,
-			IEnumerable<KeyValuePair<string, string>>? properties = null);
+		Task<IPool> AddAsync(PoolId id, string name, AddPoolOptions options);
 
 		/// <summary>
 		/// Enumerates all the pools
@@ -115,57 +52,8 @@ namespace Horde.Server.Agents.Pools
 		/// Updates a pool
 		/// </summary>
 		/// <param name="pool">The pool to update</param>
-		/// <param name="newName">New name for the pool</param>
-		/// <param name="newCondition">New condition for the pool</param>
-		/// <param name="newEnableAutoscaling">New setting for whether to enable autoscaling</param>
-		/// <param name="newMinAgents">Minimum number of agents in the pool</param>
-		/// <param name="newNumReserveAgents">Minimum number of idle agents to maintain</param>
-		/// <param name="newWorkspaces">New workspaces for the pool</param>
-		/// <param name="newUseAutoSdk">New setting for whether to use autosdk</param>
-		/// <param name="newProperties">New properties for the pool</param>
-		/// <param name="conformInterval">Interval between conforms. Set to zero to disable.</param>
-		/// <param name="lastScaleUpTime">New time for last (auto) scale up</param>
-		/// <param name="lastScaleDownTime">New time for last (auto) scale down</param>
-		/// <param name="scaleOutCooldown">Cooldown time between scale-out events</param>
-		/// <param name="scaleInCooldown">Cooldown time between scale-in events</param>
-		/// <param name="shutdownIfDisabledGracePeriod">Time to wait before shutting down a disabled agent</param>
-		/// <param name="lastScaleResult">Result from last scaling in/out attempt</param>
-		/// <param name="lastAgentCount">Last calculated agent count</param>
-		/// <param name="lastDesiredAgentCount">Last calculated desired agent count</param>
-		/// <param name="sizeStrategy">Pool sizing strategy</param>
-		/// <param name="newSizeStrategies">List of pool sizing strategies</param>
-		/// <param name="newFleetManagers">List of fleet managers</param>
-		/// <param name="leaseUtilizationSettings">Settings for lease utilization strategy</param>
-		/// <param name="jobQueueSettings">Settings for job queue strategy</param>
-		/// <param name="computeQueueAwsMetricSettings">Settings for compute queue AWS metric strategy</param>
-		/// <param name="useDefaultStrategy">Whether to use the default strategy</param>
-		/// <returns>Async task</returns>
-		Task<IPool?> TryUpdateAsync(
-			IPool pool,
-			string? newName = null,
-			Condition? newCondition = null,
-			bool? newEnableAutoscaling = null,
-			int? newMinAgents = null,
-			int? newNumReserveAgents = null,
-			List<AgentWorkspace>? newWorkspaces = null,
-			bool? newUseAutoSdk = null,
-			Dictionary<string, string?>? newProperties = null,
-			TimeSpan? conformInterval = null,
-			DateTime? lastScaleUpTime = null,
-			DateTime? lastScaleDownTime = null,
-			TimeSpan? scaleOutCooldown = null,
-			TimeSpan? scaleInCooldown = null,
-			TimeSpan? shutdownIfDisabledGracePeriod = null,
-			ScaleResult? lastScaleResult = null,
-			int? lastAgentCount = null,
-			int? lastDesiredAgentCount = null,
-			PoolSizeStrategy? sizeStrategy = null,
-			List<PoolSizeStrategyInfo>? newSizeStrategies = null,
-			List<FleetManagerInfo>? newFleetManagers = null,
-			LeaseUtilizationSettings? leaseUtilizationSettings = null,
-			JobQueueSettings? jobQueueSettings = null,
-			ComputeQueueAwsMetricSettings? computeQueueAwsMetricSettings = null,
-			bool? useDefaultStrategy = null);
+		/// <param name="options">Options for the update</param>
+		Task<IPool?> TryUpdateAsync(IPool pool, UpdatePoolOptions options);
 
 		/// <summary>
 		/// Updates the list of pools from a config file
@@ -173,4 +61,238 @@ namespace Horde.Server.Agents.Pools
 		/// <param name="poolConfigs">Configuration for the pools, and revision string for the config file containing them</param>
 		Task ConfigureAsync(IReadOnlyList<(PoolConfig Config, string Revision)> poolConfigs);
 	}
+
+	/// <summary>
+	/// Settings for creating a new pool
+	/// </summary>
+	public class AddPoolOptions
+	{
+		/// <summary>
+		/// Condition for agents to be included in this pool
+		/// </summary>
+		public Condition? Condition { get; set; }
+
+		/// <summary>
+		/// Whether to enable autoscaling for this pool
+		/// </summary>
+		public bool? EnableAutoscaling { get; set; }
+
+		/// <summary>
+		/// Minimum number of agents in the pool
+		/// </summary>
+		public int? MinAgents { get; set; }
+
+		/// <summary>
+		/// Minimum number of idle agents to maintain
+		/// </summary>
+		public int? NumReserveAgents { get; set; }
+
+		/// <summary>
+		/// Interval between conforms. Set to zero to disable.
+		/// </summary>
+		public TimeSpan? ConformInterval { get; set; }
+
+		/// <summary>
+		/// Cooldown time between scale-out events
+		/// </summary>
+		public TimeSpan? ScaleOutCooldown { get; set; }
+
+		/// <summary>
+		/// Cooldown time between scale-in events
+		/// </summary>
+		public TimeSpan? ScaleInCooldown { get; set; }
+
+		/// <summary>
+		/// Pool sizing strategies
+		/// </summary>
+		public List<PoolSizeStrategyInfo>? SizeStrategies { get; set; }
+
+		/// <summary>
+		/// Fleet managers
+		/// </summary>
+		public List<FleetManagerInfo>? FleetManagers { get; set; }
+
+		/// <summary>
+		/// Pool sizing strategy
+		/// </summary>
+		public PoolSizeStrategy? SizeStrategy { get; set; }
+
+		/// <summary>
+		/// Settings for lease utilization strategy
+		/// </summary>
+		public LeaseUtilizationSettings? LeaseUtilizationSettings { get; set; }
+
+		/// <summary>
+		/// Settings for job queue strategy
+		/// </summary>
+		public JobQueueSettings? JobQueueSettings { get; set; }
+
+		/// <summary>
+		/// Settings for compute queue AWS metric strategy
+		/// </summary>
+		public ComputeQueueAwsMetricSettings? ComputeQueueAwsMetricSettings { get; set; }
+
+		/// <summary>
+		/// Properties for the pool
+		/// </summary>
+		public IEnumerable<KeyValuePair<string, string>>? Properties { get; set; }
+	}
+
+	/// <summary>
+	/// Settings for updating a pool state
+	/// </summary>
+	public class UpdatePoolOptions
+	{
+		/// <summary>
+		/// New name for the pool
+		/// </summary>
+		public string? Name { get; set; }
+
+		/// <summary>
+		/// New condition for the pool
+		/// </summary>
+		public Condition? Condition { get; set; }
+
+		/// <summary>
+		/// New setting for whether to enable autoscaling
+		/// </summary>
+		public bool? EnableAutoscaling { get; set; }
+
+		/// <summary>
+		/// Minimum number of agents in the pool
+		/// </summary>
+		public int? MinAgents { get; set; }
+
+		/// <summary>
+		/// Minimum number of idle agents to maintain
+		/// </summary>
+		public int? NumReserveAgents { get; set; }
+
+		/// <summary>
+		/// New workspaces for the pool
+		/// </summary>
+		public List<AgentWorkspace>? Workspaces { get; set; }
+
+		/// <summary>
+		/// New setting for whether to use autosdk
+		/// </summary>
+		public bool? UseAutoSdk { get; set; }
+
+		/// <summary>
+		/// New properties for the pool
+		/// </summary>
+		public Dictionary<string, string?>? Properties { get; set; }
+
+		/// <summary>
+		/// Interval between conforms. Set to zero to disable.
+		/// </summary>
+		public TimeSpan? ConformInterval { get; set; }
+
+		/// <summary>
+		/// New time for last (auto) scale up
+		/// </summary>
+		public DateTime? LastScaleUpTime { get; set; }
+
+		/// <summary>
+		/// New time for last (auto) scale down
+		/// </summary>
+		public DateTime? LastScaleDownTime { get; set; }
+
+		/// <summary>
+		/// Cooldown time between scale-out events
+		/// </summary>
+		public TimeSpan? ScaleOutCooldown { get; set; }
+
+		/// <summary>
+		/// Cooldown time between scale-in events
+		/// </summary>
+		public TimeSpan? ScaleInCooldown { get; set; }
+
+		/// <summary>
+		/// Time to wait before shutting down a disabled agent
+		/// </summary>
+		public TimeSpan? ShutdownIfDisabledGracePeriod { get; set; }
+
+		/// <summary>
+		/// Result from last scaling in/out attempt
+		/// </summary>
+		public ScaleResult? LastScaleResult { get; set; }
+
+		/// <summary>
+		/// Last calculated agent count
+		/// </summary>
+		public int? LastAgentCount { get; set; }
+
+		/// <summary>
+		/// Last calculated desired agent count
+		/// </summary>
+		public int? LastDesiredAgentCount { get; set; }
+
+		/// <summary>
+		/// Pool sizing strategy
+		/// </summary>
+		public PoolSizeStrategy? SizeStrategy { get; set; }
+
+		/// <summary>
+		/// List of pool sizing strategies
+		/// </summary>
+		public List<PoolSizeStrategyInfo>? SizeStrategies { get; set; }
+
+		/// <summary>
+		/// List of fleet managers
+		/// </summary>
+		public List<FleetManagerInfo>? FleetManagers { get; set; }
+
+		/// <summary>
+		/// Settings for lease utilization strategy
+		/// </summary>
+		public LeaseUtilizationSettings? LeaseUtilizationSettings { get; set; }
+
+		/// <summary>
+		/// Settings for job queue strategy
+		/// </summary>
+		public JobQueueSettings? JobQueueSettings { get; set; }
+
+		/// <summary>
+		/// Settings for compute queue AWS metric strategy
+		/// </summary>
+		public ComputeQueueAwsMetricSettings? ComputeQueueAwsMetricSettings { get; set; }
+
+		/// <summary>
+		/// Whether to use the default strategy
+		/// </summary>
+		public bool? UseDefaultStrategy { get; set; }
+	}
+
+	/// <summary>
+	/// Exception when conflicting definitions for a pool are encountered
+	/// </summary>
+	public sealed class PoolConflictException : Exception
+	{
+		/// <summary>
+		/// Pool with conflicting definitions
+		/// </summary>
+		public PoolId PoolId { get; }
+
+		/// <summary>
+		/// The first revision string
+		/// </summary>
+		public string PrevRevision { get; }
+
+		/// <summary>
+		/// The second revision string
+		/// </summary>
+		public string NextRevision { get; }
+
+		/// <summary>
+		/// Constructor
+		/// </summary>
+		public PoolConflictException(PoolId poolId, string prevRevision, string nextRevision)
+			: base($"Duplicate definitions for {poolId} - first in {prevRevision}, now {nextRevision}")
+		{
+			PoolId = poolId;
+			PrevRevision = prevRevision;
+			NextRevision = nextRevision;
+		}
+	}
 }
diff --git a/Engine/Source/Programs/Horde/Horde.Server/Agents/Pools/PoolCollection.cs b/Engine/Source/Programs/Horde/Horde.Server/Agents/Pools/PoolCollection.cs
--- a/Engine/Source/Programs/Horde/Horde.Server/Agents/Pools/PoolCollection.cs
+++ b/Engine/Source/Programs/Horde/Horde.Server/Agents/Pools/PoolCollection.cs
@@ -160,48 +160,32 @@ namespace Horde.Server.Agents.Pools
 		}
 
 		/// <inheritdoc/>
-		public async Task<IPool> AddAsync(
-			PoolId id,
-			string name,
-			Condition? condition,
-			bool? enableAutoscaling,
-			int? minAgents,
-			int? numReserveAgents,
-			TimeSpan? conformInterval,
-			TimeSpan? scaleOutCooldown,
-			TimeSpan? scaleInCooldown,
-			List<PoolSizeStrategyInfo>? sizeStrategies,
-			List<FleetManagerInfo>? fleetManagers,
-			PoolSizeStrategy? sizeStrategy,
-			LeaseUtilizationSettings? leaseUtilizationSettings,
-			JobQueueSettings? jobQueueSettings,
-			ComputeQueueAwsMetricSettings? computeQueueAwsMetricSettings,
-			IEnumerable<KeyValuePair<string, string>>? properties)
+		public async Task<IPool> AddAsync(PoolId id, string name, AddPoolOptions options)
 		{
 			PoolDocument pool = new PoolDocument();
 			pool.Id = id;
 			pool.Name = name;
-			pool.Condition = condition;
-			if (enableAutoscaling != null)
+			pool.Condition = options.Condition;
+			if (options.EnableAutoscaling != null)
 			{
-				pool.EnableAutoscaling = enableAutoscaling.Value;
+				pool.EnableAutoscaling = options.EnableAutoscaling.Value;
 			}
-			pool.MinAgents = minAgents;
-			pool.NumReserveAgents = numReserveAgents;
-			if (properties != null)
+			pool.MinAgents = options.MinAgents;
+			pool.NumReserveAgents = options.NumReserveAgents;
+			if (options.Properties != null)
 			{
-				pool.Properties = new Dictionary<string, string>(properties);
+				pool.Properties = new Dictionary<string, string>(options.Properties);
 			}
 
-			pool.ConformInterval = conformInterval;
-			pool.ScaleOutCooldown = scaleOutCooldown;
-			pool.ScaleInCooldown = scaleInCooldown;
-			pool.SizeStrategy = sizeStrategy;
-			pool.LeaseUtilizationSettings = leaseUtilizationSettings;
-			pool.JobQueueSettings = jobQueueSettings;
-			pool.ComputeQueueAwsMetricSettings = computeQueueAwsMetricSettings;
-			if (sizeStrategies != null) { pool.SizeStrategies = sizeStrategies; }
-			if (fleetManagers != null) { pool.FleetManagers = fleetManagers; }
+			pool.ConformInterval = options.ConformInterval;
+			pool.ScaleOutCooldown = options.ScaleOutCooldown;
+			pool.ScaleInCooldown = options.ScaleInCooldown;
+			pool.SizeStrategy = options.SizeStrategy;
+			pool.LeaseUtilizationSettings = options.LeaseUtilizationSettings;
+			pool.JobQueueSettings = options.JobQueueSettings;
+			pool.ComputeQueueAwsMetricSettings = options.ComputeQueueAwsMetricSettings;
+			if (options.SizeStrategies != null) { pool.SizeStrategies = options.SizeStrategies; }
+			if (options.FleetManagers != null) { pool.FleetManagers = options.FleetManagers; }
 			
 			await _pools.InsertOneAsync(pool);
 			return pool;
@@ -277,164 +261,139 @@ namespace Horde.Server.Agents.Pools
 		}
 
 		/// <inheritdoc/>
-		public Task<IPool?> TryUpdateAsync(
-			IPool pool,
-			string? newName,
-			Condition? newCondition,
-			bool? newEnableAutoscaling,
-			int? newMinAgents,
-			int? newNumReserveAgents,
-			List<AgentWorkspace>? newWorkspaces,
-			bool? newUseAutoSdk,
-			Dictionary<string, string?>? newProperties,
-			TimeSpan? conformInterval,
-			DateTime? lastScaleUpTime,
-			DateTime? lastScaleDownTime,
-			TimeSpan? scaleOutCooldown,
-			TimeSpan? scaleInCooldown, 
-			TimeSpan? shutdownIfDisabledGracePeriod, 
-			ScaleResult? lastScaleResult, 
-			int? lastAgentCount, 
-			int? lastDesiredAgentCount, 
-			PoolSizeStrategy? sizeStrategy,
-			List<PoolSizeStrategyInfo>? newSizeStrategies,
-			List<FleetManagerInfo>? newFleetManagers,
-			LeaseUtilizationSettings? leaseUtilizationSettings,
-			JobQueueSettings? jobQueueSettings, 
-			ComputeQueueAwsMetricSettings? computeQueueAwsMetricSettings, 
-			bool? useDefaultStrategy)
+		public Task<IPool?> TryUpdateAsync(IPool pool, UpdatePoolOptions options)
 		{
 			TransactionBuilder<PoolDocument> transaction = new TransactionBuilder<PoolDocument>();
-			if (newName != null)
+			if (options.Name != null)
 			{
-				transaction.Set(x => x.Name, newName);
+				transaction.Set(x => x.Name, options.Name);
 			}
-			if (newCondition != null)
+			if (options.Condition != null)
 			{
-				if (newCondition.IsEmpty())
+				if (options.Condition.IsEmpty())
 				{
 					transaction.Unset(x => x.Condition!);
 				}
 				else
 				{
-					transaction.Set(x => x.Condition, newCondition);
+					transaction.Set(x => x.Condition, options.Condition);
 				}
 			}
-			if (newEnableAutoscaling != null)
+			if (options.EnableAutoscaling != null)
 			{
-				if (newEnableAutoscaling.Value)
+				if (options.EnableAutoscaling.Value)
 				{
 					transaction.Unset(x => x.EnableAutoscaling);
 				}
 				else
 				{
-					transaction.Set(x => x.EnableAutoscaling, newEnableAutoscaling.Value);
+					transaction.Set(x => x.EnableAutoscaling, options.EnableAutoscaling.Value);
 				}
 			}
-			if (newMinAgents != null)
+			if (options.MinAgents != null)
 			{
-				if (newMinAgents.Value < 0)
+				if (options.MinAgents.Value < 0)
 				{
 					transaction.Unset(x => x.MinAgents!);
 				}
 				else
 				{
-					transaction.Set(x => x.MinAgents, newMinAgents.Value);
+					transaction.Set(x => x.MinAgents, options.MinAgents.Value);
 				}
 			}
-			if (newNumReserveAgents != null)
+			if (options.NumReserveAgents != null)
 			{
-				if (newNumReserveAgents.Value < 0)
+				if (options.NumReserveAgents.Value < 0)
 				{
 					transaction.Unset(x => x.NumReserveAgents!);
 				}
 				else
 				{
-					transaction.Set(x => x.NumReserveAgents, newNumReserveAgents.Value);
+					transaction.Set(x => x.NumReserveAgents, options.NumReserveAgents.Value);
 				}
 			}
-			if (newWorkspaces != null)
+			if (options.Workspaces != null)
 			{
-				transaction.Set(x => x.Workspaces, newWorkspaces);
+				transaction.Set(x => x.Workspaces, options.Workspaces);
 			}
-			if (newUseAutoSdk != null)
+			if (options.UseAutoSdk != null)
 			{
-				transaction.Set(x => x.UseAutoSdk, newUseAutoSdk.Value);
+				transaction.Set(x => x.UseAutoSdk, options.UseAutoSdk.Value);
 			}
-			if (newProperties != null)
+			if (options.Properties != null)
 			{
-				transaction.UpdateDictionary(x => x.Properties, newProperties);
+				transaction.UpdateDictionary(x => x.Properties, options.Properties);
 			}
-			if (conformInterval != null)
+			if (options.ConformInterval != null)
 			{
-				if (conformInterval.Value < TimeSpan.Zero)
+				if (options.ConformInterval.Value < TimeSpan.Zero)
 				{
 					transaction.Unset(x => x.ConformInterval!);
 				}
 				else
 				{
-					transaction.Set(x => x.ConformInterval, conformInterval);
+					transaction.Set(x => x.ConformInterval, options.ConformInterval);
 				}
 			}
-			if (lastScaleUpTime != null)
+			if (options.LastScaleUpTime != null)
 			{
-				transaction.Set(x => x.LastScaleUpTime, lastScaleUpTime);
+				transaction.Set(x => x.LastScaleUpTime, options.LastScaleUpTime);
 			}
-			if (lastScaleDownTime != null)
+			if (options.LastScaleDownTime != null)
 			{
-				transaction.Set(x => x.LastScaleDownTime, lastScaleDownTime);
+				transaction.Set(x => x.LastScaleDownTime, options.LastScaleDownTime);
 			}
-			if (scaleOutCooldown != null)
+			if (options.ScaleOutCooldown != null)
 			{
-				transaction.Set(x => x.ScaleOutCooldown, scaleOutCooldown);
+				transaction.Set(x => x.ScaleOutCooldown, options.ScaleOutCooldown);
 			}
-			if (scaleInCooldown != null)
+			if (options.ScaleInCooldown != null)
 			{
-				transaction.Set(x => x.ScaleInCooldown, scaleInCooldown);
+				transaction.Set(x => x.ScaleInCooldown, options.ScaleInCooldown);
 			}
-			if (shutdownIfDisabledGracePeriod != null)
+			if (options.ShutdownIfDisabledGracePeriod != null)
 			{
-				transaction.Set(x => x.ShutdownIfDisabledGracePeriod, shutdownIfDisabledGracePeriod);
+				transaction.Set(x => x.ShutdownIfDisabledGracePeriod, options.ShutdownIfDisabledGracePeriod);
 			}
-			if (lastScaleResult != null)
+			if (options.LastScaleResult != null)
 			{
-				transaction.Set(x => x.LastScaleResult, lastScaleResult);
+				transaction.Set(x => x.LastScaleResult, options.LastScaleResult);
 			}
-			if (lastAgentCount != null)
+			if (options.LastAgentCount != null)
 			{
-				transaction.Set(x => x.LastAgentCount, lastAgentCount);
+				transaction.Set(x => x.LastAgentCount, options.LastAgentCount);
 			}
-			if (lastDesiredAgentCount != null)
+			if (options.LastDesiredAgentCount != null)
 			{
-				transaction.Set(x => x.LastDesiredAgentCount, lastDesiredAgentCount);
+				transaction.Set(x => x.LastDesiredAgentCount, options.LastDesiredAgentCount);
 			}
-			if (sizeStrategy != null)
+			if (options.SizeStrategy != null)
 			{
-				transaction.Set(x => x.SizeStrategy, sizeStrategy);				
+				transaction.Set(x => x.SizeStrategy, options.SizeStrategy);				
 			}
-			else if (useDefaultStrategy != null && useDefaultStrategy.Value)
+			else if (options.UseDefaultStrategy != null && options.UseDefaultStrategy.Value)
 			{
 				transaction.Set(x => x.SizeStrategy, null);
 			}
-			if (newSizeStrategies != null)
+			if (options.SizeStrategies != null)
 			{
-				transaction.Set(x => x.SizeStrategies, newSizeStrategies);				
+				transaction.Set(x => x.SizeStrategies, options.SizeStrategies);				
 			}
-			if (newFleetManagers != null)
+			if (options.FleetManagers != null)
 			{
-				transaction.Set(x => x.FleetManagers, newFleetManagers);				
+				transaction.Set(x => x.FleetManagers, options.FleetManagers);				
 			}
-			if (leaseUtilizationSettings != null)
+			if (options.LeaseUtilizationSettings != null)
 			{
-				transaction.Set(x => x.LeaseUtilizationSettings, leaseUtilizationSettings);
+				transaction.Set(x => x.LeaseUtilizationSettings, options.LeaseUtilizationSettings);
 			}
-			if (jobQueueSettings != null)
+			if (options.JobQueueSettings != null)
 			{
-				transaction.Set(x => x.JobQueueSettings, jobQueueSettings);
+				transaction.Set(x => x.JobQueueSettings, options.JobQueueSettings);
 			}
-			if (computeQueueAwsMetricSettings != null)
+			if (options.ComputeQueueAwsMetricSettings != null)
 			{
-				transaction.Set(x => x.ComputeQueueAwsMetricSettings, computeQueueAwsMetricSettings);
+				transaction.Set(x => x.ComputeQueueAwsMetricSettings, options.ComputeQueueAwsMetricSettings);
 			}
 			return TryUpdateAsync(pool, transaction);
 		}
diff --git a/Engine/Source/Programs/Horde/Horde.Server/Agents/Pools/PoolService.cs b/Engine/Source/Programs/Horde/Horde.Server/Agents/Pools/PoolService.cs
--- a/Engine/Source/Programs/Horde/Horde.Server/Agents/Pools/PoolService.cs
+++ b/Engine/Source/Programs/Horde/Horde.Server/Agents/Pools/PoolService.cs
@@ -56,55 +56,11 @@ namespace Horde.Server.Agents.Pools
 		/// Creates a new pool
 		/// </summary>
 		/// <param name="name">Name of the new pool</param>
-		/// <param name="condition">Condition for agents to be automatically included in this pool</param>
-		/// <param name="enableAutoscaling">Whether to enable autoscaling for this pool</param>
-		/// <param name="minAgents">Minimum number of agents in the pool</param>
-		/// <param name="numReserveAgents">Minimum number of idle agents to maintain</param>
-		/// <param name="conformInterval">Interval between conforms. Set to zero to disable.</param>
-		/// <param name="scaleOutCooldown">Cooldown time between scale-out events</param>
-		/// <param name="scaleInCooldown">Cooldown time between scale-in events</param>
-		/// <param name="sizeStrategies">Pool sizing strategies</param>
-		/// <param name="fleetManagers">Fleet managers</param>
-		/// <param name="sizeStrategy">Pool sizing strategy</param>
-		/// <param name="leaseUtilizationSettings">Settings for lease utilization strategy</param>
-		/// <param name="jobQueueSettings">Settings for job queue strategy</param>
-		/// <param name="computeQueueAwsMetricSettings">Settings for compute queue AWS metric strategy</param>
-		/// <param name="properties">Properties for the new pool</param>
+		/// <param name="options">Options for the new pool</param>
 		/// <returns>The new pool document</returns>
-		public Task<IPool> CreatePoolAsync(
-			string name,
-			Condition? condition = null,
-			bool? enableAutoscaling = null,
-			int? minAgents = null,
-			int? numReserveAgents = null,
-			TimeSpan? conformInterval = null,
-			TimeSpan? scaleOutCooldown = null,
-			TimeSpan? scaleInCooldown = null,
-			List<PoolSizeStrategyInfo>? sizeStrategies = null,
-			List<FleetManagerInfo>? fleetManagers = null,
-			PoolSizeStrategy? sizeStrategy = null,
-			LeaseUtilizationSettings? leaseUtilizationSettings = null,
-			JobQueueSettings? jobQueueSettings = null,
-			ComputeQueueAwsMetricSettings? computeQueueAwsMetricSettings = null,
-			Dictionary<string, string>? properties = null)
+		public Task<IPool> CreatePoolAsync(string name, AddPoolOptions options)
 		{
-			return _pools.AddAsync(
-				new PoolId(StringId.Sanitize(name)),
-				name,
-				condition,
-				enableAutoscaling,
-				minAgents,
-				numReserveAgents,
-				conformInterval,
-				scaleOutCooldown,
-				scaleInCooldown,
-				sizeStrategies,
-				fleetManagers,
-				sizeStrategy,
-				leaseUtilizationSettings,
-				jobQueueSettings,
-				computeQueueAwsMetricSettings,
-				properties);
+			return _pools.AddAsync(new PoolId(StringId.Sanitize(name)), name, options);
 		}
 
 		/// <summary>
@@ -121,63 +77,13 @@ namespace Horde.Server.Agents.Pools
 		/// Updates an existing pool
 		/// </summary>
 		/// <param name="pool">The pool to update</param>
-		/// <param name="newName">The new name for the pool</param>
-		/// <param name="newCondition">New requirements for the pool</param>
-		/// <param name="newEnableAutoscaling">Whether to enable autoscaling</param>
-		/// <param name="newMinAgents">Minimum number of agents in the pool</param>
-		/// <param name="newNumReserveAgents">Minimum number of idle agents to maintain</param>
-		/// <param name="newProperties">Properties on the pool to update. Any properties with a value of null will be removed.</param>
-		/// <param name="conformInterval">Interval between conforms. Set to zero to disable.</param>
-		/// <param name="scaleOutCooldown">Cooldown time between scale-out events</param>
-		/// <param name="scaleInCooldown">Cooldown time between scale-in events</param>
-		/// <param name="sizeStrategy">New pool sizing strategy for the pool</param>
-		/// <param name="newSizeStrategies">New pool sizing strategies for the pool</param>
-		/// <param name="newFleetManagers">New fleet managers for the pool</param>
-		/// <param name="leaseUtilizationSettings">Settings for lease utilization-based strategy</param>
-		/// <param name="jobQueueSettings">Settings for job queue-based strategy</param>
-		/// <param name="computeQueueAwsMetricSettings">Settings for compute queue AWS metric strategy</param>
-		/// <param name="useDefaultStrategy">Whether to use the default strategy</param>
+		/// <param name="options">Options for the update</param>
 		/// <returns>Async task object</returns>
-		public async Task<IPool?> UpdatePoolAsync(
-			IPool? pool,
-			string? newName = null,
-			Condition? newCondition = null,
-			bool? newEnableAutoscaling = null,
-			int? newMinAgents = null,
-			int? newNumReserveAgents = null,
-			Dictionary<string, string?>? newProperties = null,
-			TimeSpan? conformInterval = null,
-			TimeSpan? scaleOutCooldown = null,
-			TimeSpan? scaleInCooldown = null,
-			PoolSizeStrategy? sizeStrategy = null,
-			List<PoolSizeStrategyInfo>? newSizeStrategies = null,
-			List<FleetManagerInfo>? newFleetManagers = null,
-			LeaseUtilizationSettings? leaseUtilizationSettings = null,
-			JobQueueSettings? jobQueueSettings = null,
-			ComputeQueueAwsMetricSettings? computeQueueAwsMetricSettings = null,
-			bool? useDefaultStrategy = null)
+		public async Task<IPool?> UpdatePoolAsync(IPool? pool, UpdatePoolOptions options)
 		{
 			for (; pool != null; pool = await _pools.GetAsync(pool.Id))
 			{
-				IPool? newPool = await _pools.TryUpdateAsync(
-					pool,
-					newName,
-					newCondition,
-					newEnableAutoscaling,
-					newMinAgents,
-					newNumReserveAgents,
-					newProperties: newProperties,
-					conformInterval: conformInterval,
-					scaleOutCooldown: scaleOutCooldown,
-					scaleInCooldown: scaleInCooldown,
-					sizeStrategy: sizeStrategy,
-					newSizeStrategies: newSizeStrategies,
-					newFleetManagers: newFleetManagers,
-					leaseUtilizationSettings: leaseUtilizationSettings,
-					jobQueueSettings: jobQueueSettings,
-					computeQueueAwsMetricSettings: computeQueueAwsMetricSettings,
-					useDefaultStrategy: useDefaultStrategy);
-				
+				IPool? newPool = await _pools.TryUpdateAsync(pool, options);
 				if (newPool != null)
 				{
 					return newPool;
diff --git a/Engine/Source/Programs/Horde/Horde.Server/Agents/Pools/PoolUpdateService.cs b/Engine/Source/Programs/Horde/Horde.Server/Agents/Pools/PoolUpdateService.cs
--- a/Engine/Source/Programs/Horde/Horde.Server/Agents/Pools/PoolUpdateService.cs
+++ b/Engine/Source/Programs/Horde/Horde.Server/Agents/Pools/PoolUpdateService.cs
@@ -197,7 +197,7 @@ namespace Horde.Server.Agents.Pools
 					{
 						_logger.LogInformation("New workspaces for pool {Pool}:{Workspaces}", currentPool.Id, String.Join("", newWorkspaces.Select(x => $"\n  Identifier=\"{x.Identifier}\", Stream={x.Stream}")));
 
-						IPool? result = await _pools.TryUpdateAsync(currentPool, newWorkspaces: newWorkspaces, newUseAutoSdk: useAutoSdk);
+						IPool? result = await _pools.TryUpdateAsync(currentPool, new UpdatePoolOptions { Workspaces = newWorkspaces, UseAutoSdk = useAutoSdk });
 						if (result == null)
 						{
 							_logger.LogInformation("Pool modified; will retry");
diff --git a/Engine/Source/Programs/Horde/Horde.Server/Agents/Pools/PoolsController.cs b/Engine/Source/Programs/Horde/Horde.Server/Agents/Pools/PoolsController.cs
--- a/Engine/Source/Programs/Horde/Horde.Server/Agents/Pools/PoolsController.cs
+++ b/Engine/Source/Programs/Horde/Horde.Server/Agents/Pools/PoolsController.cs
@@ -58,9 +58,23 @@ namespace Horde.Server.Agents.Pools
 			TimeSpan? scaleOutCooldown = create.ScaleOutCooldown == null ? null : TimeSpan.FromSeconds(create.ScaleOutCooldown.Value);
 			TimeSpan? scaleInCooldown = create.ScaleInCooldown == null ? null : TimeSpan.FromSeconds(create.ScaleInCooldown.Value);
 
-			IPool newPool = await _poolService.CreatePoolAsync(
-				create.Name, create.Condition, create.EnableAutoscaling, create.MinAgents, create.NumReserveAgents, conformInterval,
-				scaleOutCooldown, scaleInCooldown, sizeStrategies, fleetManagers, create.SizeStrategy, luSettings, jqSettings, cqamSettings, create.Properties);
+			AddPoolOptions options = new AddPoolOptions();
+			options.Condition = create.Condition;
+			options.EnableAutoscaling = create.EnableAutoscaling;
+			options.MinAgents = create.MinAgents;
+			options.NumReserveAgents = create.NumReserveAgents; 
+			options.ConformInterval = conformInterval;
+			options.ScaleOutCooldown = scaleOutCooldown;
+			options.ScaleInCooldown = scaleInCooldown; 
+			options.SizeStrategies = sizeStrategies;
+			options.FleetManagers = fleetManagers; 
+			options.SizeStrategy = create.SizeStrategy;
+			options.LeaseUtilizationSettings = luSettings; 
+			options.JobQueueSettings = jqSettings; 
+			options.ComputeQueueAwsMetricSettings = cqamSettings;
+			options.Properties = create.Properties;
+
+			IPool newPool = await _poolService.CreatePoolAsync(create.Name, options);
 			return new CreatePoolResponse(newPool.Id.ToString());
 		}
 
@@ -145,9 +159,25 @@ namespace Horde.Server.Agents.Pools
 			TimeSpan? scaleOutCooldown = update.ScaleOutCooldown == null ? null : TimeSpan.FromSeconds(update.ScaleOutCooldown.Value);
 			TimeSpan? scaleInCooldown = update.ScaleInCooldown == null ? null : TimeSpan.FromSeconds(update.ScaleInCooldown.Value);
 
-			await _poolService.UpdatePoolAsync(pool, update.Name, update.Condition, update.EnableAutoscaling,
-				update.MinAgents, update.NumReserveAgents, update.Properties, conformInterval, scaleOutCooldown, scaleInCooldown, update.SizeStrategy, newSizeStrategies, newFleetManagers,
-				update.LeaseUtilizationSettings?.Convert(), update.JobQueueSettings?.Convert(), update.ComputeQueueAwsMetricSettings?.Convert(), update.UseDefaultStrategy);
+			UpdatePoolOptions options = new UpdatePoolOptions();
+			options.Name = update.Name;
+			options.Condition = update.Condition;
+			options.EnableAutoscaling = update.EnableAutoscaling;
+			options.MinAgents = update.MinAgents;
+			options.NumReserveAgents = update.NumReserveAgents;
+			options.Properties = update.Properties;
+			options.ConformInterval = conformInterval;
+			options.ScaleOutCooldown = scaleOutCooldown;
+			options.ScaleInCooldown = scaleInCooldown;
+			options.SizeStrategy = update.SizeStrategy;
+			options.SizeStrategies = newSizeStrategies;
+			options.FleetManagers = newFleetManagers;
+			options.LeaseUtilizationSettings = update.LeaseUtilizationSettings?.Convert();
+			options.JobQueueSettings = update.JobQueueSettings?.Convert();
+			options.ComputeQueueAwsMetricSettings = update.ComputeQueueAwsMetricSettings?.Convert();
+			options.UseDefaultStrategy = update.UseDefaultStrategy;
+
+			await _poolService.UpdatePoolAsync(pool, options);
 			return new OkResult();
 		}
 
@@ -197,7 +227,7 @@ namespace Horde.Server.Agents.Pools
 					return NotFound(poolIdValue);
 				}
 
-				await _poolService.UpdatePoolAsync(pool, update.Name, newProperties: update.Properties);
+				await _poolService.UpdatePoolAsync(pool, new UpdatePoolOptions { Name = update.Name, Properties = update.Properties });
 			}
 			return Ok();
 		}
