Horde: Support filtering the AutoSDK view to sync.

diff --git a/Engine/Source/Programs/Horde/Horde.Server/Agents/IAgent.cs b/Engine/Source/Programs/Horde/Horde.Server/Agents/IAgent.cs
--- a/Engine/Source/Programs/Horde/Horde.Server/Agents/IAgent.cs
+++ b/Engine/Source/Programs/Horde/Horde.Server/Agents/IAgent.cs
@@ -186,6 +186,81 @@ namespace Horde.Server.Agents
 		}
 	}
 
+	/// <summary>
+	/// Configuration for an AutoSDK workspace
+	/// </summary>
+	public class AutoSdkConfig
+	{
+		/// <summary>
+		/// Disables use of AutoSDK.
+		/// </summary>
+		public static AutoSdkConfig None { get; } = new AutoSdkConfig(Array.Empty<string>());
+
+		/// <summary>
+		/// Syncs the entire AutoSDK folder.
+		/// </summary>
+		public static AutoSdkConfig Full { get; } = new AutoSdkConfig(null);
+
+		/// <summary>
+		/// Additive filter for paths to include in the workspace
+		/// </summary>
+		public IReadOnlyList<string> View { get; set; }
+
+		static readonly string[] s_fullView = new[] { "..." };
+
+		/// <summary>
+		/// Constructor
+		/// </summary>
+		/// <param name="filter">Filter for the workspace</param>
+		public AutoSdkConfig(IEnumerable<string>? filter)
+		{
+			if (filter == null)
+			{
+				View = s_fullView;
+			}
+			else
+			{
+				View = filter.OrderBy(x => x).Distinct().ToArray();
+			}
+		}
+
+		/// <summary>
+		/// Merge two workspaces together
+		/// </summary>
+		[return: NotNullIfNotNull(nameof(lhs))]
+		[return: NotNullIfNotNull(nameof(rhs))]
+		public static AutoSdkConfig? Merge(AutoSdkConfig? lhs, AutoSdkConfig? rhs)
+		{
+			if (lhs == null)
+			{
+				return rhs;
+			}
+			if (rhs == null)
+			{
+				return lhs;
+			}
+
+			return new AutoSdkConfig(Enumerable.Concat(lhs.View, rhs.View));
+		}
+
+		/// <summary>
+		/// Test whether two views are equal
+		/// </summary>
+		public static bool Equals(AutoSdkConfig? lhs, AutoSdkConfig? rhs)
+		{
+			if (lhs == null || lhs.View.Count == 0)
+			{
+				return rhs == null || rhs.View.Count == 0;
+			}
+			if (rhs == null || rhs.View.Count == 0)
+			{
+				return false;
+			}
+
+			return Enumerable.SequenceEqual(lhs.View, rhs.View);
+		}
+	}
+
 	/// <summary>
 	/// Document describing an active lease
 	/// </summary>
@@ -710,9 +785,10 @@ namespace Horde.Server.Agents
 		/// </summary>
 		/// <param name="agent"></param>
 		/// <param name="globalConfig"></param>
+		/// <param name="autoSdkConfig">Config for autosdk workspaces</param>
 		/// <param name="workspaces"></param>
 		/// <returns></returns>
-		public static HashSet<AgentWorkspace> GetAutoSdkWorkspaces(this IAgent agent, GlobalConfig globalConfig, List<AgentWorkspace> workspaces)
+		public static HashSet<AgentWorkspace> GetAutoSdkWorkspaces(this IAgent agent, GlobalConfig globalConfig, AutoSdkConfig autoSdkConfig, List<AgentWorkspace> workspaces)
 		{
 			HashSet<AgentWorkspace> autoSdkWorkspaces = new HashSet<AgentWorkspace>();
 			foreach (string? clusterName in workspaces.Select(x => x.Cluster).Distinct())
@@ -720,7 +796,7 @@ namespace Horde.Server.Agents
 				PerforceCluster? cluster = globalConfig.FindPerforceCluster(clusterName);
 				if (cluster != null)
 				{
-					AgentWorkspace? autoSdkWorkspace = GetAutoSdkWorkspace(agent, cluster);
+					AgentWorkspace? autoSdkWorkspace = GetAutoSdkWorkspace(agent, cluster, autoSdkConfig);
 					if (autoSdkWorkspace != null)
 					{
 						autoSdkWorkspaces.Add(autoSdkWorkspace);
@@ -735,14 +811,15 @@ namespace Horde.Server.Agents
 		/// </summary>
 		/// <param name="agent"></param>
 		/// <param name="cluster">The perforce cluster to get a workspace for</param>
+		/// <param name="autoSdkConfig">Configuration for autosdk</param>
 		/// <returns></returns>
-		public static AgentWorkspace? GetAutoSdkWorkspace(this IAgent agent, PerforceCluster cluster)
+		public static AgentWorkspace? GetAutoSdkWorkspace(this IAgent agent, PerforceCluster cluster, AutoSdkConfig autoSdkConfig)
 		{
 			foreach (AutoSdkWorkspace autoSdk in cluster.AutoSdk)
 			{
 				if (autoSdk.Stream != null && autoSdk.Properties.All(x => agent.Properties.Contains(x)))
 				{
-					return new AgentWorkspace(cluster.Name, autoSdk.UserName, autoSdk.Name ?? "AutoSDK", autoSdk.Stream!, null, true, null);
+					return new AgentWorkspace(cluster.Name, autoSdk.UserName, autoSdk.Name ?? "AutoSDK", autoSdk.Stream!, autoSdkConfig.View.ToList(), true, null);
 				}
 			}
 			return null;
@@ -830,14 +907,16 @@ namespace Horde.Server.Agents
 		/// <param name="streamConfig">The stream object</param>
 		/// <param name="agentType">The agent type</param>
 		/// <param name="workspace">Receives the agent workspace definition</param>
+		/// <param name="autoSdkConfig">Receives the autosdk workspace config</param>
 		/// <returns>True if the agent type was valid, and an agent workspace could be created</returns>
-		public static bool TryGetAgentWorkspace(this StreamConfig streamConfig, AgentConfig agentType, [NotNullWhen(true)] out (AgentWorkspace, bool)? workspace)
+		public static bool TryGetAgentWorkspace(this StreamConfig streamConfig, AgentConfig agentType, [NotNullWhen(true)] out AgentWorkspace? workspace, out AutoSdkConfig? autoSdkConfig)
 		{
 			// Get the workspace settings
 			if (agentType.Workspace == null)
 			{
 				// Use the default settings (fast switching workspace, clean 
-				workspace = (new AgentWorkspace(null, null, streamConfig.GetDefaultWorkspaceIdentifier(), streamConfig.Name, null, false, null), true);
+				workspace = new AgentWorkspace(null, null, streamConfig.GetDefaultWorkspaceIdentifier(), streamConfig.Name, null, false, null);
+				autoSdkConfig = AutoSdkConfig.Full;
 				return true;
 			}
 			else
@@ -847,6 +926,7 @@ namespace Horde.Server.Agents
 				if (!streamConfig.WorkspaceTypes.TryGetValue(agentType.Workspace, out workspaceType))
 				{
 					workspace = null;
+					autoSdkConfig = null;
 					return false;
 				}
 
@@ -866,7 +946,17 @@ namespace Horde.Server.Agents
 				}
 
 				// Create the new workspace
-				workspace = (new AgentWorkspace(workspaceType.Cluster, workspaceType.UserName, identifier, workspaceType.Stream ?? streamConfig.Name, workspaceType.View, workspaceType.Incremental, workspaceType.Method), workspaceType.UseAutoSdk);
+				workspace = new AgentWorkspace(workspaceType.Cluster, workspaceType.UserName, identifier, workspaceType.Stream ?? streamConfig.Name, workspaceType.View, workspaceType.Incremental, workspaceType.Method);
+
+				if (workspaceType.UseAutoSdk)
+				{
+					autoSdkConfig = AutoSdkConfig.Merge(new AutoSdkConfig(streamConfig.AutoSdkView), new AutoSdkConfig(workspaceType.AutoSdkView));
+				}
+				else
+				{
+					autoSdkConfig = null;
+				}
+
 				return true;
 			}
 		}
diff --git a/Engine/Source/Programs/Horde/Horde.Server/Agents/Pools/IPool.cs b/Engine/Source/Programs/Horde/Horde.Server/Agents/Pools/IPool.cs
--- a/Engine/Source/Programs/Horde/Horde.Server/Agents/Pools/IPool.cs
+++ b/Engine/Source/Programs/Horde/Horde.Server/Agents/Pools/IPool.cs
@@ -113,6 +113,11 @@ namespace Horde.Server.Agents.Pools
 		/// </summary>
 		public IReadOnlyList<AgentWorkspace> Workspaces { get; }
 
+		/// <summary>
+		/// AutoSDK view for this pool
+		/// </summary>
+		public AutoSdkConfig? AutoSdkConfig { get; }
+
 		/// <summary>
 		/// Revision string for the config file describing this pool, or null if added manually.
 		/// </summary>
diff --git a/Engine/Source/Programs/Horde/Horde.Server/Agents/Pools/IPoolCollection.cs b/Engine/Source/Programs/Horde/Horde.Server/Agents/Pools/IPoolCollection.cs
--- a/Engine/Source/Programs/Horde/Horde.Server/Agents/Pools/IPoolCollection.cs
+++ b/Engine/Source/Programs/Horde/Horde.Server/Agents/Pools/IPoolCollection.cs
@@ -174,9 +174,9 @@ namespace Horde.Server.Agents.Pools
 		public List<AgentWorkspace>? Workspaces { get; set; }
 
 		/// <summary>
-		/// New setting for whether to use autosdk
+		/// Settings for the autosdk workspace
 		/// </summary>
-		public bool? UseAutoSdk { get; set; }
+		public AutoSdkConfig? AutoSdkConfig { get; set; }
 
 		/// <summary>
 		/// New properties for the pool
diff --git a/Engine/Source/Programs/Horde/Horde.Server/Agents/Pools/PoolCollection.cs b/Engine/Source/Programs/Horde/Horde.Server/Agents/Pools/PoolCollection.cs
--- a/Engine/Source/Programs/Horde/Horde.Server/Agents/Pools/PoolCollection.cs
+++ b/Engine/Source/Programs/Horde/Horde.Server/Agents/Pools/PoolCollection.cs
@@ -33,8 +33,8 @@ namespace Horde.Server.Agents.Pools
 
 			public List<AgentWorkspace> Workspaces { get; set; } = new List<AgentWorkspace>();
 
-			[BsonIgnoreIfDefault, BsonDefaultValue(true)]
-			public bool UseAutoSdk { get; set; } = true;
+			[BsonIgnoreIfNull]
+			public AutoSdkConfig? AutoSdkConfig { get; set; }
 
 			public Dictionary<string, string> Properties { get; set; } = new Dictionary<string, string>();
 
@@ -112,7 +112,6 @@ namespace Horde.Server.Agents.Pools
 				Id = other.Id;
 				Name = String.IsNullOrEmpty(other.Name)? other.Id.ToString() : other.Name;
 				Condition = other.Condition;
-				UseAutoSdk = other.UseAutoSdk;
 				Properties = new Dictionary<string, string>(other.Properties);
 				EnableAutoscaling = other.EnableAutoscaling;
 				MinAgents = other.MinAgents;
@@ -316,9 +315,16 @@ namespace Horde.Server.Agents.Pools
 			{
 				transaction.Set(x => x.Workspaces, options.Workspaces);
 			}
-			if (options.UseAutoSdk != null)
+			if (options.AutoSdkConfig != null)
 			{
-				transaction.Set(x => x.UseAutoSdk, options.UseAutoSdk.Value);
+				if (options.AutoSdkConfig == AutoSdkConfig.None)
+				{
+					transaction.Unset(x => x.AutoSdkConfig!);
+				}
+				else
+				{
+					transaction.Set(x => x.AutoSdkConfig, options.AutoSdkConfig);
+				}
 			}
 			if (options.Properties != null)
 			{
diff --git a/Engine/Source/Programs/Horde/Horde.Server/Agents/Pools/PoolConfig.cs b/Engine/Source/Programs/Horde/Horde.Server/Agents/Pools/PoolConfig.cs
--- a/Engine/Source/Programs/Horde/Horde.Server/Agents/Pools/PoolConfig.cs
+++ b/Engine/Source/Programs/Horde/Horde.Server/Agents/Pools/PoolConfig.cs
@@ -4,7 +4,6 @@ using System;
 using System.Collections.Generic;
 using EpicGames.Horde.Common;
 using Horde.Server.Agents.Fleet;
-using Horde.Server.Utilities;
 
 namespace Horde.Server.Agents.Pools
 {
@@ -28,11 +27,6 @@ namespace Horde.Server.Agents.Pools
 		/// </summary>
 		public Condition? Condition { get; }
 
-		/// <summary>
-		/// Whether to sync autosdk alongside the workspaces for this pool
-		/// </summary>
-		public bool UseAutoSdk { get; }
-
 		/// <summary>
 		/// Arbitrary properties related to this pool
 		/// </summary>
diff --git a/Engine/Source/Programs/Horde/Horde.Server/Agents/Pools/PoolService.cs b/Engine/Source/Programs/Horde/Horde.Server/Agents/Pools/PoolService.cs
--- a/Engine/Source/Programs/Horde/Horde.Server/Agents/Pools/PoolService.cs
+++ b/Engine/Source/Programs/Horde/Horde.Server/Agents/Pools/PoolService.cs
@@ -155,7 +155,7 @@ namespace Horde.Server.Agents.Pools
 		/// <returns>List of workspaces</returns>
 		public async Task<HashSet<AgentWorkspace>> GetWorkspacesAsync(IAgent agent, DateTime validAtTime, GlobalConfig globalConfig)
 		{
-			bool addAutoSdkWorkspace = false;
+			AutoSdkConfig? autoSdkConfig = null;
 			HashSet<AgentWorkspace> workspaces = new HashSet<AgentWorkspace>();
 
 			Dictionary<PoolId, IPool> poolMapping = await GetPoolLookupAsync(validAtTime);
@@ -165,13 +165,13 @@ namespace Horde.Server.Agents.Pools
 				if (poolMapping.TryGetValue(poolId, out pool))
 				{
 					workspaces.UnionWith(pool.Workspaces);
-					addAutoSdkWorkspace |= pool.UseAutoSdk;
+					autoSdkConfig = AutoSdkConfig.Merge(autoSdkConfig, pool.AutoSdkConfig);
 				}
 			}
 
-			if (addAutoSdkWorkspace)
+			if (autoSdkConfig != null)
 			{
-				workspaces.UnionWith(agent.GetAutoSdkWorkspaces(globalConfig, workspaces.ToList()));
+				workspaces.UnionWith(agent.GetAutoSdkWorkspaces(globalConfig, autoSdkConfig, workspaces.ToList()));
 			}
 
 			return workspaces;
diff --git a/Engine/Source/Programs/Horde/Horde.Server/Agents/Pools/PoolUpdateService.cs b/Engine/Source/Programs/Horde/Horde.Server/Agents/Pools/PoolUpdateService.cs
--- a/Engine/Source/Programs/Horde/Horde.Server/Agents/Pools/PoolUpdateService.cs
+++ b/Engine/Source/Programs/Horde/Horde.Server/Agents/Pools/PoolUpdateService.cs
@@ -142,7 +142,7 @@ namespace Horde.Server.Agents.Pools
 				List<IPool> currentPools = await _pools.GetAsync();
 
 				// Lookup table of pool id to workspaces
-				HashSet<PoolId> poolsWithAutoSdk = new HashSet<PoolId>();
+				Dictionary<PoolId, AutoSdkConfig> poolToAutoSdkView = new Dictionary<PoolId, AutoSdkConfig>();
 				Dictionary<PoolId, List<AgentWorkspace>> poolToAgentWorkspaces = new Dictionary<PoolId, List<AgentWorkspace>>();
 
 				// Capture the current config state
@@ -154,10 +154,8 @@ namespace Horde.Server.Agents.Pools
 					foreach (KeyValuePair<string, AgentConfig> agentTypePair in streamConfig.AgentTypes)
 					{
 						// Create the new agent workspace
-						(AgentWorkspace, bool)? result;
-						if (streamConfig.TryGetAgentWorkspace(agentTypePair.Value, out result))
+						if (streamConfig.TryGetAgentWorkspace(agentTypePair.Value, out AgentWorkspace? agentWorkspace, out AutoSdkConfig? autoSdkConfig))
 						{
-							(AgentWorkspace agentWorkspace, bool useAutoSdk) = result.Value;
 							AgentConfig agentType = agentTypePair.Value;
 
 							// Find or add a list of workspaces for this pool
@@ -173,9 +171,11 @@ namespace Horde.Server.Agents.Pools
 							{
 								agentWorkspaces.Add(agentWorkspace);
 							}
-							if (useAutoSdk)
+							if (autoSdkConfig != null)
 							{
-								poolsWithAutoSdk.Add(agentType.Pool);
+								AutoSdkConfig? existingAutoSdkConfig;
+								poolToAutoSdkView.TryGetValue(agentType.Pool, out existingAutoSdkConfig);
+								poolToAutoSdkView[agentType.Pool] = AutoSdkConfig.Merge(autoSdkConfig, existingAutoSdkConfig);
 							}
 						}
 					}
@@ -191,13 +191,19 @@ namespace Horde.Server.Agents.Pools
 						newWorkspaces = new List<AgentWorkspace>();
 					}
 
+					// Get the autosdk view
+					AutoSdkConfig? newAutoSdkConfig;
+					if (!poolToAutoSdkView.TryGetValue(currentPool.Id, out newAutoSdkConfig))
+					{
+						newAutoSdkConfig = AutoSdkConfig.None;
+					}
+
 					// Update the pools document
-					bool useAutoSdk = poolsWithAutoSdk.Contains(currentPool.Id);
-					if (!AgentWorkspace.SetEquals(currentPool.Workspaces, newWorkspaces) || currentPool.Workspaces.Count != newWorkspaces.Count || currentPool.UseAutoSdk != useAutoSdk)
+					if (!AgentWorkspace.SetEquals(currentPool.Workspaces, newWorkspaces) || currentPool.Workspaces.Count != newWorkspaces.Count || !AutoSdkConfig.Equals(currentPool.AutoSdkConfig, newAutoSdkConfig))
 					{
 						_logger.LogInformation("New workspaces for pool {Pool}:{Workspaces}", currentPool.Id, String.Join("", newWorkspaces.Select(x => $"\n  Identifier=\"{x.Identifier}\", Stream={x.Stream}")));
 
-						IPool? result = await _pools.TryUpdateAsync(currentPool, new UpdatePoolOptions { Workspaces = newWorkspaces, UseAutoSdk = useAutoSdk });
+						IPool? result = await _pools.TryUpdateAsync(currentPool, new UpdatePoolOptions { Workspaces = newWorkspaces, AutoSdkConfig = newAutoSdkConfig });
 						if (result == null)
 						{
 							_logger.LogInformation("Pool modified; will retry");
diff --git a/Engine/Source/Programs/Horde/Horde.Server/Jobs/JobTaskSource.cs b/Engine/Source/Programs/Horde/Horde.Server/Jobs/JobTaskSource.cs
--- a/Engine/Source/Programs/Horde/Horde.Server/Jobs/JobTaskSource.cs
+++ b/Engine/Source/Programs/Horde/Horde.Server/Jobs/JobTaskSource.cs
@@ -75,7 +75,7 @@ namespace Horde.Server.Jobs
 			/// <summary>
 			/// Whether or not to use the AutoSDK.
 			/// </summary>
-			public bool _useAutoSdk;
+			public AutoSdkConfig? _autoSdkConfig;
 
 			/// <summary>
 			/// Task for creating a lease and assigning to a waiter
@@ -99,14 +99,14 @@ namespace Horde.Server.Jobs
 			/// <param name="batchIdx">The batch index to execute</param>
 			/// <param name="poolId">Unique id of the pool of machines to allocate from</param>
 			/// <param name="workspace">The workspace that this job should run in</param>
-			/// <param name="useAutoSdk">Whether or not to use the AutoSDK</param>
-			public QueueItem(IJob job, int batchIdx, PoolId poolId, AgentWorkspace workspace, bool useAutoSdk)
+			/// <param name="autoSdkConfig">Whether or not to use the AutoSDK</param>
+			public QueueItem(IJob job, int batchIdx, PoolId poolId, AgentWorkspace workspace, AutoSdkConfig? autoSdkConfig)
 			{
 				_job = job;
 				_batchIdx = batchIdx;
 				_poolId = poolId;
 				_workspace = workspace;
-				_useAutoSdk = useAutoSdk;
+				_autoSdkConfig = autoSdkConfig;
 			}
 		}
 
@@ -418,7 +418,7 @@ namespace Horde.Server.Jobs
 					{
 						newJob = await SkipBatchAsync(newJob, batch.Id, graph, JobStepBatchError.NoAgentsInPool);
 					}
-					else if (!stream.Config.TryGetAgentWorkspace(agentType, out (AgentWorkspace, bool)? workspaceResult))
+					else if (!stream.Config.TryGetAgentWorkspace(agentType, out AgentWorkspace? workspace, out AutoSdkConfig? autoSdkConfig))
 					{
 						newJob = await SkipBatchAsync(newJob, batch.Id, graph, JobStepBatchError.UnknownWorkspace);
 					}
@@ -454,8 +454,7 @@ namespace Horde.Server.Jobs
 
 						if (newJob != null)
 						{
-							(AgentWorkspace workspace, bool useAutoSdk) = workspaceResult.Value;
-							QueueItem newQueueItem = new QueueItem(newJob, batchIdx, agentType.Pool, workspace, useAutoSdk);
+							QueueItem newQueueItem = new QueueItem(newJob, batchIdx, agentType.Pool, workspace, autoSdkConfig);
 							newQueue.Add(newQueueItem);
 							newBatchIdToQueueItem[(newJob.Id, batch.Id)] = newQueueItem;
 
@@ -629,7 +628,7 @@ namespace Horde.Server.Jobs
 								else
 								{
 									RemoveQueueItem(existingItem);
-									InsertQueueItem(job, batchIdx, existingItem._poolId, existingItem._workspace, existingItem._useAutoSdk);
+									InsertQueueItem(job, batchIdx, existingItem._poolId, existingItem._workspace, existingItem._autoSdkConfig);
 								}
 							}
 							continue;
@@ -642,11 +641,9 @@ namespace Horde.Server.Jobs
 						AgentConfig? agentType;
 						if (streamConfig.AgentTypes.TryGetValue(group.AgentType, out agentType))
 						{
-							(AgentWorkspace, bool)? result;
-							if (streamConfig.TryGetAgentWorkspace(agentType, out result))
+							if (streamConfig.TryGetAgentWorkspace(agentType, out AgentWorkspace? agentWorkspace, out AutoSdkConfig? autoSdkConfig))
 							{
-								(AgentWorkspace agentWorkspace, bool useAutoSdk) = result.Value;
-								InsertQueueItem(job, batchIdx, agentType.Pool, agentWorkspace, useAutoSdk);
+								InsertQueueItem(job, batchIdx, agentType.Pool, agentWorkspace, autoSdkConfig);
 							}
 						}
 					}
@@ -763,7 +760,7 @@ namespace Horde.Server.Jobs
 				leaseName.Append(CultureInfo.InvariantCulture, $" - {job.Name}");
 
 				// Encode the payload
-				ExecuteJobTask? task = await CreateExecuteJobTaskAsync(leaseId, streamConfig, job, batch, agent, item._workspace, item._useAutoSdk, logId);
+				ExecuteJobTask? task = await CreateExecuteJobTaskAsync(leaseId, streamConfig, job, batch, agent, item._workspace, item._autoSdkConfig, logId);
 				if (task != null)
 				{
 					byte[] payload = Any.Pack(task).ToByteArray();
@@ -813,7 +810,7 @@ namespace Horde.Server.Jobs
 			return null;
 		}
 
-		async Task<ExecuteJobTask?> CreateExecuteJobTaskAsync(LeaseId leaseId, StreamConfig streamConfig, IJob job, IJobStepBatch batch, IAgent agent, AgentWorkspace workspace, bool useAutoSdk, LogId logId)
+		async Task<ExecuteJobTask?> CreateExecuteJobTaskAsync(LeaseId leaseId, StreamConfig streamConfig, IJob job, IJobStepBatch batch, IAgent agent, AgentWorkspace workspace, AutoSdkConfig? autoSdkConfig, LogId logId)
 		{
 			// Get the lease name
 			StringBuilder leaseName = new StringBuilder($"{streamConfig.Name} - ");
@@ -858,7 +855,7 @@ namespace Horde.Server.Jobs
 				return null;
 			}
 
-			AgentWorkspace? autoSdkWorkspace = useAutoSdk ? agent.GetAutoSdkWorkspace(cluster) : null;
+			AgentWorkspace? autoSdkWorkspace = (autoSdkConfig != null)? agent.GetAutoSdkWorkspace(cluster, autoSdkConfig) : null;
 			if (autoSdkWorkspace != null)
 			{
 				autoSdkWorkspace.Method = workspace.Method;
@@ -867,6 +864,7 @@ namespace Horde.Server.Jobs
 				{
 					return null;
 				}
+
 				task.AutoSdkWorkspace = workspaces.Last();
 			}
 
@@ -1095,13 +1093,13 @@ namespace Horde.Server.Jobs
 		/// <param name="batchIdx"></param>
 		/// <param name="poolId">The pool to use</param>
 		/// <param name="workspace">The workspace for this item to run in</param>
-		/// <param name="useAutoSdk">Whether or not to use the AutoSDK</param>
+		/// <param name="autoSdkConfig">Whether or not to use the AutoSDK</param>
 		/// <returns></returns>
-		void InsertQueueItem(IJob job, int batchIdx, PoolId poolId, AgentWorkspace workspace, bool useAutoSdk)
+		void InsertQueueItem(IJob job, int batchIdx, PoolId poolId, AgentWorkspace workspace, AutoSdkConfig? autoSdkConfig)
 		{
 			_logger.LogDebug("Adding queued job {JobId}, batch {BatchId} [Pool: {Pool}, Workspace: {Workspace}]", job.Id, job.Batches[batchIdx].Id, poolId, workspace.Identifier);
 
-			QueueItem newItem = new QueueItem(job, batchIdx, poolId, workspace, useAutoSdk);
+			QueueItem newItem = new QueueItem(job, batchIdx, poolId, workspace, autoSdkConfig);
 			_batchIdToQueueItem[newItem.Id] = newItem;
 			_queue.Add(newItem);
 
diff --git a/Engine/Source/Programs/Horde/Horde.Server/Streams/StreamConfig.cs b/Engine/Source/Programs/Horde/Horde.Server/Streams/StreamConfig.cs
--- a/Engine/Source/Programs/Horde/Horde.Server/Streams/StreamConfig.cs
+++ b/Engine/Source/Programs/Horde/Horde.Server/Streams/StreamConfig.cs
@@ -151,6 +151,11 @@ namespace Horde.Server.Streams
 		/// </summary>
 		public JobOptions JobOptions { get; set; } = new JobOptions();
 
+		/// <summary>
+		/// View for the AutoSDK paths to sync. If null, the whole thing will be synced.
+		/// </summary>
+		public List<string>? AutoSdkView { get; set; }
+
 		/// <summary>
 		/// Legacy name for the default preflight template
 		/// </summary>
@@ -605,6 +610,11 @@ namespace Horde.Server.Streams
 		/// </summary>
 		public bool UseAutoSdk { get; set; } = true;
 
+		/// <summary>
+		/// View for the AutoSDK paths to sync. If null, the whole thing will be synced.
+		/// </summary>
+		public List<string>? AutoSdkView { get; set; }
+
 		/// <summary>
 		/// Method to use when syncing/materializing data from Perforce
 		/// </summary>
