Horde: Append stream and workspace autosdk views, rather than taking the superset.

diff --git a/Engine/Source/Programs/Horde/Horde.Server/Agents/IAgent.cs b/Engine/Source/Programs/Horde/Horde.Server/Agents/IAgent.cs
--- a/Engine/Source/Programs/Horde/Horde.Server/Agents/IAgent.cs
+++ b/Engine/Source/Programs/Horde/Horde.Server/Agents/IAgent.cs
@@ -950,7 +950,7 @@ namespace Horde.Server.Agents
 
 				if (workspaceType.UseAutoSdk)
 				{
-					autoSdkConfig = AutoSdkConfig.Merge(new AutoSdkConfig(streamConfig.AutoSdkView), new AutoSdkConfig(workspaceType.AutoSdkView));
+					autoSdkConfig = new AutoSdkConfig(Enumerable.Concat(streamConfig.AutoSdkView ?? Enumerable.Empty<string>(), workspaceType.AutoSdkView ?? Enumerable.Empty<string>()));
 				}
 				else
 				{
