Horde: Add a simpler secret management solution; secrets are now defined in the global config file, and can be retrieved via a query to /api/v1/secrets/{name}. Each secret can have a collection of key/value pairs associated with it.

diff --git a/Engine/Source/Programs/Horde/Horde.Server.Tests/TestSetup.cs b/Engine/Source/Programs/Horde/Horde.Server.Tests/TestSetup.cs
--- a/Engine/Source/Programs/Horde/Horde.Server.Tests/TestSetup.cs
+++ b/Engine/Source/Programs/Horde/Horde.Server.Tests/TestSetup.cs
@@ -114,7 +114,6 @@ namespace Horde.Server.Tests
 		public IArtifactCollectionV1 ArtifactCollection => ServiceProvider.GetRequiredService<IArtifactCollectionV1>();
 		public IDowntimeService DowntimeService => ServiceProvider.GetRequiredService<IDowntimeService>();
 		public RpcService RpcService => ServiceProvider.GetRequiredService<RpcService>();
-		public CredentialService CredentialService => ServiceProvider.GetRequiredService<CredentialService>();
 		public PoolService PoolService => ServiceProvider.GetRequiredService<PoolService>();
 		public LifetimeService LifetimeService => ServiceProvider.GetRequiredService<LifetimeService>();
 		public ScheduleService ScheduleService => ServiceProvider.GetRequiredService<ScheduleService>();
@@ -247,7 +246,6 @@ namespace Horde.Server.Tests
 			services.AddSingleton<ConsistencyService>();
 			services.AddSingleton<RequestTrackerService>();
 			services.AddSingleton<GlobalsService>();
-			services.AddSingleton<CredentialService>();
 			services.AddSingleton<JobTaskSource>();
 			services.AddSingleton<IDowntimeService, DowntimeServiceStub>();
 			services.AddSingleton<IssueService>();
diff --git a/Engine/Source/Programs/Horde/Horde.Server/Secrets/Credential.cs b/Engine/Source/Programs/Horde/Horde.Server/Secrets/Credential.cs
deleted file mode 100644
diff --git a/Engine/Source/Programs/Horde/Horde.Server/Secrets/CredentialAclAction.cs b/Engine/Source/Programs/Horde/Horde.Server/Secrets/CredentialAclAction.cs
deleted file mode 100644
diff --git a/Engine/Source/Programs/Horde/Horde.Server/Secrets/CredentialMessages.cs b/Engine/Source/Programs/Horde/Horde.Server/Secrets/CredentialMessages.cs
deleted file mode 100644
diff --git a/Engine/Source/Programs/Horde/Horde.Server/Secrets/CredentialService.cs b/Engine/Source/Programs/Horde/Horde.Server/Secrets/CredentialService.cs
deleted file mode 100644
diff --git a/Engine/Source/Programs/Horde/Horde.Server/Secrets/CredentialsController.cs b/Engine/Source/Programs/Horde/Horde.Server/Secrets/CredentialsController.cs
deleted file mode 100644
diff --git a/Engine/Source/Programs/Horde/Horde.Server/Secrets/SecretAclAction.cs b/Engine/Source/Programs/Horde/Horde.Server/Secrets/SecretAclAction.cs
new file mode 100644
--- /dev/null
+++ b/Engine/Source/Programs/Horde/Horde.Server/Secrets/SecretAclAction.cs
@@ -0,0 +1,17 @@
+// Copyright Epic Games, Inc. All Rights Reserved.
+
+using Horde.Server.Acls;
+
+namespace Horde.Server.Secrets
+{
+	/// <summary>
+	/// ACL actions for secrets
+	/// </summary>
+	public static class SecretAclAction
+	{
+		/// <summary>
+		/// View a credential
+		/// </summary>
+		public static AclAction ViewSecret { get; } = new AclAction("ViewSecret");
+	}
+}
diff --git a/Engine/Source/Programs/Horde/Horde.Server/Secrets/SecretConfig.cs b/Engine/Source/Programs/Horde/Horde.Server/Secrets/SecretConfig.cs
new file mode 100644
--- /dev/null
+++ b/Engine/Source/Programs/Horde/Horde.Server/Secrets/SecretConfig.cs
@@ -0,0 +1,53 @@
+// Copyright Epic Games, Inc. All Rights Reserved.
+
+using System.Collections.Generic;
+using System.Security.Claims;
+using System.Text.Json.Serialization;
+using Horde.Server.Acls;
+using Horde.Server.Server;
+
+namespace Horde.Server.Secrets
+{
+	/// <summary>
+	/// Configuration for a secret value
+	/// </summary>
+	public class SecretConfig
+	{
+		[JsonIgnore]
+		internal GlobalConfig GlobalConfig { get; private set; } = null!;
+
+		/// <summary>
+		/// Identifier for this secret
+		/// </summary>
+		public SecretId Id { get; set; }
+
+		/// <summary>
+		/// Key/value pairs associated with this secret
+		/// </summary>
+		public Dictionary<string, string> Data { get; set; } = new Dictionary<string, string>();
+
+		/// <summary>
+		/// Defines access to this particular secret
+		/// </summary>
+		public AclConfig? Acl { get; set; }
+
+		/// <summary>
+		/// Called after the config has been read
+		/// </summary>
+		/// <param name="globalConfig">Parent GlobalConfig object</param>
+		public void PostLoad(GlobalConfig globalConfig)
+		{
+			GlobalConfig = globalConfig;
+		}
+
+		/// <summary>
+		/// Authorizes a user to perform a given action
+		/// </summary>
+		/// <param name="action">The action being performed</param>
+		/// <param name="user">The principal to validate</param>
+		public bool Authorize(AclAction action, ClaimsPrincipal user)
+		{
+			return Acl?.Authorize(action, user) ?? GlobalConfig.Authorize(action, user);
+		}
+	}
+}
diff --git a/Engine/Source/Programs/Horde/Horde.Server/Secrets/SecretId.cs b/Engine/Source/Programs/Horde/Horde.Server/Secrets/SecretId.cs
new file mode 100644
--- /dev/null
+++ b/Engine/Source/Programs/Horde/Horde.Server/Secrets/SecretId.cs
@@ -0,0 +1,46 @@
+// Copyright Epic Games, Inc. All Rights Reserved.
+
+using System.ComponentModel;
+using EpicGames.Core;
+using EpicGames.Horde;
+using EpicGames.Serialization;
+using Horde.Server.Utilities;
+
+namespace Horde.Server.Secrets
+{
+	/// <summary>
+	/// Identifier for a secret
+	/// </summary>
+	/// <param name="Id">Id to construct from</param>
+	[JsonSchemaString]
+	[TypeConverter(typeof(StringIdTypeConverter<SecretId, SecretIdConverter>))]
+	[StringIdConverter(typeof(SecretIdConverter))]
+	[CbConverter(typeof(StringIdCbConverter<SecretId, SecretIdConverter>))]
+	public record struct SecretId(StringId Id)
+	{
+		/// <summary>
+		/// Constructor
+		/// </summary>
+		public SecretId(string id) : this(new StringId(id))
+		{
+		}
+
+		/// <inheritdoc cref="StringId.IsEmpty"/>
+		public bool IsEmpty => Id.IsEmpty;
+
+		/// <inheritdoc/>
+		public override string ToString() => Id.ToString();
+	}
+
+	/// <summary>
+	/// Converter to and from <see cref="StringId"/> instances.
+	/// </summary>
+	class SecretIdConverter : StringIdConverter<SecretId>
+	{
+		/// <inheritdoc/>
+		public override SecretId FromStringId(StringId id) => new SecretId(id);
+
+		/// <inheritdoc/>
+		public override StringId ToStringId(SecretId value) => value.Id;
+	}
+}
diff --git a/Engine/Source/Programs/Horde/Horde.Server/Secrets/SecretsController.cs b/Engine/Source/Programs/Horde/Horde.Server/Secrets/SecretsController.cs
new file mode 100644
--- /dev/null
+++ b/Engine/Source/Programs/Horde/Horde.Server/Secrets/SecretsController.cs
@@ -0,0 +1,102 @@
+// Copyright Epic Games, Inc. All Rights Reserved.
+
+using System.Collections.Generic;
+using Horde.Server.Server;
+using Horde.Server.Utilities;
+using Microsoft.AspNetCore.Authorization;
+using Microsoft.AspNetCore.Mvc;
+using Microsoft.Extensions.Options;
+
+namespace Horde.Server.Secrets
+{
+	/// <summary>
+	/// Response listing all the secrets available to the current user
+	/// </summary>
+	public class GetSecretsResponse
+	{
+		/// <summary>
+		/// List of secret ids
+		/// </summary>
+		public List<SecretId> Ids { get; set; } = new List<SecretId>();
+	}
+
+	/// <summary>
+	/// Gets data for a particular secret
+	/// </summary>
+	public class GetSecretResponse
+	{
+		readonly SecretConfig _config;
+
+		/// <summary>
+		/// Id of the secret
+		/// </summary>
+		public SecretId Id => _config.Id;
+
+		/// <summary>
+		/// Key value pairs for the secret
+		/// </summary>
+		public Dictionary<string, string> Data => _config.Data;
+
+		internal GetSecretResponse(SecretConfig config) => _config = config;
+	}
+
+	/// <summary>
+	/// Controller for the /api/v1/credentials endpoint
+	/// </summary>
+	[ApiController]
+	[Authorize]
+	[Route("[controller]")]
+	public class SecretsController : HordeControllerBase
+	{
+		private readonly IOptionsSnapshot<GlobalConfig> _globalConfig;
+
+		/// <summary>
+		/// Constructor
+		/// </summary>
+		public SecretsController(IOptionsSnapshot<GlobalConfig> globalConfig)
+		{
+			_globalConfig = globalConfig;
+		}
+
+		/// <summary>
+		/// Query all the secrets available for the current user
+		/// </summary>
+		[HttpGet]
+		[Route("/api/v1/secrets")]
+		public ActionResult<GetSecretsResponse> GetSecrets()
+		{
+			GetSecretsResponse response = new GetSecretsResponse();
+			foreach (SecretConfig secret in _globalConfig.Value.Secrets)
+			{
+				if (secret.Authorize(SecretAclAction.ViewSecret, User))
+				{
+					response.Ids.Add(secret.Id);
+				}
+			}
+			return response;
+		}
+
+		/// <summary>
+		/// Retrieve information about a specific secret
+		/// </summary>
+		/// <param name="secretId">Id of the secret to retrieve</param>
+		/// <param name="filter">Filter for properties to return</param>
+		/// <returns>Information about the requested secret</returns>
+		[HttpGet]
+		[Route("/api/v1/secrets/{secretId}")]
+		[ProducesResponseType(typeof(GetSecretResponse), 200)]
+		public ActionResult<object> GetSecret(SecretId secretId, [FromQuery] PropertyFilter? filter = null)
+		{
+			if (!_globalConfig.Value.TryGetSecret(secretId, out SecretConfig? secretConfig))
+			{
+				return NotFound(secretId);
+			}
+			if (!secretConfig.Authorize(SecretAclAction.ViewSecret, User))
+			{
+				return Forbid(SecretAclAction.ViewSecret, secretId);
+			}
+
+			return new GetSecretResponse(secretConfig).ApplyFilter(filter);
+		}
+	}
+}
diff --git a/Engine/Source/Programs/Horde/Horde.Server/Server/GlobalConfig.cs b/Engine/Source/Programs/Horde/Horde.Server/Server/GlobalConfig.cs
--- a/Engine/Source/Programs/Horde/Horde.Server/Server/GlobalConfig.cs
+++ b/Engine/Source/Programs/Horde/Horde.Server/Server/GlobalConfig.cs
@@ -19,11 +19,13 @@ using Horde.Server.Agents.Software;
 using Horde.Server.Artifacts;
 using Horde.Server.Configuration;
 using Horde.Server.Projects;
+using Horde.Server.Secrets;
 using Horde.Server.Storage;
 using Horde.Server.Streams;
 using Horde.Server.Tools;
 using Horde.Server.Users;
 using Horde.Server.Utilities;
+using Microsoft.Extensions.Configuration;
 
 namespace Horde.Server.Server
 {
@@ -165,6 +167,11 @@ namespace Horde.Server.Server
 		/// </summary>
 		public List<ComputeClusterConfig> Compute { get; set; } = new List<ComputeClusterConfig>();
 
+		/// <summary>
+		/// List of secrets
+		/// </summary>
+		public List<SecretConfig> Secrets { get; set; } = new List<SecretConfig>();
+
 		/// <summary>
 		/// Device configuration
 		/// </summary>
@@ -213,6 +220,7 @@ namespace Horde.Server.Server
 		private readonly Dictionary<ClusterId, ComputeClusterConfig> _computeClusterLookup = new Dictionary<ClusterId, ComputeClusterConfig>();
 		private readonly Dictionary<AclScopeName, IAclScope> _aclScopeLookup = new Dictionary<AclScopeName, IAclScope>();
 		private readonly Dictionary<ArtifactType, ArtifactTypeConfig> _artifactTypeLookup = new Dictionary<ArtifactType, ArtifactTypeConfig>();
+		private readonly Dictionary<SecretId, SecretConfig> _secretLookup = new Dictionary<SecretId, SecretConfig>();
 
 		/// <summary>
 		/// Called after the config file has been read
@@ -271,6 +279,13 @@ namespace Horde.Server.Server
 				_artifactTypeLookup.Add(artifactType.Name, artifactType);
 			}
 
+			_secretLookup.Clear();
+			foreach (SecretConfig secret in Secrets)
+			{
+				_secretLookup.Add(secret.Id, secret);
+				secret.PostLoad(this);
+			}
+
 			Storage.PostLoad(this);
 		}
 
@@ -331,6 +346,14 @@ namespace Horde.Server.Server
 		/// <returns>True on success</returns>
 		public bool TryGetComputeCluster(ClusterId clusterId, [NotNullWhen(true)] out ComputeClusterConfig? config) => _computeClusterLookup.TryGetValue(clusterId, out config);
 
+		/// <summary>
+		/// Attempts to get compute cluster configuration from this object
+		/// </summary>
+		/// <param name="secretId">Secret id</param>
+		/// <param name="config">Receives the secret configuration on success</param>
+		/// <returns>True on success</returns>
+		public bool TryGetSecret(SecretId secretId, [NotNullWhen(true)] out SecretConfig? config) => _secretLookup.TryGetValue(secretId, out config);
+
 		/// <summary>
 		/// Authorizes a user to perform a given action
 		/// </summary>
diff --git a/Engine/Source/Programs/Horde/Horde.Server/Server/MongoService.cs b/Engine/Source/Programs/Horde/Horde.Server/Server/MongoService.cs
--- a/Engine/Source/Programs/Horde/Horde.Server/Server/MongoService.cs
+++ b/Engine/Source/Programs/Horde/Horde.Server/Server/MongoService.cs
@@ -196,11 +196,6 @@ namespace Horde.Server.Server
 		/// </summary>
 		public IMongoDatabase Database { get; private set; }
 
-		/// <summary>
-		/// Collection of credential documents
-		/// </summary>
-		public IMongoCollection<Credential> Credentials { get; }
-
 		/// <summary>
 		/// Collection of singleton documents
 		/// </summary>
@@ -333,8 +328,6 @@ namespace Horde.Server.Server
 
 				SingletonsV1 = GetCollection<BsonDocument>("Singletons");
 				SingletonsV2 = GetCollection<BsonDocument>("SingletonsV2");
-
-				Credentials = GetCollection<Credential>("Credentials");
 			}
 			catch (Exception ex)
 			{
diff --git a/Engine/Source/Programs/Horde/Horde.Server/Startup.cs b/Engine/Source/Programs/Horde/Horde.Server/Startup.cs
--- a/Engine/Source/Programs/Horde/Horde.Server/Startup.cs
+++ b/Engine/Source/Programs/Horde/Horde.Server/Startup.cs
@@ -452,7 +452,6 @@ namespace Horde.Server
 			services.AddSingleton<ConsistencyService>();
 			services.AddSingleton<RequestTrackerService>();
 			services.AddSingleton<ComputeService>();
-			services.AddSingleton<CredentialService>();
 			services.AddSingleton<MongoService>();
 			services.AddSingleton<GlobalsService>();
 			services.AddSingleton<ICommitService, CommitService>();
diff --git a/Engine/Source/Programs/Horde/Horde.Server/Utilities/HordeControllerBase.cs b/Engine/Source/Programs/Horde/Horde.Server/Utilities/HordeControllerBase.cs
--- a/Engine/Source/Programs/Horde/Horde.Server/Utilities/HordeControllerBase.cs
+++ b/Engine/Source/Programs/Horde/Horde.Server/Utilities/HordeControllerBase.cs
@@ -10,6 +10,7 @@ using Horde.Server.Artifacts;
 using Horde.Server.Jobs;
 using Horde.Server.Jobs.Bisect;
 using Horde.Server.Projects;
+using Horde.Server.Secrets;
 using Horde.Server.Streams;
 using Horde.Server.Tools;
 using Microsoft.AspNetCore.Http;
@@ -109,6 +110,15 @@ namespace Horde.Server.Utilities
 			return Forbid(action, "project {ProjectId}", projectId);
 		}
 
+		/// <summary>
+		/// Returns a 403 (forbidden) response with the given action and object
+		/// </summary>
+		[NonAction]
+		protected ActionResult Forbid(AclAction action, SecretId secretId)
+		{
+			return Forbid(action, "secret {SecretId}", secretId);
+		}
+
 		/// <summary>
 		/// Returns a 403 (forbidden) response with the given action and object
 		/// </summary>
@@ -262,6 +272,15 @@ namespace Horde.Server.Utilities
 			return NotFound("Bisect task {BisectTaskId} not found", bisectTaskId);
 		}
 
+		/// <summary>
+		/// Returns a 404 response for the given object
+		/// </summary>
+		[NonAction]
+		protected ActionResult NotFound(SecretId secretId)
+		{
+			return NotFound("Secret {SecretId} not found", secretId);
+		}
+
 		/// <summary>
 		/// Returns a 404 response for the given object
 		/// </summary>
