Horde: Move methods for querying the server config into the debug controller, and create specific ACL actions for functionality previously deemed admin-only.

diff --git a/Engine/Source/Programs/Horde/Horde.Server/Agents/AgentsController.cs b/Engine/Source/Programs/Horde/Horde.Server/Agents/AgentsController.cs
--- a/Engine/Source/Programs/Horde/Horde.Server/Agents/AgentsController.cs
+++ b/Engine/Source/Programs/Horde/Horde.Server/Agents/AgentsController.cs
@@ -100,7 +100,7 @@ namespace Horde.Server.Agents
 		async ValueTask<object> GetAgentResponseAsync(IAgent agent, PropertyFilter? filter = null)
 		{
 			double? rate = null;
-			if (_globalConfig.Value.Authorize(AdminAclAction.ViewCosts, User))
+			if (_globalConfig.Value.Authorize(ServerAclAction.ViewCosts, User))
 			{
 				rate = await _agentService.GetRateAsync(agent.Id);
 			}
@@ -317,7 +317,7 @@ namespace Horde.Server.Agents
 			List<ILease> leases = await _agentService.FindLeasesAsync(agentId, sessionId, startTime?.UtcDateTime, finishTime?.UtcDateTime, index, count);
 
 			double? agentRate = null;
-			if (_globalConfig.Value.Authorize(AdminAclAction.ViewCosts, User))
+			if (_globalConfig.Value.Authorize(ServerAclAction.ViewCosts, User))
 			{
 				agentRate = await _agentService.GetRateAsync(agentId);
 			}
@@ -360,7 +360,7 @@ namespace Horde.Server.Agents
 			}
 
 			double? agentRate = null;
-			if (_globalConfig.Value.Authorize(AdminAclAction.ViewCosts, User))
+			if (_globalConfig.Value.Authorize(ServerAclAction.ViewCosts, User))
 			{
 				agentRate = await _agentService.GetRateAsync(agentId);
 			}
diff --git a/Engine/Source/Programs/Horde/Horde.Server/Agents/Leases/LeasesController.cs b/Engine/Source/Programs/Horde/Horde.Server/Agents/Leases/LeasesController.cs
--- a/Engine/Source/Programs/Horde/Horde.Server/Agents/Leases/LeasesController.cs
+++ b/Engine/Source/Programs/Horde/Horde.Server/Agents/Leases/LeasesController.cs
@@ -69,7 +69,7 @@ namespace Horde.Server.Agents.Leases
 				return Forbid(LeaseAclAction.ViewLeases);
 			}
 
-			bool includeCosts = _globalConfig.Value.Authorize(AdminAclAction.ViewCosts, User);
+			bool includeCosts = _globalConfig.Value.Authorize(ServerAclAction.ViewCosts, User);
 
 			List<ILease> leases;
 			if (minFinishTime == null && maxFinishTime == null)
@@ -131,7 +131,7 @@ namespace Horde.Server.Agents.Leases
 			}
 
 			double? agentRate = null;
-			if (_globalConfig.Value.Authorize(AdminAclAction.ViewCosts, User))
+			if (_globalConfig.Value.Authorize(ServerAclAction.ViewCosts, User))
 			{
 				agentRate = await _agentService.GetRateAsync(agent.Id);
 			}
diff --git a/Engine/Source/Programs/Horde/Horde.Server/Jobs/JobsController.cs b/Engine/Source/Programs/Horde/Horde.Server/Jobs/JobsController.cs
--- a/Engine/Source/Programs/Horde/Horde.Server/Jobs/JobsController.cs
+++ b/Engine/Source/Programs/Horde/Horde.Server/Jobs/JobsController.cs
@@ -415,7 +415,7 @@ namespace Horde.Server.Jobs
 			}
 
 			IGraph graph = await _jobService.GetGraphAsync(job);
-			bool includeCosts = streamConfig.Authorize(AdminAclAction.ViewCosts, User);
+			bool includeCosts = streamConfig.Authorize(ServerAclAction.ViewCosts, User);
 			return await CreateJobResponseAsync(job, graph, includeCosts, filter);
 		}
 
@@ -825,7 +825,7 @@ namespace Horde.Server.Jobs
 				StreamConfig? streamConfig;
 				if (_globalConfig.Value.TryGetStream(grouping.Key, out streamConfig) && streamConfig.Authorize(JobAclAction.ViewJob, User))
 				{
-					bool includeCosts = streamConfig.Authorize(AdminAclAction.ViewCosts, User);
+					bool includeCosts = streamConfig.Authorize(ServerAclAction.ViewCosts, User);
 					foreach (IJob job in grouping)
 					{
 						IGraph graph = await _jobService.GetGraphAsync(job);
@@ -1047,7 +1047,7 @@ namespace Horde.Server.Jobs
 				return Forbid(JobAclAction.ViewJob, job.StreamId);
 			}
 
-			bool includeCosts = streamConfig.Authorize(AdminAclAction.ViewCosts, User);
+			bool includeCosts = streamConfig.Authorize(ServerAclAction.ViewCosts, User);
 
 			List<object> responses = new List<object>();
 			foreach (IJobStepBatch batch in job.Batches)
@@ -1124,7 +1124,7 @@ namespace Horde.Server.Jobs
 				return Forbid(JobAclAction.ViewJob, job.StreamId);
 			}
 
-			bool includeCosts = streamConfig.Authorize(AdminAclAction.ViewCosts, User);
+			bool includeCosts = streamConfig.Authorize(ServerAclAction.ViewCosts, User);
 			foreach (IJobStepBatch batch in job.Batches)
 			{
 				if (batch.Id == batchId)
diff --git a/Engine/Source/Programs/Horde/Horde.Server/Notifications/SubscriptionsController.cs b/Engine/Source/Programs/Horde/Horde.Server/Notifications/SubscriptionsController.cs
--- a/Engine/Source/Programs/Horde/Horde.Server/Notifications/SubscriptionsController.cs
+++ b/Engine/Source/Programs/Horde/Horde.Server/Notifications/SubscriptionsController.cs
@@ -131,7 +131,7 @@ namespace Horde.Server.Notifications
 				{
 					return BadRequest($"Invalid user id: '{subscription.UserId}'.");
 				}
-				if (authorizedUsers.Add(newUserId) && !_globalConfig.Value.Authorize(AdminAclAction.Impersonate, User))
+				if (authorizedUsers.Add(newUserId) && !_globalConfig.Value.Authorize(ServerAclAction.Impersonate, User))
 				{
 					return Forbid();
 				}
diff --git a/Engine/Source/Programs/Horde/Horde.Server/Server/AdminAclAction.cs b/Engine/Source/Programs/Horde/Horde.Server/Server/AdminAclAction.cs
--- a/Engine/Source/Programs/Horde/Horde.Server/Server/AdminAclAction.cs
+++ b/Engine/Source/Programs/Horde/Horde.Server/Server/AdminAclAction.cs
@@ -18,20 +18,5 @@ namespace Horde.Server.Server
 		/// Ability to write any data to the server.
 		/// </summary>
 		public static AclAction AdminWrite { get; } = new AclAction("AdminWrite");
-
-		/// <summary>
-		/// Ability to impersonate another user
-		/// </summary>
-		public static AclAction Impersonate { get; } = new AclAction("Impersonate");
-
-		/// <summary>
-		/// View estimated costs for particular operations
-		/// </summary>
-		public static AclAction ViewCosts { get; } = new AclAction("ViewCosts");
-
-		/// <summary>
-		/// Issue bearer token for the current user
-		/// </summary>
-		public static AclAction IssueBearerToken { get; } = new AclAction("IssueBearerToken");
 	}
 }
diff --git a/Engine/Source/Programs/Horde/Horde.Server/Server/AdminController.cs b/Engine/Source/Programs/Horde/Horde.Server/Server/AdminController.cs
--- a/Engine/Source/Programs/Horde/Horde.Server/Server/AdminController.cs
+++ b/Engine/Source/Programs/Horde/Horde.Server/Server/AdminController.cs
@@ -4,8 +4,11 @@ using System;
 using System.Collections.Generic;
 using System.Linq;
 using System.Security.Claims;
+using System.Text.Json;
 using System.Threading.Tasks;
 using Horde.Server.Acls;
+using Horde.Server.Projects;
+using Horde.Server.Secrets;
 using Horde.Server.Utilities;
 using Microsoft.AspNetCore.Authorization;
 using Microsoft.AspNetCore.Mvc;
@@ -75,29 +78,14 @@ namespace Horde.Server.Server
 		[Route("/api/v1/admin/token")]
 		public async Task<ActionResult<string>> GetTokenAsync()
 		{
-			if (!_globalConfig.Value.Authorize(AdminAclAction.IssueBearerToken, User))
+			if (!_globalConfig.Value.Authorize(ServerAclAction.IssueBearerToken, User))
 			{
-				return Forbid(AdminAclAction.IssueBearerToken);
+				return Forbid(ServerAclAction.IssueBearerToken);
 			}
 
 			return await _aclService.IssueBearerTokenAsync(User.Claims, GetDefaultExpiryTime());
 		}
 
-		/// <summary>
-		/// Returns the fully parsed config object.
-		/// </summary>
-		[HttpGet]
-		[Route("/api/v1/admin/config")]
-		public ActionResult<object> GetConfig()
-		{
-			if (!_globalConfig.Value.Authorize(AdminAclAction.AdminRead, User))
-			{
-				return Forbid(AdminAclAction.AdminRead);
-			}
-
-			return _globalConfig.Value;
-		}
-
 		/// <summary>
 		/// Issues a token for the given roles. Issues a token for the current user if not specified.
 		/// </summary>
diff --git a/Engine/Source/Programs/Horde/Horde.Server/Server/DebugController.cs b/Engine/Source/Programs/Horde/Horde.Server/Server/DebugController.cs
--- a/Engine/Source/Programs/Horde/Horde.Server/Server/DebugController.cs
+++ b/Engine/Source/Programs/Horde/Horde.Server/Server/DebugController.cs
@@ -20,6 +20,7 @@ using Horde.Server.Jobs;
 using Horde.Server.Jobs.Graphs;
 using Horde.Server.Jobs.Templates;
 using Horde.Server.Logs;
+using Horde.Server.Projects;
 using Horde.Server.Utilities;
 using JetBrains.Profiler.SelfApi;
 using Microsoft.AspNetCore.Authorization;
@@ -241,14 +242,12 @@ namespace Horde.Server.Server
 	/// </summary>
 	[ApiController]
 	[Authorize]
-	public class SecureDebugController : ControllerBase
+	public class SecureDebugController : HordeControllerBase
 	{
 		private static readonly Random s_random = new ();
 		
-		private readonly AclService _aclService;
 		private readonly MongoService _mongoService;
 		private readonly JobTaskSource _jobTaskSource;
-		private readonly ITemplateCollection _templateCollection;
 		private readonly IGraphCollection _graphCollection;
 		private readonly ILogFileCollection _logFileCollection;
 		private readonly IOptions<ServerSettings> _settings;
@@ -258,14 +257,12 @@ namespace Horde.Server.Server
 		/// <summary>
 		/// Constructor
 		/// </summary>
-		public SecureDebugController(AclService aclService, MongoService mongoService, JobTaskSource jobTaskSource,
-			ITemplateCollection templateCollection, IGraphCollection graphCollection,
+		public SecureDebugController(MongoService mongoService, JobTaskSource jobTaskSource,
+			IGraphCollection graphCollection,
 			ILogFileCollection logFileCollection, IOptions<ServerSettings> settings, IOptionsSnapshot<GlobalConfig> globalConfig, ILogger<SecureDebugController> logger)
 		{
-			_aclService = aclService;
 			_mongoService = mongoService;
 			_jobTaskSource = jobTaskSource;
-			_templateCollection = templateCollection;
 			_graphCollection = graphCollection;
 			_logFileCollection = logFileCollection;
 			_settings = settings;
@@ -316,26 +313,44 @@ namespace Horde.Server.Server
 		}
 
 		/// <summary>
-		/// Returns the complete config Horde uses
+		/// Returns the fully parsed config object.
 		/// </summary>
-		/// <returns>Information about the config</returns>
 		[HttpGet]
-		[Route("/api/v1/debug/config")]
-		public ActionResult GetConfig()
+		[Route("/api/v1/server/debug/appsettings")]
+		public ActionResult<object> GetAppSettings()
 		{
-			if (!_globalConfig.Value.Authorize(AdminAclAction.AdminRead, User))
+			if (!_globalConfig.Value.Authorize(ServerAclAction.ViewConfig, User))
 			{
-				return Forbid();
+				return Forbid(ServerAclAction.ViewConfig);
 			}
 
-			JsonSerializerOptions options = new JsonSerializerOptions();
-			options.WriteIndented = true;
+			return _globalConfig.Value.ServerSettings;
+		}
 
-			Startup.ConfigureJsonSerializer(options);
+		/// <summary>
+		/// Returns the fully parsed config object.
+		/// </summary>
+		[HttpGet]
+		[Route("/api/v1/server/debug/config")]
+		public ActionResult<object> GetConfig()
+		{
+			if (!_globalConfig.Value.Authorize(ServerAclAction.ViewConfig, User))
+			{
+				return Forbid(ServerAclAction.ViewConfig);
+			}
+
+			// Duplicate the config, so we can redact stuff that we don't want to return through the browser
+			byte[] data = JsonSerializer.SerializeToUtf8Bytes(_globalConfig.Value);
+			GlobalConfig config = JsonSerializer.Deserialize<GlobalConfig>(data)!;
 
-			return Ok(JsonSerializer.Serialize(_settings.Value, options));
+			foreach (ProjectConfig project in config.Projects)
+			{
+				project.Logo = null;
+			}
+
+			return config;
 		}
-		
+
 		/// <summary>
 		/// Generate log message of varying size
 		/// </summary>
diff --git a/Engine/Source/Programs/Horde/Horde.Server/Server/GlobalConfig.cs b/Engine/Source/Programs/Horde/Horde.Server/Server/GlobalConfig.cs
--- a/Engine/Source/Programs/Horde/Horde.Server/Server/GlobalConfig.cs
+++ b/Engine/Source/Programs/Horde/Horde.Server/Server/GlobalConfig.cs
@@ -380,7 +380,7 @@ namespace Horde.Server.Server
 			}
 			else
 			{
-				return this.Authorize(AdminAclAction.Impersonate, user);
+				return this.Authorize(ServerAclAction.Impersonate, user);
 			}
 		}
 
diff --git a/Engine/Source/Programs/Horde/Horde.Server/Server/ServerAclAction.cs b/Engine/Source/Programs/Horde/Horde.Server/Server/ServerAclAction.cs
new file mode 100644
--- /dev/null
+++ b/Engine/Source/Programs/Horde/Horde.Server/Server/ServerAclAction.cs
@@ -0,0 +1,32 @@
+// Copyright Epic Games, Inc. All Rights Reserved.
+
+using Horde.Server.Acls;
+
+namespace Horde.Server.Server
+{
+	/// <summary>
+	/// General server ACL actions
+	/// </summary>
+	public static class ServerAclAction
+	{
+		/// <summary>
+		/// Ability to impersonate another user
+		/// </summary>
+		public static AclAction Impersonate { get; } = new AclAction("Impersonate");
+
+		/// <summary>
+		/// View estimated costs for particular operations
+		/// </summary>
+		public static AclAction ViewCosts { get; } = new AclAction("ViewCosts");
+
+		/// <summary>
+		/// View estimated costs for particular operations
+		/// </summary>
+		public static AclAction ViewConfig { get; } = new AclAction("ViewConfig");
+
+		/// <summary>
+		/// Issue bearer token for the current user
+		/// </summary>
+		public static AclAction IssueBearerToken { get; } = new AclAction("IssueBearerToken");
+	}
+}
diff --git a/Engine/Source/Programs/Horde/Horde.Server/Server/ServerController.cs b/Engine/Source/Programs/Horde/Horde.Server/Server/ServerController.cs
--- a/Engine/Source/Programs/Horde/Horde.Server/Server/ServerController.cs
+++ b/Engine/Source/Programs/Horde/Horde.Server/Server/ServerController.cs
@@ -2,9 +2,13 @@
 
 using System.Diagnostics;
 using System.Reflection;
+using System.Text.Json;
 using System.Threading.Tasks;
+using Horde.Server.Acls;
 using Horde.Server.Agents;
+using Horde.Server.Projects;
 using Horde.Server.Tools;
+using Horde.Server.Utilities;
 using HordeCommon;
 using Microsoft.AspNetCore.Authorization;
 using Microsoft.AspNetCore.Hosting.Server;
@@ -20,7 +24,7 @@ namespace Horde.Server.Server
 	[ApiController]
 	[Authorize]
 	[Route("[controller]")]
-	public class ServerController : ControllerBase
+	public class ServerController : HordeControllerBase
 	{
 		readonly IToolCollection _toolCollection;
 		readonly IClock _clock;
@@ -35,7 +39,7 @@ namespace Horde.Server.Server
 			_clock = clock;
 			_globalConfig = globalConfig;
 		}
-		
+
 		/// <summary>
 		/// Get server version
 		/// </summary>
diff --git a/Engine/Source/Programs/Horde/Horde.Server/ServerSettings.cs b/Engine/Source/Programs/Horde/Horde.Server/ServerSettings.cs
--- a/Engine/Source/Programs/Horde/Horde.Server/ServerSettings.cs
+++ b/Engine/Source/Programs/Horde/Horde.Server/ServerSettings.cs
@@ -856,7 +856,7 @@ namespace Horde.Server
 		{
 			AclConfig defaultAcl = new AclConfig();
 			defaultAcl.Entries.Add(new AclEntryConfig(new AclClaimConfig(ClaimTypes.Role, "internal:AgentRegistration"), new[] { AgentAclAction.CreateAgent, SessionAclAction.CreateSession }));
-			defaultAcl.Entries.Add(new AclEntryConfig(HordeClaims.AgentRegistrationClaim, new[] { AgentAclAction.CreateAgent, SessionAclAction.CreateSession, AgentAclAction.UpdateAgent, AgentSoftwareAclAction.DownloadSoftware, PoolAclAction.CreatePool, PoolAclAction.UpdatePool, PoolAclAction.ViewPool, PoolAclAction.DeletePool, PoolAclAction.ListPools, StreamAclAction.ViewStream, ProjectAclAction.ViewProject, JobAclAction.ViewJob, AdminAclAction.ViewCosts }));
+			defaultAcl.Entries.Add(new AclEntryConfig(HordeClaims.AgentRegistrationClaim, new[] { AgentAclAction.CreateAgent, SessionAclAction.CreateSession, AgentAclAction.UpdateAgent, AgentSoftwareAclAction.DownloadSoftware, PoolAclAction.CreatePool, PoolAclAction.UpdatePool, PoolAclAction.ViewPool, PoolAclAction.DeletePool, PoolAclAction.ListPools, StreamAclAction.ViewStream, ProjectAclAction.ViewProject, JobAclAction.ViewJob, ServerAclAction.ViewCosts }));
 			defaultAcl.Entries.Add(new AclEntryConfig(HordeClaims.AgentRoleClaim, new[] { ProjectAclAction.ViewProject, StreamAclAction.ViewStream, LogAclAction.CreateEvent, AgentSoftwareAclAction.DownloadSoftware }));
 			defaultAcl.Entries.Add(new AclEntryConfig(HordeClaims.DownloadSoftwareClaim, new[] { AgentSoftwareAclAction.DownloadSoftware }));
 			defaultAcl.Entries.Add(new AclEntryConfig(HordeClaims.UploadSoftwareClaim, new[] { AgentSoftwareAclAction.UploadSoftware }));
