Horde: Create a unique ACL for access to the debug endpoint.

diff --git a/Engine/Source/Programs/Horde/Horde.Server/Server/DebugController.cs b/Engine/Source/Programs/Horde/Horde.Server/Server/DebugController.cs
--- a/Engine/Source/Programs/Horde/Horde.Server/Server/DebugController.cs
+++ b/Engine/Source/Programs/Horde/Horde.Server/Server/DebugController.cs
@@ -278,9 +278,9 @@ namespace Horde.Server.Server
 		[Route("/api/v1/debug/environment")]
 		public ActionResult GetServerEnvVars()
 		{
-			if (!_globalConfig.Value.Authorize(AdminAclAction.AdminRead, User))
+			if (!_globalConfig.Value.Authorize(ServerAclAction.Debug, User))
 			{
-				return Forbid();
+				return Forbid(ServerAclAction.Debug);
 			}
 
 			StringBuilder content = new StringBuilder();
@@ -304,9 +304,9 @@ namespace Horde.Server.Server
 		[Route("/api/v1/debug/queue")]
 		public ActionResult<object> GetQueueStatus()
 		{
-			if (!_globalConfig.Value.Authorize(AdminAclAction.AdminRead, User))
+			if (!_globalConfig.Value.Authorize(ServerAclAction.Debug, User))
 			{
-				return Forbid();
+				return Forbid(ServerAclAction.Debug);
 			}
 
 			return _jobTaskSource.GetStatus();
@@ -319,9 +319,9 @@ namespace Horde.Server.Server
 		[Route("/api/v1/server/debug/appsettings")]
 		public ActionResult<object> GetAppSettings()
 		{
-			if (!_globalConfig.Value.Authorize(ServerAclAction.ViewConfig, User))
+			if (!_globalConfig.Value.Authorize(ServerAclAction.Debug, User))
 			{
-				return Forbid(ServerAclAction.ViewConfig);
+				return Forbid(ServerAclAction.Debug);
 			}
 
 			return _globalConfig.Value.ServerSettings;
@@ -334,9 +334,9 @@ namespace Horde.Server.Server
 		[Route("/api/v1/server/debug/config")]
 		public ActionResult<object> GetConfig()
 		{
-			if (!_globalConfig.Value.Authorize(ServerAclAction.ViewConfig, User))
+			if (!_globalConfig.Value.Authorize(ServerAclAction.Debug, User))
 			{
-				return Forbid(ServerAclAction.ViewConfig);
+				return Forbid(ServerAclAction.Debug);
 			}
 
 			// Duplicate the config, so we can redact stuff that we don't want to return through the browser
@@ -364,9 +364,9 @@ namespace Horde.Server.Server
 			[FromQuery] int argCount = 0,
 			[FromQuery] int argLen = 10)
 		{
-			if (!_globalConfig.Value.Authorize(AdminAclAction.AdminRead, User))
+			if (!_globalConfig.Value.Authorize(ServerAclAction.Debug, User))
 			{
-				return Forbid();
+				return Forbid(ServerAclAction.Debug);
 			}
 
 			string RandomString(int length)
@@ -417,9 +417,9 @@ namespace Horde.Server.Server
 		[ProducesResponseType(200, Type = typeof(GetGraphResponse))]
 		public async Task<ActionResult<List<object>>> GetGraphsAsync([FromQuery] int? index = null, [FromQuery] int? count = null, [FromQuery] PropertyFilter? filter = null)
 		{
-			if (!_globalConfig.Value.Authorize(AdminAclAction.AdminRead, User))
+			if (!_globalConfig.Value.Authorize(ServerAclAction.Debug, User))
 			{
-				return Forbid();
+				return Forbid(ServerAclAction.Debug);
 			}
 
 			List<IGraph> graphs = await _graphCollection.FindAllAsync(null, index, count);
@@ -435,9 +435,9 @@ namespace Horde.Server.Server
 		[ProducesResponseType(200, Type = typeof(GetGraphResponse))]
 		public async Task<ActionResult<object>> GetGraphAsync(string graphId, [FromQuery] PropertyFilter? filter = null)
 		{
-			if (!_globalConfig.Value.Authorize(AdminAclAction.AdminRead, User))
+			if (!_globalConfig.Value.Authorize(ServerAclAction.Debug, User))
 			{
-				return Forbid();
+				return Forbid(ServerAclAction.Debug);
 			}
 
 			IGraph graph = await _graphCollection.GetAsync(ContentHash.Parse(graphId));
@@ -454,9 +454,9 @@ namespace Horde.Server.Server
 		[Route("/api/v1/debug/logs/{LogFileId}")]
 		public async Task<ActionResult<object>> GetLogAsync(LogId logFileId, [FromQuery] PropertyFilter? filter = null)
 		{
-			if (!_globalConfig.Value.Authorize(AdminAclAction.AdminRead, User))
+			if (!_globalConfig.Value.Authorize(ServerAclAction.Debug, User))
 			{
-				return Forbid();
+				return Forbid(ServerAclAction.Debug);
 			}
 
 			ILogFile? logFile = await _logFileCollection.GetLogFileAsync(logFileId, CancellationToken.None);
@@ -476,9 +476,9 @@ namespace Horde.Server.Server
 		[Route("/api/v1/debug/collections/{Name}")]
 		public async Task<ActionResult<object>> GetDocumentsAsync(string name, [FromQuery] string? filter = null, [FromQuery] int index = 0, [FromQuery] int count = 10)
 		{
-			if (!_globalConfig.Value.Authorize(AdminAclAction.AdminRead, User))
+			if (!_globalConfig.Value.Authorize(ServerAclAction.Debug, User))
 			{
-				return Forbid();
+				return Forbid(ServerAclAction.Debug);
 			}
 
 			IMongoCollection<Dictionary<string, object>> collection = _mongoService.GetCollection<Dictionary<string, object>>(name);
@@ -494,6 +494,11 @@ namespace Horde.Server.Server
 		[Route("/api/v1/debug/profiler/start")]
 		public async Task<ActionResult> StartProfiler()
 		{
+			if (!_globalConfig.Value.Authorize(ServerAclAction.Debug, User))
+			{
+				return Forbid(ServerAclAction.Debug);
+			}
+
 			await DotTrace.EnsurePrerequisiteAsync();
 
 			string snapshotDir = Path.Join(Path.GetTempPath(), "horde-profiler-snapshots");
@@ -518,6 +523,11 @@ namespace Horde.Server.Server
 		[Route("/api/v1/debug/profiler/stop")]
 		public ActionResult StopProfiler()
 		{
+			if (!_globalConfig.Value.Authorize(ServerAclAction.Debug, User))
+			{
+				return Forbid(ServerAclAction.Debug);
+			}
+
 			DotTrace.SaveData();
 			DotTrace.Detach();
 			return new ContentResult { ContentType = "text/plain", StatusCode = (int)HttpStatusCode.OK, Content = "Profiling session stopped" };
@@ -531,6 +541,11 @@ namespace Horde.Server.Server
 		[Route("/api/v1/debug/profiler/download")]
 		public ActionResult DownloadProfilingData()
 		{
+			if (!_globalConfig.Value.Authorize(ServerAclAction.Debug, User))
+			{
+				return Forbid(ServerAclAction.Debug);
+			}
+
 			string snapshotZipFile = DotTrace.GetCollectedSnapshotFilesArchive(false);
 			if (!System.IO.File.Exists(snapshotZipFile))
 			{
@@ -548,6 +563,11 @@ namespace Horde.Server.Server
 		[Route("/api/v1/debug/exception")]
 		public ActionResult ThrowException()
 		{
+			if (!_globalConfig.Value.Authorize(ServerAclAction.Debug, User))
+			{
+				return Forbid(ServerAclAction.Debug);
+			}
+
 			int numberArg = 42;
 			string stringArg = "hello";
 			throw new Exception($"Message: numberArg:{numberArg}, stringArg:{stringArg}");
diff --git a/Engine/Source/Programs/Horde/Horde.Server/Server/ServerAclAction.cs b/Engine/Source/Programs/Horde/Horde.Server/Server/ServerAclAction.cs
--- a/Engine/Source/Programs/Horde/Horde.Server/Server/ServerAclAction.cs
+++ b/Engine/Source/Programs/Horde/Horde.Server/Server/ServerAclAction.cs
@@ -10,19 +10,19 @@ namespace Horde.Server.Server
 	public static class ServerAclAction
 	{
 		/// <summary>
-		/// Ability to impersonate another user
+		/// Access to the debug endpoints
 		/// </summary>
-		public static AclAction Impersonate { get; } = new AclAction("Impersonate");
+		public static AclAction Debug { get; } = new AclAction("Debug");
 
 		/// <summary>
-		/// View estimated costs for particular operations
+		/// Ability to impersonate another user
 		/// </summary>
-		public static AclAction ViewCosts { get; } = new AclAction("ViewCosts");
+		public static AclAction Impersonate { get; } = new AclAction("Impersonate");
 
 		/// <summary>
 		/// View estimated costs for particular operations
 		/// </summary>
-		public static AclAction ViewConfig { get; } = new AclAction("ViewConfig");
+		public static AclAction ViewCosts { get; } = new AclAction("ViewCosts");
 
 		/// <summary>
 		/// Issue bearer token for the current user
