Horde: Propagate claims from the user that starts a job to tokens issued for agents executing leases related to that job.

diff --git a/Engine/Source/Programs/Horde/Horde.Server/Jobs/IJob.cs b/Engine/Source/Programs/Horde/Horde.Server/Jobs/IJob.cs
--- a/Engine/Source/Programs/Horde/Horde.Server/Jobs/IJob.cs
+++ b/Engine/Source/Programs/Horde/Horde.Server/Jobs/IJob.cs
@@ -663,6 +663,11 @@ namespace Horde.Server.Jobs
 		/// </summary>
 		public JobOptions? JobOptions { get; }
 
+		/// <summary>
+		/// Claims inherited from the user that started this job
+		/// </summary>
+		public IReadOnlyList<AclClaimConfig> Claims { get; }
+
 		/// <summary>
 		/// Largest value of the CombinedPriority value for batches in the ready state.
 		/// </summary>
diff --git a/Engine/Source/Programs/Horde/Horde.Server/Jobs/IJobCollection.cs b/Engine/Source/Programs/Horde/Horde.Server/Jobs/IJobCollection.cs
--- a/Engine/Source/Programs/Horde/Horde.Server/Jobs/IJobCollection.cs
+++ b/Engine/Source/Programs/Horde/Horde.Server/Jobs/IJobCollection.cs
@@ -17,6 +17,7 @@ using HordeCommon.Rpc.Tasks;
 using MongoDB.Bson;
 using Horde.Server.Jobs.Bisect;
 using System.Threading;
+using Horde.Server.Acls;
 
 namespace Horde.Server.Jobs
 {
@@ -55,6 +56,9 @@ namespace Horde.Server.Jobs
 		/// <inheritdoc cref="IJob.JobOptions"/>
 		public JobOptions? JobOptions { get; set; }
 
+		/// <inheritdoc cref="IJob.Claims"/>
+		public List<AclClaimConfig> Claims { get; set; } = new List<AclClaimConfig>();
+
 		/// <summary>
 		/// List of downstream job triggers
 		/// </summary>
diff --git a/Engine/Source/Programs/Horde/Horde.Server/Jobs/JobCollection.cs b/Engine/Source/Programs/Horde/Horde.Server/Jobs/JobCollection.cs
--- a/Engine/Source/Programs/Horde/Horde.Server/Jobs/JobCollection.cs
+++ b/Engine/Source/Programs/Horde/Horde.Server/Jobs/JobCollection.cs
@@ -283,6 +283,10 @@ namespace Horde.Server.Jobs
 			[BsonIgnoreIfNull]
 			public JobOptions? JobOptions { get; set; }
 
+			[BsonIgnoreIfNull]
+			public List<AclClaimConfig> Claims { get; set; }
+			IReadOnlyList<AclClaimConfig> IJob.Claims => Claims;
+
 			[BsonIgnoreIfNull]
 			public DateTime? CreateTimeUtc { get; set; }
 
@@ -352,6 +356,7 @@ namespace Horde.Server.Jobs
 				AutoSubmit = options.AutoSubmit ?? false;
 				UpdateIssues = options.UpdateIssues ?? (options.StartedByUserId == null && ( options.PreflightChange == 0 || options.PreflightChange == null));
 				PromoteIssuesByDefault = options.PromoteIssuesByDefault ?? false;
+				Claims = options.Claims;
 				JobOptions = options.JobOptions;
 				CreateTimeUtc = createTimeUtc;
 				ChainedJobs.AddRange(options.JobTriggers.Select(x => new ChainedJobDocument(x)));
diff --git a/Engine/Source/Programs/Horde/Horde.Server/Jobs/JobTaskSource.cs b/Engine/Source/Programs/Horde/Horde.Server/Jobs/JobTaskSource.cs
--- a/Engine/Source/Programs/Horde/Horde.Server/Jobs/JobTaskSource.cs
+++ b/Engine/Source/Programs/Horde/Horde.Server/Jobs/JobTaskSource.cs
@@ -835,7 +835,9 @@ namespace Horde.Server.Jobs
 			string storagePrefix = $"{job.StreamId}/{job.Change}-{job.Id}";
 			claims.Add(new AclClaimConfig(HordeClaimTypes.ReadNamespace, $"{Namespace.Artifacts}:{storagePrefix}"));
 			claims.Add(new AclClaimConfig(HordeClaimTypes.WriteNamespace, $"{Namespace.Artifacts}:{storagePrefix}"));
-		
+
+			claims.AddRange(job.Claims);
+
 			// Encode the payload
 			ExecuteJobTask task = new ExecuteJobTask();
 			task.JobId = job.Id.ToString();
diff --git a/Engine/Source/Programs/Horde/Horde.Server/Jobs/JobsController.cs b/Engine/Source/Programs/Horde/Horde.Server/Jobs/JobsController.cs
--- a/Engine/Source/Programs/Horde/Horde.Server/Jobs/JobsController.cs
+++ b/Engine/Source/Programs/Horde/Horde.Server/Jobs/JobsController.cs
@@ -194,6 +194,7 @@ namespace Horde.Server.Jobs
 			options.Priority = priority;
 			options.AutoSubmit = create.AutoSubmit;
 			options.UpdateIssues = updateIssues;
+			options.Claims.AddRange(User.Claims.Select(x => new AclClaimConfig(x)));
 			options.Arguments.AddRange(arguments);
 
 			foreach ((string key, string value) in environment)
