Horde: Fix initialization warning for job documents.

diff --git a/Engine/Source/Programs/Horde/Horde.Server/Jobs/JobCollection.cs b/Engine/Source/Programs/Horde/Horde.Server/Jobs/JobCollection.cs
--- a/Engine/Source/Programs/Horde/Horde.Server/Jobs/JobCollection.cs
+++ b/Engine/Source/Programs/Horde/Horde.Server/Jobs/JobCollection.cs
@@ -283,8 +283,7 @@ namespace Horde.Server.Jobs
 			[BsonIgnoreIfNull]
 			public JobOptions? JobOptions { get; set; }
 
-			[BsonIgnoreIfNull]
-			public List<AclClaimConfig> Claims { get; set; }
+			public List<AclClaimConfig> Claims { get; set; } = new List<AclClaimConfig>();
 			IReadOnlyList<AclClaimConfig> IJob.Claims => Claims;
 
 			[BsonIgnoreIfNull]
