Horde: Additional logging for autosdk view changes.

diff --git a/Engine/Source/Programs/Horde/Horde.Server/Agents/Pools/PoolUpdateService.cs b/Engine/Source/Programs/Horde/Horde.Server/Agents/Pools/PoolUpdateService.cs
--- a/Engine/Source/Programs/Horde/Horde.Server/Agents/Pools/PoolUpdateService.cs
+++ b/Engine/Source/Programs/Horde/Horde.Server/Agents/Pools/PoolUpdateService.cs
@@ -203,6 +203,11 @@ namespace Horde.Server.Agents.Pools
 					{
 						_logger.LogInformation("New workspaces for pool {Pool}:{Workspaces}", currentPool.Id, String.Join("", newWorkspaces.Select(x => $"\n  Identifier=\"{x.Identifier}\", Stream={x.Stream}")));
 
+						if (newAutoSdkConfig != null)
+						{
+							_logger.LogInformation("New autosdk view for pool {Pool}:{View}", currentPool.Id, String.Join("", newAutoSdkConfig.View.Select(x => $"\n  {x}"))); 
+						}
+
 						IPool? result = await _pools.TryUpdateAsync(currentPool, new UpdatePoolOptions { Workspaces = newWorkspaces, AutoSdkConfig = newAutoSdkConfig });
 						if (result == null)
 						{
