Horde: Speculative fix for build error on farm.

diff --git a/Engine/Source/Programs/Horde/Horde.Server/Agents/IAgent.cs b/Engine/Source/Programs/Horde/Horde.Server/Agents/IAgent.cs
--- a/Engine/Source/Programs/Horde/Horde.Server/Agents/IAgent.cs
+++ b/Engine/Source/Programs/Horde/Horde.Server/Agents/IAgent.cs
@@ -227,8 +227,8 @@ namespace Horde.Server.Agents
 		/// <summary>
 		/// Merge two workspaces together
 		/// </summary>
-		[return: NotNullIfNotNull(nameof(lhs))]
-		[return: NotNullIfNotNull(nameof(rhs))]
+		[return: NotNullIfNotNull("lhs")]
+		[return: NotNullIfNotNull("rhs")]
 		public static AutoSdkConfig? Merge(AutoSdkConfig? lhs, AutoSdkConfig? rhs)
 		{
 			if (lhs == null)
