Horde: Add missing comparison methods causing static analyzer warnings.

diff --git a/Engine/Source/Programs/Horde/Horde.Server/Jobs/IJobStepRef.cs b/Engine/Source/Programs/Horde/Horde.Server/Jobs/IJobStepRef.cs
--- a/Engine/Source/Programs/Horde/Horde.Server/Jobs/IJobStepRef.cs
+++ b/Engine/Source/Programs/Horde/Horde.Server/Jobs/IJobStepRef.cs
@@ -2,6 +2,7 @@
 
 using System;
 using System.Collections.Generic;
+using System.Diagnostics.CodeAnalysis;
 using Horde.Server.Agents;
 using Horde.Server.Agents.Pools;
 using Horde.Server.Jobs.Bisect;
@@ -17,7 +18,7 @@ namespace Horde.Server.Jobs
 	/// Unique id struct for JobStepRef objects. Includes a job id, batch id, and step id to uniquely identify the step.
 	/// </summary>
 	[BsonSerializer(typeof(JobStepRefIdSerializer))]
-	public struct JobStepRefId : IComparable<JobStepRefId>
+	public struct JobStepRefId : IEquatable<JobStepRefId>, IComparable<JobStepRefId>
 	{
 		/// <summary>
 		/// The job id
@@ -67,6 +68,12 @@ namespace Horde.Server.Jobs
 			return $"{JobId}:{BatchId}:{StepId}";
 		}
 
+		/// <inheritdoc/>
+		public override int GetHashCode() => HashCode.Combine(JobId, BatchId, StepId);
+
+		/// <inheritdoc/>
+		public override bool Equals([NotNullWhen(true)] object? obj) => obj is JobStepRefId other && Equals(other);
+
 		/// <inheritdoc/>
 		public bool Equals(JobStepRefId other)
 		{
@@ -90,6 +97,24 @@ namespace Horde.Server.Jobs
 
 			return StepId.Value.CompareTo(other.StepId.Value);
 		}
+
+		/// <inheritdoc/>
+		public static bool operator ==(JobStepRefId lhs, JobStepRefId rhs) => lhs.Equals(rhs);
+
+		/// <inheritdoc/>
+		public static bool operator !=(JobStepRefId lhs, JobStepRefId rhs) => !lhs.Equals(rhs);
+
+		/// <inheritdoc/>
+		public static bool operator <(JobStepRefId lhs, JobStepRefId rhs) => lhs.CompareTo(rhs) < 0;
+
+		/// <inheritdoc/>
+		public static bool operator >(JobStepRefId lhs, JobStepRefId rhs) => lhs.CompareTo(rhs) > 0;
+
+		/// <inheritdoc/>
+		public static bool operator <=(JobStepRefId lhs, JobStepRefId rhs) => lhs.CompareTo(rhs) <= 0;
+
+		/// <inheritdoc/>
+		public static bool operator >=(JobStepRefId lhs, JobStepRefId rhs) => lhs.CompareTo(rhs) >= 0;
 	}
 
 	/// <summary>
