Horde: Add detailed message showing exactly what will be synced for each workspace.

diff --git a/Engine/Source/Programs/Horde/Horde.Agent/Utility/WorkspaceInfo.cs b/Engine/Source/Programs/Horde/Horde.Agent/Utility/WorkspaceInfo.cs
--- a/Engine/Source/Programs/Horde/Horde.Agent/Utility/WorkspaceInfo.cs
+++ b/Engine/Source/Programs/Horde/Horde.Agent/Utility/WorkspaceInfo.cs
@@ -231,6 +231,7 @@ namespace Horde.Agent.Utility
 			await newRepository.RevertAsync(perforceClient, cancellationToken);
 
 			// Create the workspace info
+			logger.LogInformation("Syncing {ClientName} to {BaseDir} from {Server}, using view:{View}", clientName, workspaceDir, info.ServerAddress, String.Join("", view.Select(x => $"\n  {x}")));
 			return new WorkspaceInfo(perforceClient, hostName, streamName, streamView, metadataDir, workspaceDir, view, removeUntrackedFiles, newRepository);
 		}
 
