Horde: Add an interface for extending the secret service with external data.

diff --git a/Engine/Source/Programs/Horde/Horde.Server/Secrets/ISecretProvider.cs b/Engine/Source/Programs/Horde/Horde.Server/Secrets/ISecretProvider.cs
new file mode 100644
--- /dev/null
+++ b/Engine/Source/Programs/Horde/Horde.Server/Secrets/ISecretProvider.cs
@@ -0,0 +1,29 @@
+// Copyright Epic Games, Inc. All Rights Reserved.
+
+using System.Collections.Generic;
+using System.Security.Claims;
+using System.Threading;
+using System.Threading.Tasks;
+
+namespace Horde.Server.Secrets
+{
+	/// <summary>
+	/// Interface for a service that can provide secret data
+	/// </summary>
+	public interface ISecretProvider
+	{
+		/// <summary>
+		/// Name of this provider
+		/// </summary>
+		string Name { get; }
+
+		/// <summary>
+		/// Gets values for a particular secret
+		/// </summary>
+		/// <param name="config">Configuration for the secret to fetch</param>
+		/// <param name="user">User requesting the data</param>
+		/// <param name="data">Dictionary to receive any enumerated properties</param>
+		/// <param name="cancellationToken">Cancellation token for the operation</param>
+		Task GetValuesAsync(ExternalSecretConfig config, ClaimsPrincipal user, Dictionary<string, string> data, CancellationToken cancellationToken);
+	}
+}
diff --git a/Engine/Source/Programs/Horde/Horde.Server/Secrets/SecretConfig.cs b/Engine/Source/Programs/Horde/Horde.Server/Secrets/SecretConfig.cs
--- a/Engine/Source/Programs/Horde/Horde.Server/Secrets/SecretConfig.cs
+++ b/Engine/Source/Programs/Horde/Horde.Server/Secrets/SecretConfig.cs
@@ -1,5 +1,6 @@
 // Copyright Epic Games, Inc. All Rights Reserved.
 
+using System;
 using System.Collections.Generic;
 using System.Security.Claims;
 using System.Text.Json.Serialization;
@@ -27,6 +28,11 @@ namespace Horde.Server.Secrets
 		/// </summary>
 		public Dictionary<string, string> Data { get; set; } = new Dictionary<string, string>();
 
+		/// <summary>
+		/// Providers to source key/value pairs from
+		/// </summary>
+		public List<ExternalSecretConfig> Sources { get; set; } = new List<ExternalSecretConfig>();
+
 		/// <summary>
 		/// Defines access to this particular secret
 		/// </summary>
@@ -51,4 +57,30 @@ namespace Horde.Server.Secrets
 			return Acl?.Authorize(action, user) ?? GlobalConfig.Authorize(action, user);
 		}
 	}
+
+	/// <summary>
+	/// Configuration for an external secret provider
+	/// </summary>
+	public class ExternalSecretConfig
+	{
+		/// <summary>
+		/// Name of the provider to use
+		/// </summary>
+		public string Provider { get; set; } = String.Empty;
+
+		/// <summary>
+		/// Optional key indicating the parameter to set in the resulting data array
+		/// </summary>
+		public string? Key { get; set; }
+
+		/// <summary>
+		/// Optional value indicating what to fetch from the provider
+		/// </summary>
+		public string? Path { get; set; }
+
+		/// <summary>
+		/// Additional provider-specific arguments
+		/// </summary>
+		public Dictionary<string, string>? Arguments { get; set; }
+	}
 }
diff --git a/Engine/Source/Programs/Horde/Horde.Server/Secrets/SecretService.cs b/Engine/Source/Programs/Horde/Horde.Server/Secrets/SecretService.cs
new file mode 100644
--- /dev/null
+++ b/Engine/Source/Programs/Horde/Horde.Server/Secrets/SecretService.cs
@@ -0,0 +1,70 @@
+// Copyright Epic Games, Inc. All Rights Reserved.
+
+using System;
+using System.Collections.Generic;
+using System.Linq;
+using System.Security.Claims;
+using System.Threading;
+using System.Threading.Tasks;
+using Microsoft.Extensions.Logging;
+
+namespace Horde.Server.Secrets
+{
+	/// <summary>
+	/// Service for retrieving secrets
+	/// </summary>
+	public class SecretService
+	{
+		readonly Dictionary<string, ISecretProvider> _providers;
+		readonly ILogger _logger;
+
+		/// <summary>
+		/// Constructor
+		/// </summary>
+		/// <param name="providers"></param>
+		public SecretService(IEnumerable<ISecretProvider> providers, ILogger<SecretService> logger)
+		{
+			_providers = providers.ToDictionary(x => x.Name, x => x);
+			_logger = logger;
+		}
+
+		/// <summary>
+		/// Resolve a secret to concrete values
+		/// </summary>
+		/// <param name="config">Configuration for the secret</param>
+		/// <param name="user">User requesting the secret</param>
+		/// <param name="data">The resolved secret values</param>
+		/// <param name="cancellationToken">Cancellation token for the operation</param>
+		public async Task<bool> ResolveAsync(SecretConfig config, ClaimsPrincipal user, Dictionary<string, string> data, CancellationToken cancellationToken)
+		{
+			// Add the hard-coded secrets
+			foreach ((string key, string value) in config.Data)
+			{
+				data.Add(key, value);
+			}
+
+			// Fetch all the secrets from external providers
+			foreach (ExternalSecretConfig source in config.Sources)
+			{
+				ISecretProvider? provider;
+				if (!_providers.TryGetValue(source.Provider, out provider))
+				{
+					_logger.LogError("Unknown secret provider provider {Provider} for secret {SecretId}", source.Provider, config.Id);
+					return false;
+				}
+
+				try
+				{
+					await provider.GetValuesAsync(source, user, data, cancellationToken);
+				}
+				catch (Exception ex)
+				{
+					_logger.LogError(ex, "Exception while fetching secret {SecretId} from {Provider}: {Message}", config.Id, source.Provider, ex.Message);
+					return false;
+				}
+			}
+
+			return true;
+		}
+	}
+}
diff --git a/Engine/Source/Programs/Horde/Horde.Server/Startup.cs b/Engine/Source/Programs/Horde/Horde.Server/Startup.cs
--- a/Engine/Source/Programs/Horde/Horde.Server/Startup.cs
+++ b/Engine/Source/Programs/Horde/Horde.Server/Startup.cs
@@ -510,6 +510,9 @@ namespace Horde.Server
 				services.AddSingleton<IExternalIssueService, DefaultExternalIssueService>();
 			}
 
+			// Secret providers
+			services.AddSingleton<SecretService>();
+
 			// Storage providers
 			services.AddSingleton<IStorageBackendProvider, StorageBackendProvider>();
 			services.AddSingleton(sp => CreateStorageBackend(sp, settings.LogStorage).ForType<PersistentLogStorage>());
