Horde: Fix static analysis warning.

diff --git a/Engine/Source/Programs/Horde/Horde.Server/Secrets/SecretService.cs b/Engine/Source/Programs/Horde/Horde.Server/Secrets/SecretService.cs
--- a/Engine/Source/Programs/Horde/Horde.Server/Secrets/SecretService.cs
+++ b/Engine/Source/Programs/Horde/Horde.Server/Secrets/SecretService.cs
@@ -21,7 +21,6 @@ namespace Horde.Server.Secrets
 		/// <summary>
 		/// Constructor
 		/// </summary>
-		/// <param name="providers"></param>
 		public SecretService(IEnumerable<ISecretProvider> providers, ILogger<SecretService> logger)
 		{
 			_providers = providers.ToDictionary(x => x.Name, x => x);
