Horde: Include stream names in sync messages.

diff --git a/Engine/Source/Programs/Horde/Horde.Agent/Utility/WorkspaceInfo.cs b/Engine/Source/Programs/Horde/Horde.Agent/Utility/WorkspaceInfo.cs
--- a/Engine/Source/Programs/Horde/Horde.Agent/Utility/WorkspaceInfo.cs
+++ b/Engine/Source/Programs/Horde/Horde.Agent/Utility/WorkspaceInfo.cs
@@ -231,7 +231,7 @@ namespace Horde.Agent.Utility
 			await newRepository.RevertAsync(perforceClient, cancellationToken);
 
 			// Create the workspace info
-			logger.LogInformation("Syncing {ClientName} to {BaseDir} from {Server}, using view:{View}", clientName, workspaceDir, info.ServerAddress, String.Join("", view.Select(x => $"\n  {x}")));
+			logger.LogInformation("Syncing {ClientName} to {BaseDir} from {Server}, using stream {Stream} and view:{View}", clientName, workspaceDir, info.ServerAddress, streamName, String.Join("", view.Select(x => $"\n  {x}")));
 			return new WorkspaceInfo(perforceClient, hostName, streamName, streamView, metadataDir, workspaceDir, view, removeUntrackedFiles, newRepository);
 		}
 
