Horde: Variable rename.

diff --git a/Engine/Source/Programs/Horde/Horde.Server/Agents/IAgent.cs b/Engine/Source/Programs/Horde/Horde.Server/Agents/IAgent.cs
--- a/Engine/Source/Programs/Horde/Horde.Server/Agents/IAgent.cs
+++ b/Engine/Source/Programs/Horde/Horde.Server/Agents/IAgent.cs
@@ -922,8 +922,8 @@ namespace Horde.Server.Agents
 			else
 			{
 				// Try to get the matching workspace type
-				WorkspaceConfig? workspaceType;
-				if (!streamConfig.WorkspaceTypes.TryGetValue(agentType.Workspace, out workspaceType))
+				WorkspaceConfig? workspaceConfig;
+				if (!streamConfig.WorkspaceTypes.TryGetValue(agentType.Workspace, out workspaceConfig))
 				{
 					workspace = null;
 					autoSdkConfig = null;
@@ -932,11 +932,11 @@ namespace Horde.Server.Agents
 
 				// Get the workspace identifier
 				string identifier;
-				if (workspaceType.Identifier != null)
+				if (workspaceConfig.Identifier != null)
 				{
-					identifier = workspaceType.Identifier;
+					identifier = workspaceConfig.Identifier;
 				}
-				else if (workspaceType.Incremental)
+				else if (workspaceConfig.Incremental)
 				{
 					identifier = $"{streamConfig.GetEscapedName()}+{agentType.Workspace}";
 				}
@@ -946,11 +946,11 @@ namespace Horde.Server.Agents
 				}
 
 				// Create the new workspace
-				workspace = new AgentWorkspace(workspaceType.Cluster, workspaceType.UserName, identifier, workspaceType.Stream ?? streamConfig.Name, workspaceType.View, workspaceType.Incremental, workspaceType.Method);
+				workspace = new AgentWorkspace(workspaceConfig.Cluster, workspaceConfig.UserName, identifier, workspaceConfig.Stream ?? streamConfig.Name, workspaceConfig.View, workspaceConfig.Incremental, workspaceConfig.Method);
 
-				if (workspaceType.UseAutoSdk)
+				if (workspaceConfig.UseAutoSdk)
 				{
-					autoSdkConfig = new AutoSdkConfig(Enumerable.Concat(streamConfig.AutoSdkView ?? Enumerable.Empty<string>(), workspaceType.AutoSdkView ?? Enumerable.Empty<string>()));
+					autoSdkConfig = new AutoSdkConfig(Enumerable.Concat(streamConfig.AutoSdkView ?? Enumerable.Empty<string>(), workspaceConfig.AutoSdkView ?? Enumerable.Empty<string>()));
 				}
 				else
 				{
