Horde: Adding job redirector
#rnx
#skipci

diff --git a/Engine/Source/Programs/Horde/HordeDashboard/src/App.tsx b/Engine/Source/Programs/Horde/HordeDashboard/src/App.tsx
--- a/Engine/Source/Programs/Horde/HordeDashboard/src/App.tsx
+++ b/Engine/Source/Programs/Horde/HordeDashboard/src/App.tsx
@@ -32,6 +32,7 @@ import { UtilizationReportView } from './components/UtilizationReportView';
 import hordePlugins from './Plugins';
 import { modeColors, preloadFonts } from './styles/Styles';
 import { StepIssueReportTest } from './components/test/IssueStepReport';
+import { JobRedirector } from './components/JobRedirector';
 
 
 let router: any;
@@ -132,7 +133,8 @@ const Main: React.FC = () => {
                { path: "index", element: (dashboard.user?.dashboardFeatures?.showLandingPage === true) ? <DocView /> : <UserHomeView /> },
                { path: "project/:projectId", element: <ProjectHome /> },
                { path: "pools", element: <PoolView /> },
-               { path: "job/:jobId", element: <JobDetailViewV2 /> },
+               { path: "job/:jobId", element: <JobDetailViewV2 /> },               
+               { path: "job", element: <JobRedirector /> },
                { path: "log/:logId", element: <LogView /> },
                { path: "testreport/:testdataId", element: <TestReportView /> },
                { path: "stream/:streamId", element: <StreamView /> },
diff --git a/Engine/Source/Programs/Horde/HordeDashboard/src/components/JobRedirector.tsx b/Engine/Source/Programs/Horde/HordeDashboard/src/components/JobRedirector.tsx
new file mode 100644
--- /dev/null
+++ b/Engine/Source/Programs/Horde/HordeDashboard/src/components/JobRedirector.tsx
@@ -0,0 +1,62 @@
+// Copyright Epic Games, Inc. All Rights Reserved.
+
+import React, { useState } from 'react';
+import { useNavigate } from 'react-router-dom';
+import backend from '../backend';
+import { useQuery } from './JobDetailCommon';
+import ErrorHandler from './ErrorHandler';
+
+function setError(message: string) {
+
+   console.error(message);
+
+   message = `${message}\n\nReferring URL: ${encodeURI(window.location.href)}`
+
+   ErrorHandler.set({
+
+      title: `Error handling Preflight`,
+      message: message
+
+   }, true);
+}
+
+// redirect from external source, where horde stream id, etc are not known by that application
+export const JobRedirector: React.FC = () => {
+
+   const [state, setState] = useState({ backendQueried: false })
+
+   const navigate = useNavigate();
+   const query = useQuery();
+
+   const streamId = !query.get("streamId") ? "" : query.get("streamId")!;
+   const templateId = !query.get("templateId") ? "" : query.get("templateId")!;
+   const nodeName = !query.get("nodeName") ? "" : query.get("nodeName")!;
+
+   if (!streamId || !templateId || !nodeName) {
+      setError("Please specify streamId, templateId, and nodeName");
+      return null;
+   }
+
+   if (!state.backendQueried) {
+
+      backend.getJobStepHistory(streamId, nodeName, 1, templateId).then(results => {
+         if (!results?.length) {
+            setError(`No history results for ${streamId} / ${templateId} / ${nodeName}`);
+         }
+
+         const result = results[0];
+
+         navigate(`/job/${result.jobId}?step=${result.stepId}`, { replace: true });
+
+      }).catch(reason => {
+         setError(`Error getting last job for ${streamId} / ${templateId} / ${nodeName}: Please check parameters\n${reason}`);
+      });
+
+      return null;
+   }
+
+   setState({ backendQueried: true })
+
+   return null;
+
+}
