Horde: Remove ObjectId extension methods.

diff --git a/Engine/Source/Programs/Horde/Horde.Server/Jobs/Artifacts/ArtifactsControllerV1.cs b/Engine/Source/Programs/Horde/Horde.Server/Jobs/Artifacts/ArtifactsControllerV1.cs
--- a/Engine/Source/Programs/Horde/Horde.Server/Jobs/Artifacts/ArtifactsControllerV1.cs
+++ b/Engine/Source/Programs/Horde/Horde.Server/Jobs/Artifacts/ArtifactsControllerV1.cs
@@ -109,7 +109,7 @@ namespace Horde.Server.Jobs.Artifacts
 		[Route("/api/v1/artifacts/{artifactId}")]
 		public async Task<ActionResult<CreateJobArtifactResponse>> UpdateArtifact(string artifactId, IFormFile file)
 		{
-			IArtifactV1? artifact = await _artifactCollection.GetArtifactAsync(artifactId.ToObjectId());
+			IArtifactV1? artifact = await _artifactCollection.GetArtifactAsync(ObjectId.Parse(artifactId));
 			if (artifact == null)
 			{
 				return NotFound();
@@ -182,7 +182,7 @@ namespace Horde.Server.Jobs.Artifacts
 		[ProducesResponseType(typeof(GetJobArtifactResponse), 200)]
 		public async Task<ActionResult<object>> GetArtifact(string artifactId, bool code = false, [FromQuery] PropertyFilter? filter = null)
 		{
-			IArtifactV1? artifact = await _artifactCollection.GetArtifactAsync(artifactId.ToObjectId());
+			IArtifactV1? artifact = await _artifactCollection.GetArtifactAsync(ObjectId.Parse(artifactId));
 			if (artifact == null)
 			{
 				return NotFound();
@@ -206,7 +206,7 @@ namespace Horde.Server.Jobs.Artifacts
 		[Route("/api/v1/artifacts/{artifactId}/data")]
 		public async Task<ActionResult> GetArtifactData(string artifactId)
 		{
-			IArtifactV1? artifact = await _artifactCollection.GetArtifactAsync(artifactId.ToObjectId());
+			IArtifactV1? artifact = await _artifactCollection.GetArtifactAsync(ObjectId.Parse(artifactId));
 			if (artifact == null)
 			{
 				return NotFound();
@@ -273,7 +273,7 @@ namespace Horde.Server.Jobs.Artifacts
 			JwtSecurityTokenHandler handler = new JwtSecurityTokenHandler();
 			ClaimsPrincipal principal = handler.ValidateToken(code, parameters, out _);
 
-			IArtifactV1? artifact = await _artifactCollection.GetArtifactAsync(artifactId.ToObjectId());
+			IArtifactV1? artifact = await _artifactCollection.GetArtifactAsync(ObjectId.Parse(artifactId));
 			if (artifact == null)
 			{
 				return NotFound();
@@ -328,7 +328,7 @@ namespace Horde.Server.Jobs.Artifacts
 				foreach (string artifactId in artifactZipRequest.ArtifactIds)
 				{
 					IArtifactV1? artifact;
-					if (idToArtifact.TryGetValue(artifactId.ToObjectId(), out artifact))
+					if (idToArtifact.TryGetValue(ObjectId.Parse(artifactId), out artifact))
 					{
 						zipArtifacts.Add(artifact);
 					}
diff --git a/Engine/Source/Programs/Horde/Horde.Server/Jobs/JobRpcService.cs b/Engine/Source/Programs/Horde/Horde.Server/Jobs/JobRpcService.cs
--- a/Engine/Source/Programs/Horde/Horde.Server/Jobs/JobRpcService.cs
+++ b/Engine/Source/Programs/Horde/Horde.Server/Jobs/JobRpcService.cs
@@ -257,7 +257,7 @@ namespace Horde.Server.Jobs
 		/// <returns>Information about the new agent</returns>
 		public async Task<RpcGetJobResponse> GetJob(GetJobRequest request, ServerCallContext context)
 		{
-			JobId jobIdValue = new JobId(request.JobId.ToObjectId());
+			JobId jobIdValue = new JobId(ObjectId.Parse(request.JobId));
 
 			IJob? job = await _jobService.GetJobAsync(jobIdValue);
 			if (job == null)
@@ -918,7 +918,7 @@ namespace Horde.Server.Jobs
 
 			IJobStepBatch batch = AuthorizeBatch(job, request.BatchId.ToSubResourceId(), context);
 
-			Report newReport = new Report { Name = request.Name, Placement = request.Placement, ArtifactId = request.ArtifactId.ToObjectId() };
+			Report newReport = new Report { Name = request.Name, Placement = request.Placement, ArtifactId = ObjectId.Parse(request.ArtifactId) };
 			if (request.Scope == ReportScope.Job)
 			{
 				_logger.LogDebug("Adding report to job {JobId}: {Name} -> {ArtifactId}", job.Id, request.Name, request.ArtifactId);
diff --git a/Engine/Source/Programs/Horde/Horde.Server/Jobs/TestData/TestDataController.cs b/Engine/Source/Programs/Horde/Horde.Server/Jobs/TestData/TestDataController.cs
--- a/Engine/Source/Programs/Horde/Horde.Server/Jobs/TestData/TestDataController.cs
+++ b/Engine/Source/Programs/Horde/Horde.Server/Jobs/TestData/TestDataController.cs
@@ -364,7 +364,7 @@ namespace Horde.Server.Jobs.TestData
 		[ProducesResponseType(typeof(GetTestDataResponse), 200)]
 		public async Task<ActionResult<object>> GetTestDataAsync(string testDataId, [FromQuery] PropertyFilter? filter = null)
 		{
-			ITestData? testData = await _testDataCollection.GetAsync(testDataId.ToObjectId());
+			ITestData? testData = await _testDataCollection.GetAsync(ObjectId.Parse(testDataId));
 			if (testData == null)
 			{
 				return NotFound();
diff --git a/Engine/Source/Programs/Horde/Horde.Server/Startup.cs b/Engine/Source/Programs/Horde/Horde.Server/Startup.cs
--- a/Engine/Source/Programs/Horde/Horde.Server/Startup.cs
+++ b/Engine/Source/Programs/Horde/Horde.Server/Startup.cs
@@ -221,7 +221,11 @@ namespace Horde.Server
 				{
 					throw new InvalidDataException("Unable to parse object id");
 				}
-				return str.ToObjectId();
+				if (str.Length == 0)
+				{
+					return ObjectId.Empty;
+				}
+				return ObjectId.Parse(str);
 			}
 
 			public override void Write(Utf8JsonWriter writer, ObjectId objectId, JsonSerializerOptions options)
diff --git a/Engine/Source/Programs/Horde/Horde.Server/Utilities/ObjectIdHelpers.cs b/Engine/Source/Programs/Horde/Horde.Server/Utilities/ObjectIdHelpers.cs
deleted file mode 100644
