Horde: Fix broken serialization of global config in debug controller

diff --git a/Engine/Source/Programs/Horde/Horde.Server/Configuration/ConfigService.cs b/Engine/Source/Programs/Horde/Horde.Server/Configuration/ConfigService.cs
--- a/Engine/Source/Programs/Horde/Horde.Server/Configuration/ConfigService.cs
+++ b/Engine/Source/Programs/Horde/Horde.Server/Configuration/ConfigService.cs
@@ -373,6 +373,16 @@ namespace Horde.Server.Configuration
 			}
 		}
 
+		internal byte[] Serialize(GlobalConfig config)
+		{
+			return JsonSerializer.SerializeToUtf8Bytes(config, _jsonOptions);
+		}
+		
+		internal GlobalConfig? Deserialize(byte[] data)
+		{
+			return JsonSerializer.Deserialize<GlobalConfig>(data, _jsonOptions);
+		}
+
 		/// <summary>
 		/// Writes information about the current snapshot to the logger
 		/// </summary>
diff --git a/Engine/Source/Programs/Horde/Horde.Server/Server/DebugController.cs b/Engine/Source/Programs/Horde/Horde.Server/Server/DebugController.cs
--- a/Engine/Source/Programs/Horde/Horde.Server/Server/DebugController.cs
+++ b/Engine/Source/Programs/Horde/Horde.Server/Server/DebugController.cs
@@ -16,6 +16,7 @@ using System.Threading.Tasks;
 using System.Web;
 using EpicGames.Core;
 using Horde.Server.Acls;
+using Horde.Server.Configuration;
 using Horde.Server.Jobs;
 using Horde.Server.Jobs.Graphs;
 using Horde.Server.Jobs.Templates;
@@ -247,6 +248,7 @@ namespace Horde.Server.Server
 		private static readonly Random s_random = new ();
 		
 		private readonly MongoService _mongoService;
+		private readonly ConfigService _configService;
 		private readonly JobTaskSource _jobTaskSource;
 		private readonly IGraphCollection _graphCollection;
 		private readonly ILogFileCollection _logFileCollection;
@@ -257,11 +259,12 @@ namespace Horde.Server.Server
 		/// <summary>
 		/// Constructor
 		/// </summary>
-		public SecureDebugController(MongoService mongoService, JobTaskSource jobTaskSource,
+		public SecureDebugController(MongoService mongoService, ConfigService configService, JobTaskSource jobTaskSource,
 			IGraphCollection graphCollection,
 			ILogFileCollection logFileCollection, IOptions<ServerSettings> settings, IOptionsSnapshot<GlobalConfig> globalConfig, ILogger<SecureDebugController> logger)
 		{
 			_mongoService = mongoService;
+			_configService = configService;
 			_jobTaskSource = jobTaskSource;
 			_graphCollection = graphCollection;
 			_logFileCollection = logFileCollection;
@@ -340,8 +343,8 @@ namespace Horde.Server.Server
 			}
 
 			// Duplicate the config, so we can redact stuff that we don't want to return through the browser
-			byte[] data = JsonSerializer.SerializeToUtf8Bytes(_globalConfig.Value);
-			GlobalConfig config = JsonSerializer.Deserialize<GlobalConfig>(data)!;
+			byte[] data = _configService.Serialize(_globalConfig.Value);
+			GlobalConfig config = _configService.Deserialize(data)!;
 
 			foreach (ProjectConfig project in config.Projects)
 			{
