Horde: Lower AWS ASG lifecycle update interval to every minute

diff --git a/Engine/Source/Programs/Horde/Horde.Server/Agents/Fleet/AwsAutoScalingLifecycle.cs b/Engine/Source/Programs/Horde/Horde.Server/Agents/Fleet/AwsAutoScalingLifecycle.cs
--- a/Engine/Source/Programs/Horde/Horde.Server/Agents/Fleet/AwsAutoScalingLifecycle.cs
+++ b/Engine/Source/Programs/Horde/Horde.Server/Agents/Fleet/AwsAutoScalingLifecycle.cs
@@ -55,7 +55,7 @@ public sealed class AwsAutoScalingLifecycleService : IHostedService, IDisposable
 	/// <summary>
 	/// How often instance lifecycles for auto-scaling should be updated with AWS 
 	/// </summary>
-	public TimeSpan LifecycleUpdaterInterval { get; } = TimeSpan.FromMinutes(20);
+	public TimeSpan LifecycleUpdaterInterval { get; } = TimeSpan.FromMinutes(1);
 
 	private const string RedisKey = "aws-asg-lifecycle";
 	private readonly AgentService _agentService;
