Horde: Remove any AWS ASG lifecycle events that cannot be sent

Avoids the queue from filling up. If they fail, there's little use of keeping them around.

diff --git a/Engine/Source/Programs/Horde/Horde.Server/Agents/Fleet/AwsAutoScalingLifecycle.cs b/Engine/Source/Programs/Horde/Horde.Server/Agents/Fleet/AwsAutoScalingLifecycle.cs
--- a/Engine/Source/Programs/Horde/Horde.Server/Agents/Fleet/AwsAutoScalingLifecycle.cs
+++ b/Engine/Source/Programs/Horde/Horde.Server/Agents/Fleet/AwsAutoScalingLifecycle.cs
@@ -319,8 +319,19 @@ public sealed class AwsAutoScalingLifecycleService : IHostedService, IDisposable
 			ActionResult result = agent.Status == AgentStatus.Stopped ? ActionResult.Abandon : ActionResult.Continue;
 			tasks.Add(Task.Run(async () =>
 			{
-				await SendCompleteLifecycleActionAsync(infos[agent.Id].Event, result, cancellationToken);
-				if (result == ActionResult.Abandon)
+				bool deleteEvent;
+				try
+				{
+					await SendCompleteLifecycleActionAsync(infos[agent.Id].Event, result, cancellationToken);
+					deleteEvent = result == ActionResult.Abandon;
+				}
+				catch (Exception e)
+				{
+					_logger.LogError(e, "Unable to complete the AWS ASG lifecycle action. Reason: {Message}", e.Message);
+					deleteEvent = true;
+				}
+				
+				if (deleteEvent)
 				{
 					await redis.HashDeleteAsync(RedisKey, agent.Id.ToString());
 				}
@@ -349,7 +360,11 @@ public sealed class AwsAutoScalingLifecycleService : IHostedService, IDisposable
 
 	private enum ActionResult
 	{
-		Continue, Abandon
+		/// <inheritdoc cref="AwsAutoScalingLifecycleService.ActionContinue"/>
+		Continue,
+		
+		/// <inheritdoc cref="AwsAutoScalingLifecycleService.ActionAbandon"/>
+		Abandon
 	}
 }
 
