Horde: Make batch errors more obvious with error color and different icon
#rnx
#skipci

diff --git a/Engine/Source/Programs/Horde/HordeDashboard/src/components/jobDetailsV2/JobDetailViewSteps.tsx b/Engine/Source/Programs/Horde/HordeDashboard/src/components/jobDetailsV2/JobDetailViewSteps.tsx
--- a/Engine/Source/Programs/Horde/HordeDashboard/src/components/jobDetailsV2/JobDetailViewSteps.tsx
+++ b/Engine/Source/Programs/Horde/HordeDashboard/src/components/jobDetailsV2/JobDetailViewSteps.tsx
@@ -3,7 +3,7 @@ import { observer } from "mobx-react-lite";
 import React, { useEffect, useState } from "react";
 import { Link, useParams } from "react-router-dom";
 import { BatchData, GetBatchResponse, JobStepBatchError, JobStepBatchState, JobStepOutcome, JobStepState, NodeData, StepData } from "../../backend/Api";
-import dashboard from "../../backend/Dashboard";
+import dashboard, { StatusColor } from "../../backend/Dashboard";
 import { getDetailStyle } from "../../backend/JobDetails";
 import { ISideRailLink } from "../../base/components/SideRail";
 import { getBatchInitElapsed, getNiceTime, getStepElapsed, getStepETA, getStepFinishTime, getStepPercent, getStepStartTime, getStepTimingDelta } from "../../base/utilities/timeUtils";
@@ -211,7 +211,7 @@ export const StepsPanelInner: React.FC<{ jobDetails: JobDetailsV2, depStepId?: s
 
    function renderAgentRow(item: StepItem, index?: number, column?: IColumn) {
 
-      let error = false;  //(!!jobDetails.events.find(e => e.severity === EventSeverity.Error && e.logId === item.batch?.logId)) || (item.batch?.error !== JobStepBatchError.None);
+      let error = !!item.batch?.startTime && item.batch?.error !== JobStepBatchError.None;
       let warning = false; //!!jobDetails.events.find(e => e.severity === EventSeverity.Warning && e.logId === item.batch?.logId);
 
       if (error) {
@@ -223,8 +223,6 @@ export const StepsPanelInner: React.FC<{ jobDetails: JobDetailsV2, depStepId?: s
          url = `/job/${jobId}?agenttype=${encodeURIComponent(item.agentType)}&agentpool=${encodeURIComponent(item.agentPool)}`;
       }
 
-      const style = getDetailStyle(JobStepState.Completed, error ? JobStepOutcome.Failure : JobStepOutcome.Warnings);
-
       if (!item.agentId && column!.key === 'ViewLogColumn') {
          return <div />;
       }
@@ -237,24 +235,27 @@ export const StepsPanelInner: React.FC<{ jobDetails: JobDetailsV2, depStepId?: s
          batchText = item.agentId ?? (item.batch?.state ?? "Unassigned");
       }
 
+      const statusColors = dashboard.getStatusColors();
+      const errorColor = error ? statusColors.get(StatusColor.Failure) : undefined;
+
       switch (column!.key) {
 
          case 'Name':
             if (!item.agentId) {
                if (item.agentPool && (item.batch?.state === JobStepBatchState.Ready || item.batch?.error === JobStepBatchError.NoAgentsOnline || item.batch?.error === JobStepBatchError.NoAgentsInPool)) {
                   return <Stack horizontal disableShrink={true} horizontalAlign="start">
-                     {(error || warning) && <Icon styles={{ root: { paddingTop: 3, paddingRight: 8 } }} className={style!.className} iconName="Square" />}
-                     <Link to={`/pools?pool=${encodeURI(item.agentPool)}`} ><Text nowrap={true}>{batchText}</Text></Link>
-                  </Stack>                     
+                     {(error || warning) && <Icon styles={{ root: { paddingTop: 3, paddingRight: 8, color: errorColor } }} iconName="Warning" />}
+                     <Link to={`/pools?pool=${encodeURI(item.agentPool)}`} ><Text styles={{ root: { color: errorColor } }} nowrap={true}>{batchText}</Text></Link>
+                  </Stack>
                }
                return <Stack horizontal disableShrink={true} horizontalAlign="start">
-                  {(error || warning) && <Icon styles={{ root: { paddingTop: 3, paddingRight: 8 } }} className={style!.className} iconName="Square" />}
-                  <Text nowrap={true}>{batchText}</Text>
+                  {(error || warning) && <Icon styles={{ root: { paddingTop: 3, paddingRight: 8, color: errorColor } }} iconName="Warning" />}
+                  <Text styles={{ root: { color: errorColor } }} nowrap={true}>{batchText}</Text>
                </Stack>;
             } else {
-               return <Stack horizontal disableShrink={true}>{(error || warning) && <Icon styles={{ root: { paddingTop: 3, paddingRight: 8 } }} className={style!.className} iconName="Square" />}
+               return <Stack horizontal disableShrink={true}>{(error || warning) && <Icon styles={{ root: { paddingTop: 3, paddingRight: 8, color: errorColor } }} iconName="Warning" />}
                   <Stack>
-                     <Link to={url} onClick={(ev) => { ev.stopPropagation(); ev.preventDefault(); setLastSelectedAgent(item.agentId); }}><Text styles={{ root: { cursor: 'pointer' } }} >{batchText}</Text></Link>
+                     <Link to={url} onClick={(ev) => { ev.stopPropagation(); ev.preventDefault(); setLastSelectedAgent(item.agentId); }}><Text styles={{ root: { cursor: 'pointer', color: errorColor } }} >{batchText}</Text></Link>
                   </Stack>
                </Stack>
             }
@@ -660,7 +661,7 @@ export const StepsPanelV2: React.FC<{ jobDetails: JobDetailsV2, depStepId?: stri
    }
 
    const sideRail = depStepId ? depSideRail : stepsSideRail;
-   
+
    // do not use useQuery() hook as will negatively impact rendering   
    const query = new URLSearchParams(window.location.search);
    const batchFilter = query.get("batch");
