Horde: Change to error icon for agent errors
#skipci
#rnx

diff --git a/Engine/Source/Programs/Horde/HordeDashboard/src/components/jobDetailsV2/JobDetailViewSteps.tsx b/Engine/Source/Programs/Horde/HordeDashboard/src/components/jobDetailsV2/JobDetailViewSteps.tsx
--- a/Engine/Source/Programs/Horde/HordeDashboard/src/components/jobDetailsV2/JobDetailViewSteps.tsx
+++ b/Engine/Source/Programs/Horde/HordeDashboard/src/components/jobDetailsV2/JobDetailViewSteps.tsx
@@ -244,16 +244,16 @@ export const StepsPanelInner: React.FC<{ jobDetails: JobDetailsV2, depStepId?: s
             if (!item.agentId) {
                if (item.agentPool && (item.batch?.state === JobStepBatchState.Ready || item.batch?.error === JobStepBatchError.NoAgentsOnline || item.batch?.error === JobStepBatchError.NoAgentsInPool)) {
                   return <Stack horizontal disableShrink={true} horizontalAlign="start">
-                     {(error || warning) && <Icon styles={{ root: { paddingTop: 3, paddingRight: 8, color: errorColor } }} iconName="Warning" />}
+                     {(error || warning) && <Icon styles={{ root: { paddingTop: 3, paddingRight: 8, color: errorColor } }} iconName="Error" />}
                      <Link to={`/pools?pool=${encodeURI(item.agentPool)}`} ><Text styles={{ root: { color: errorColor } }} nowrap={true}>{batchText}</Text></Link>
                   </Stack>
                }
                return <Stack horizontal disableShrink={true} horizontalAlign="start">
-                  {(error || warning) && <Icon styles={{ root: { paddingTop: 3, paddingRight: 8, color: errorColor } }} iconName="Warning" />}
+                  {(error || warning) && <Icon styles={{ root: { paddingTop: 3, paddingRight: 8, color: errorColor } }} iconName="Error" />}
                   <Text styles={{ root: { color: errorColor } }} nowrap={true}>{batchText}</Text>
                </Stack>;
             } else {
-               return <Stack horizontal disableShrink={true}>{(error || warning) && <Icon styles={{ root: { paddingTop: 3, paddingRight: 8, color: errorColor } }} iconName="Warning" />}
+               return <Stack horizontal disableShrink={true}>{(error || warning) && <Icon styles={{ root: { paddingTop: 3, paddingRight: 8, color: errorColor } }} iconName="Error" />}
                   <Stack>
                      <Link to={url} onClick={(ev) => { ev.stopPropagation(); ev.preventDefault(); setLastSelectedAgent(item.agentId); }}><Text styles={{ root: { cursor: 'pointer', color: errorColor } }} >{batchText}</Text></Link>
                   </Stack>
