Horde: Allow job options to be set for jobs created through REST API

diff --git a/Engine/Source/Programs/Horde/Horde.Server/Jobs/JobMessages.cs b/Engine/Source/Programs/Horde/Horde.Server/Jobs/JobMessages.cs
--- a/Engine/Source/Programs/Horde/Horde.Server/Jobs/JobMessages.cs
+++ b/Engine/Source/Programs/Horde/Horde.Server/Jobs/JobMessages.cs
@@ -13,6 +13,7 @@ using Horde.Server.Streams;
 using Horde.Server.Users;
 using Horde.Server.Utilities;
 using HordeCommon;
+using HordeCommon.Rpc.Tasks;
 
 namespace Horde.Server.Jobs
 {
@@ -83,6 +84,11 @@ namespace Horde.Server.Jobs
 		/// The preflight changelist number
 		/// </summary>
 		public int? PreflightChange { get; set; }
+		
+		/// <summary>
+		/// Job options
+		/// </summary>
+		public JobOptions? JobOptions { get; set; }
 
 		/// <summary>
 		/// Priority for the job
diff --git a/Engine/Source/Programs/Horde/Horde.Server/Jobs/JobsController.cs b/Engine/Source/Programs/Horde/Horde.Server/Jobs/JobsController.cs
--- a/Engine/Source/Programs/Horde/Horde.Server/Jobs/JobsController.cs
+++ b/Engine/Source/Programs/Horde/Horde.Server/Jobs/JobsController.cs
@@ -196,6 +196,7 @@ namespace Horde.Server.Jobs
 			options.UpdateIssues = updateIssues;
 			options.Claims.AddRange(User.Claims.Select(x => new AclClaimConfig(x)));
 			options.Arguments.AddRange(arguments);
+			options.JobOptions ??= create.JobOptions;
 
 			foreach ((string key, string value) in environment)
 			{
