Horde: Rename NodeLocator -> BundleNodeLocator.

diff --git a/Engine/Source/Programs/Horde/Horde.Agent/Commands/Compute/CppComputeCommand.cs b/Engine/Source/Programs/Horde/Horde.Agent/Commands/Compute/CppComputeCommand.cs
--- a/Engine/Source/Programs/Horde/Horde.Agent/Commands/Compute/CppComputeCommand.cs
+++ b/Engine/Source/Programs/Horde/Horde.Agent/Commands/Compute/CppComputeCommand.cs
@@ -55,7 +55,7 @@ namespace Horde.Agent.Commands.Compute
 
 			// Create a sandbox from the data to be uploaded
 			MemoryStorageClient storage = new MemoryStorageClient();
-			NodeLocator sandbox = await CreateSandboxAsync(TaskFile, storage, cancellationToken);
+			BundleNodeLocator sandbox = await CreateSandboxAsync(TaskFile, storage, cancellationToken);
 
 			// Open a socket and upload the sandbox
 			using (AgentMessageChannel channel = lease.Socket.CreateAgentMessageChannel(ControlChannelId, 4 * 1024 * 1024, _logger))
@@ -75,7 +75,7 @@ namespace Horde.Agent.Commands.Compute
 			return true;
 		}
 
-		static async Task<NodeLocator> CreateSandboxAsync(FileReference taskFile, IStorageClient storage, CancellationToken cancellationToken)
+		static async Task<BundleNodeLocator> CreateSandboxAsync(FileReference taskFile, IStorageClient storage, CancellationToken cancellationToken)
 		{
 			await using IStorageWriter writer = storage.CreateWriter();
 
diff --git a/Engine/Source/Programs/Horde/Horde.Server.Tests/BlobStoreTests.cs b/Engine/Source/Programs/Horde/Horde.Server.Tests/BlobStoreTests.cs
--- a/Engine/Source/Programs/Horde/Horde.Server.Tests/BlobStoreTests.cs
+++ b/Engine/Source/Programs/Horde/Horde.Server.Tests/BlobStoreTests.cs
@@ -120,7 +120,7 @@ namespace Horde.Server.Tests
 			for(int idx = 0; idx < 2; idx++)
 			{
 				RefName refName = new RefName("hello");
-				await store.WriteRefTargetAsync(refName, new NodeLocator(bundle3.Header.Exports[0].Hash, locator3, 0));
+				await store.WriteRefTargetAsync(refName, new BundleNodeLocator(bundle3.Header.Exports[0].Hash, locator3, 0));
 				BlobHandle refTarget = await store.ReadRefTargetAsync(refName);
 				Assert.AreEqual(locator3, refTarget.GetLocator().Blob);
 			}
@@ -133,7 +133,7 @@ namespace Horde.Server.Tests
 
 			Bundle bundle1 = CreateTestBundle(new byte[] { 1, 2, 3 }, Array.Empty<BlobLocator>());
 			BlobLocator locator1 = await store.WriteBundleAsync(bundle1);
-			NodeLocator target = new NodeLocator(bundle1.Header.Exports[0].Hash, locator1, 0);
+			BundleNodeLocator target = new BundleNodeLocator(bundle1.Header.Exports[0].Hash, locator1, 0);
 
 			await store.WriteRefTargetAsync("test-ref-1", target);
 			await store.WriteRefTargetAsync("test-ref-2", target, new RefOptions { Lifetime = TimeSpan.FromMinutes(30.0), Extend = true });
@@ -162,7 +162,7 @@ namespace Horde.Server.Tests
 			Assert.AreEqual(default, await TryReadRefTargetAsync(store, "test-ref-3"));
 		}
 
-		static async Task<NodeLocator> TryReadRefTargetAsync(IStorageClient store, RefName name)
+		static async Task<BundleNodeLocator> TryReadRefTargetAsync(IStorageClient store, RefName name)
 		{
 			BlobHandle? handle = await store.TryReadRefTargetAsync(name);
 			if (handle == null)
diff --git a/Engine/Source/Programs/Horde/Horde.Server.Tests/GcServiceTests.cs b/Engine/Source/Programs/Horde/Horde.Server.Tests/GcServiceTests.cs
--- a/Engine/Source/Programs/Horde/Horde.Server.Tests/GcServiceTests.cs
+++ b/Engine/Source/Programs/Horde/Horde.Server.Tests/GcServiceTests.cs
@@ -42,7 +42,7 @@ namespace Horde.Server.Tests
 				int blobIdx = (int)(random.NextDouble() * blobs.Length);
 				if (roots.Add(blobs[blobIdx]))
 				{
-					await store.WriteRefTargetAsync(new RefName($"ref-{idx}"), new NodeLocator(IoHash.Zero, blobs[blobIdx], 0));
+					await store.WriteRefTargetAsync(new RefName($"ref-{idx}"), new BundleNodeLocator(IoHash.Zero, blobs[blobIdx], 0));
 				}
 			}
 
diff --git a/Engine/Source/Programs/Horde/Horde.Server.Tests/StorageServiceTests.cs b/Engine/Source/Programs/Horde/Horde.Server.Tests/StorageServiceTests.cs
--- a/Engine/Source/Programs/Horde/Horde.Server.Tests/StorageServiceTests.cs
+++ b/Engine/Source/Programs/Horde/Horde.Server.Tests/StorageServiceTests.cs
@@ -44,20 +44,20 @@ namespace Horde.Server.Tests
 			Bundle bundle = new Bundle(header, Array.Empty<ReadOnlyMemory<byte>>());
 			BlobLocator locator = await client.WriteBundleAsync(bundle);
 
-			await client.AddAliasAsync("foo", new NodeLocator(hash1, locator, 0));
-			await client.AddAliasAsync("foo", new NodeLocator(hash1, locator, 1));
-			await client.AddAliasAsync("bar", new NodeLocator(hash2, locator, 2));
+			await client.AddAliasAsync("foo", new BundleNodeLocator(hash1, locator, 0));
+			await client.AddAliasAsync("foo", new BundleNodeLocator(hash1, locator, 1));
+			await client.AddAliasAsync("bar", new BundleNodeLocator(hash2, locator, 2));
 
 			List<BlobHandle> handles;
 			
 			handles = await client.FindNodesAsync("foo").ToListAsync();
 			Assert.AreEqual(2, handles.Count);
-			Assert.AreEqual(new NodeLocator(hash1, locator, 0), handles[0].GetLocator());
-			Assert.AreEqual(new NodeLocator(hash1, locator, 1), handles[1].GetLocator());
+			Assert.AreEqual(new BundleNodeLocator(hash1, locator, 0), handles[0].GetLocator());
+			Assert.AreEqual(new BundleNodeLocator(hash1, locator, 1), handles[1].GetLocator());
 
 			handles = await client.FindNodesAsync("bar").ToListAsync();
 			Assert.AreEqual(1, handles.Count);
-			Assert.AreEqual(new NodeLocator(hash2, locator, 2), handles[0].GetLocator());
+			Assert.AreEqual(new BundleNodeLocator(hash2, locator, 2), handles[0].GetLocator());
 		}
 	}
 }
diff --git a/Engine/Source/Programs/Horde/Horde.Server/Logs/LogRpcService.cs b/Engine/Source/Programs/Horde/Horde.Server/Logs/LogRpcService.cs
--- a/Engine/Source/Programs/Horde/Horde.Server/Logs/LogRpcService.cs
+++ b/Engine/Source/Programs/Horde/Horde.Server/Logs/LogRpcService.cs
@@ -52,7 +52,7 @@ namespace Horde.Server.Logs
 
 			IStorageClientImpl store = await _storageService.GetClientAsync(Namespace.Logs, context.CancellationToken);
 			_logger.LogInformation("Updating {LogId} to node {RefTarget} (lines: {LineCount}, complete: {Complete})", request.LogId, request.Target, request.LineCount, request.Complete);
-			await store.WriteRefTargetAsync(new RefName(request.LogId), NodeLocator.Parse(request.Target));
+			await store.WriteRefTargetAsync(new RefName(request.LogId), BundleNodeLocator.Parse(request.Target));
 
 			await _logFileCollection.UpdateLineCountAsync(logFile, request.LineCount, request.Complete, CancellationToken.None);
 
diff --git a/Engine/Source/Programs/Horde/Horde.Server/Storage/StorageController.cs b/Engine/Source/Programs/Horde/Horde.Server/Storage/StorageController.cs
--- a/Engine/Source/Programs/Horde/Horde.Server/Storage/StorageController.cs
+++ b/Engine/Source/Programs/Horde/Horde.Server/Storage/StorageController.cs
@@ -364,7 +364,7 @@ namespace Horde.Server.Storage
 			}
 
 			IStorageClientImpl client = await _storageService.GetClientAsync(namespaceId, cancellationToken);
-			NodeLocator target = new NodeLocator(request.Hash, request.Blob, request.ExportIdx);
+			BundleNodeLocator target = new BundleNodeLocator(request.Hash, request.Blob, request.ExportIdx);
 			await client.WriteRefTargetAsync(refName, target, request.Options, cancellationToken);
 
 			return Ok();
@@ -415,7 +415,7 @@ namespace Horde.Server.Storage
 				return new NotFoundResult();
 			}
 
-			NodeLocator locator = target.GetLocator();
+			BundleNodeLocator locator = target.GetLocator();
 			string link = $"/api/v1/storage/{namespaceId}/nodes/{locator.Blob}?export={locator.ExportIdx}";
 			return new ReadRefResponse(target, link);
 		}
@@ -573,7 +573,7 @@ namespace Horde.Server.Storage
 
 			object content;
 
-			BlobData nodeData = await reader.ReadNodeDataAsync(new NodeLocator(export.Hash, locator, exportIdx), cancellationToken);
+			BlobData nodeData = await reader.ReadNodeDataAsync(new BundleNodeLocator(export.Hash, locator, exportIdx), cancellationToken);
 
 			Node node = Node.Deserialize(nodeData);
 			switch (node)
@@ -605,7 +605,7 @@ namespace Horde.Server.Storage
 
 		static string GetNodeLink(string linkBase, BlobHandle handle) => GetNodeLink(linkBase, handle.GetLocator());
 		
-		static string GetNodeLink(string linkBase, NodeLocator locator) => $"{linkBase}/nodes/{locator.Blob}?export={locator.ExportIdx}";
+		static string GetNodeLink(string linkBase, BundleNodeLocator locator) => $"{linkBase}/nodes/{locator.Blob}?export={locator.ExportIdx}";
 
 		static Type? GetNodeType(Guid typeGuid)
 		{
diff --git a/Engine/Source/Programs/Horde/Horde.Server/Storage/StorageService.cs b/Engine/Source/Programs/Horde/Horde.Server/Storage/StorageService.cs
--- a/Engine/Source/Programs/Horde/Horde.Server/Storage/StorageService.cs
+++ b/Engine/Source/Programs/Horde/Horde.Server/Storage/StorageService.cs
@@ -64,13 +64,13 @@ namespace Horde.Server.Storage
 		bool SupportsRedirects { get; }
 
 		/// <inheritdoc cref="IStorageClient.AddAliasAsync(Utf8String, BlobHandle, CancellationToken)"/>
-		Task AddAliasAsync(Utf8String name, NodeLocator target, CancellationToken cancellationToken = default);
+		Task AddAliasAsync(Utf8String name, BundleNodeLocator target, CancellationToken cancellationToken = default);
 
 		/// <inheritdoc cref="IStorageClient.RemoveAliasAsync(Utf8String, BlobHandle, CancellationToken)"/>
-		Task RemoveAliasAsync(Utf8String name, NodeLocator target, CancellationToken cancellationToken = default); 
+		Task RemoveAliasAsync(Utf8String name, BundleNodeLocator target, CancellationToken cancellationToken = default); 
 
 		/// <inheritdoc cref="IStorageClient.WriteRefTargetAsync(RefName, BlobHandle, RefOptions?, CancellationToken)"/>
-		Task WriteRefTargetAsync(RefName name, NodeLocator handle, RefOptions? options = null, CancellationToken cancellationToken = default);
+		Task WriteRefTargetAsync(RefName name, BundleNodeLocator handle, RefOptions? options = null, CancellationToken cancellationToken = default);
 
 		/// <summary>
 		/// Gets a redirect for a read request
@@ -227,18 +227,18 @@ namespace Horde.Server.Storage
 			public override Task AddAliasAsync(Utf8String name, BlobHandle handle, CancellationToken cancellationToken = default) => _outer.AddAliasAsync(NamespaceId, name, handle.GetLocator(), cancellationToken);
 
 			/// <inheritdoc/>
-			public Task AddAliasAsync(Utf8String name, NodeLocator locator, CancellationToken cancellationToken = default) => _outer.AddAliasAsync(NamespaceId, name, locator, cancellationToken);
+			public Task AddAliasAsync(Utf8String name, BundleNodeLocator locator, CancellationToken cancellationToken = default) => _outer.AddAliasAsync(NamespaceId, name, locator, cancellationToken);
 
 			/// <inheritdoc/>
 			public override Task RemoveAliasAsync(Utf8String name, BlobHandle handle, CancellationToken cancellationToken = default) => _outer.RemoveAliasAsync(NamespaceId, name, handle.GetLocator(), cancellationToken);
 
 			/// <inheritdoc/>
-			public Task RemoveAliasAsync(Utf8String name, NodeLocator locator, CancellationToken cancellationToken = default) => _outer.RemoveAliasAsync(NamespaceId, name, locator, cancellationToken);
+			public Task RemoveAliasAsync(Utf8String name, BundleNodeLocator locator, CancellationToken cancellationToken = default) => _outer.RemoveAliasAsync(NamespaceId, name, locator, cancellationToken);
 
 			/// <inheritdoc/>
 			public override async IAsyncEnumerable<BlobHandle> FindNodesAsync(Utf8String alias, [EnumeratorCancellation] CancellationToken cancellationToken = default)
 			{
-				await foreach (NodeLocator locator in _outer.FindNodesAsync(NamespaceId, alias, cancellationToken))
+				await foreach (BundleNodeLocator locator in _outer.FindNodesAsync(NamespaceId, alias, cancellationToken))
 				{
 					yield return new FlushedNodeHandle(TreeReader, locator);
 				}
@@ -251,7 +251,7 @@ namespace Horde.Server.Storage
 			/// <inheritdoc/>
 			public override async Task<BlobHandle?> TryReadRefTargetAsync(RefName name, RefCacheTime cacheTime = default, CancellationToken cancellationToken = default)
 			{
-				NodeLocator? locator = await _outer.TryReadRefTargetAsync(NamespaceId, name, cacheTime, cancellationToken);
+				BundleNodeLocator? locator = await _outer.TryReadRefTargetAsync(NamespaceId, name, cacheTime, cancellationToken);
 				if (locator == null)
 				{
 					return null;
@@ -262,12 +262,12 @@ namespace Horde.Server.Storage
 			/// <inheritdoc/>
 			public override async Task WriteRefTargetAsync(RefName name, BlobHandle target, RefOptions? options = null, CancellationToken cancellationToken = default)
 			{
-				NodeLocator locator = await target.FlushAsync(cancellationToken);
+				BundleNodeLocator locator = await target.FlushAsync(cancellationToken);
 				await _outer.WriteRefTargetAsync(NamespaceId, name, locator, options, cancellationToken);
 			}
 
 			/// <inheritdoc/>
-			public Task WriteRefTargetAsync(RefName name, NodeLocator target, RefOptions? options = null, CancellationToken cancellationToken = default) => _outer.WriteRefTargetAsync(NamespaceId, name, target, options, cancellationToken);
+			public Task WriteRefTargetAsync(RefName name, BundleNodeLocator target, RefOptions? options = null, CancellationToken cancellationToken = default) => _outer.WriteRefTargetAsync(NamespaceId, name, target, options, cancellationToken);
 
 			/// <inheritdoc/>
 			public override Task DeleteRefAsync(RefName name, CancellationToken cancellationToken = default) => _outer.DeleteRefAsync(NamespaceId, name, cancellationToken);
@@ -395,7 +395,7 @@ namespace Horde.Server.Storage
 			public TimeSpan? Lifetime { get; set; }
 
 			[BsonIgnore]
-			public NodeLocator Target => new NodeLocator(Hash, BlobLocator, ExportIdx);
+			public BundleNodeLocator Target => new BundleNodeLocator(Hash, BlobLocator, ExportIdx);
 
 			public RefInfo()
 			{
@@ -403,7 +403,7 @@ namespace Horde.Server.Storage
 				BlobLocator = BlobLocator.Empty;
 			}
 
-			public RefInfo(NamespaceId namespaceId, RefName name, NodeLocator target, ObjectId blobInfoId)
+			public RefInfo(NamespaceId namespaceId, RefName name, BundleNodeLocator target, ObjectId blobInfoId)
 			{
 				NamespaceId = namespaceId;
 				Name = name;
@@ -702,7 +702,7 @@ namespace Horde.Server.Storage
 		/// <param name="target">Target node for the alias</param>
 		/// <param name="cancellationToken">Cancellation token for the operation</param>
 		/// <returns>Sequence of thandles</returns>
-		async Task AddAliasAsync(NamespaceId namespaceId, Utf8String alias, NodeLocator target, CancellationToken cancellationToken = default)
+		async Task AddAliasAsync(NamespaceId namespaceId, Utf8String alias, BundleNodeLocator target, CancellationToken cancellationToken = default)
 		{
 			BlobInfo? blobInfo = await _blobCollection.Find(x => x.NamespaceId == namespaceId && x.Path == target.Blob.Path).FirstOrDefaultAsync(cancellationToken);
 			if (blobInfo == null)
@@ -728,7 +728,7 @@ namespace Horde.Server.Storage
 		/// <param name="target">Target node for the alias</param>
 		/// <param name="cancellationToken">Cancellation token for the operation</param>
 		/// <returns>Sequence of thandles</returns>
-		Task RemoveAliasAsync(NamespaceId namespaceId, Utf8String alias, NodeLocator target, CancellationToken cancellationToken = default)
+		Task RemoveAliasAsync(NamespaceId namespaceId, Utf8String alias, BundleNodeLocator target, CancellationToken cancellationToken = default)
 		{
 			throw new NotSupportedException();
 		}
@@ -740,7 +740,7 @@ namespace Horde.Server.Storage
 		/// <param name="alias">Alias for the node</param>
 		/// <param name="cancellationToken">Cancellation token for the operation</param>
 		/// <returns>Sequence of thandles</returns>
-		async IAsyncEnumerable<NodeLocator> FindNodesAsync(NamespaceId namespaceId, Utf8String alias, [EnumeratorCancellation] CancellationToken cancellationToken = default)
+		async IAsyncEnumerable<BundleNodeLocator> FindNodesAsync(NamespaceId namespaceId, Utf8String alias, [EnumeratorCancellation] CancellationToken cancellationToken = default)
 		{
 			await foreach (BlobInfo blobInfo in _blobCollection.Find(x => x.NamespaceId == namespaceId && x.Exports!.Any(y => y.Alias == alias)).ToAsyncEnumerable(cancellationToken))
 			{
@@ -750,7 +750,7 @@ namespace Horde.Server.Storage
 					{
 						if (exportInfo.Alias == alias)
 						{
-							yield return new NodeLocator(exportInfo.Hash, blobInfo.Locator, exportInfo.Index);
+							yield return new BundleNodeLocator(exportInfo.Hash, blobInfo.Locator, exportInfo.Index);
 						}
 					}
 				}
@@ -842,7 +842,7 @@ namespace Horde.Server.Storage
 		}
 
 		/// <inheritdoc/>
-		async Task<NodeLocator?> TryReadRefTargetAsync(NamespaceId namespaceId, RefName name, RefCacheTime cacheTime = default, CancellationToken cancellationToken = default)
+		async Task<BundleNodeLocator?> TryReadRefTargetAsync(NamespaceId namespaceId, RefName name, RefCacheTime cacheTime = default, CancellationToken cancellationToken = default)
 		{
 			RefCacheValue entry;
 			if (!_cache.TryGetValue(name, out entry) || RefCacheTime.IsStaleCacheEntry(entry.Time, cacheTime))
@@ -874,7 +874,7 @@ namespace Horde.Server.Storage
 		}
 
 		/// <inheritdoc/>
-		async Task WriteRefTargetAsync(NamespaceId namespaceId, RefName name, NodeLocator target, RefOptions? options = null, CancellationToken cancellationToken = default)
+		async Task WriteRefTargetAsync(NamespaceId namespaceId, RefName name, BundleNodeLocator target, RefOptions? options = null, CancellationToken cancellationToken = default)
 		{
 			BlobInfo? newBlobInfo = await _blobCollection.Find(x => x.NamespaceId == namespaceId && x.Path == target.Blob.Path).FirstOrDefaultAsync(cancellationToken);
 			if (newBlobInfo == null)
diff --git a/Engine/Source/Programs/Horde/Horde.Server/Tools/IToolCollection.cs b/Engine/Source/Programs/Horde/Horde.Server/Tools/IToolCollection.cs
--- a/Engine/Source/Programs/Horde/Horde.Server/Tools/IToolCollection.cs
+++ b/Engine/Source/Programs/Horde/Horde.Server/Tools/IToolCollection.cs
@@ -42,7 +42,7 @@ namespace Horde.Server.Tools
 		/// <param name="globalConfig">The current configuration</param>
 		/// <param name="cancellationToken">Cancellation token for the operation</param>
 		/// <returns>Updated tool document, or null if it does not exist</returns>
-		Task<ITool?> CreateDeploymentAsync(ITool tool, ToolDeploymentConfig options, NodeLocator locator, GlobalConfig globalConfig, CancellationToken cancellationToken);
+		Task<ITool?> CreateDeploymentAsync(ITool tool, ToolDeploymentConfig options, BundleNodeLocator locator, GlobalConfig globalConfig, CancellationToken cancellationToken);
 
 		/// <summary>
 		/// Updates the state of the current deployment
diff --git a/Engine/Source/Programs/Horde/Horde.Server/Tools/ToolCollection.cs b/Engine/Source/Programs/Horde/Horde.Server/Tools/ToolCollection.cs
--- a/Engine/Source/Programs/Horde/Horde.Server/Tools/ToolCollection.cs
+++ b/Engine/Source/Programs/Horde/Horde.Server/Tools/ToolCollection.cs
@@ -275,7 +275,7 @@ namespace Horde.Server.Tools
 		/// <param name="globalConfig">The current configuration</param>
 		/// <param name="cancellationToken">Cancellation token for the operation</param>
 		/// <returns>Updated tool document, or null if it does not exist</returns>
-		public async Task<ITool?> CreateDeploymentAsync(ITool tool, ToolDeploymentConfig options, NodeLocator locator, GlobalConfig globalConfig, CancellationToken cancellationToken)
+		public async Task<ITool?> CreateDeploymentAsync(ITool tool, ToolDeploymentConfig options, BundleNodeLocator locator, GlobalConfig globalConfig, CancellationToken cancellationToken)
 		{
 			ToolDeploymentId deploymentId = ToolDeploymentId.GenerateNewId();
 			RefName refName = new RefName($"{tool.Id}/{deploymentId}");
diff --git a/Engine/Source/Programs/Horde/Horde.Server/Tools/ToolsController.cs b/Engine/Source/Programs/Horde/Horde.Server/Tools/ToolsController.cs
--- a/Engine/Source/Programs/Horde/Horde.Server/Tools/ToolsController.cs
+++ b/Engine/Source/Programs/Horde/Horde.Server/Tools/ToolsController.cs
@@ -129,7 +129,7 @@ namespace Horde.Server.Tools
 		/// <summary>
 		/// Node for downloading this deployment
 		/// </summary>
-		public NodeLocator Locator { get; }
+		public BundleNodeLocator Locator { get; }
 
 		/// <summary>
 		/// Constructor
@@ -304,7 +304,7 @@ namespace Horde.Server.Tools
 
 			ToolDeploymentConfig options = new ToolDeploymentConfig { Version = request.Version, Duration = TimeSpan.FromMinutes(request.Duration ?? 0.0), CreatePaused = request.CreatePaused ?? false };
 
-			tool = await _toolCollection.CreateDeploymentAsync(tool, options, NodeLocator.Parse(request.Node), _globalConfig.Value, cancellationToken);
+			tool = await _toolCollection.CreateDeploymentAsync(tool, options, BundleNodeLocator.Parse(request.Node), _globalConfig.Value, cancellationToken);
 			if (tool == null)
 			{
 				return NotFound(id);
diff --git a/Engine/Source/Programs/Horde/Horde/Commands/Bundles/BundleDump.cs b/Engine/Source/Programs/Horde/Horde/Commands/Bundles/BundleDump.cs
--- a/Engine/Source/Programs/Horde/Horde/Commands/Bundles/BundleDump.cs
+++ b/Engine/Source/Programs/Horde/Horde/Commands/Bundles/BundleDump.cs
@@ -67,7 +67,7 @@ namespace Horde.Commands.Bundles
 					for(int referenceIdx = 0; referenceIdx < export.References.Count; referenceIdx++)
 					{
 						BundleExportRef exportRef = export.References[referenceIdx];
-						NodeLocator nodeLocator = new NodeLocator(exportRef.Hash, header.Imports[exportRef.ImportIdx], exportRef.NodeIdx);
+						BundleNodeLocator nodeLocator = new BundleNodeLocator(exportRef.Hash, header.Imports[exportRef.ImportIdx], exportRef.NodeIdx);
 						logger.LogInformation("            REF {RefIdx,-3} -> {Node}", referenceIdx, nodeLocator);
 					}
 				}
diff --git a/Engine/Source/Programs/Shared/EpicGames.Horde/Compute/AgentMessage.cs b/Engine/Source/Programs/Shared/EpicGames.Horde/Compute/AgentMessage.cs
--- a/Engine/Source/Programs/Shared/EpicGames.Horde/Compute/AgentMessage.cs
+++ b/Engine/Source/Programs/Shared/EpicGames.Horde/Compute/AgentMessage.cs
@@ -186,7 +186,7 @@ namespace EpicGames.Horde.Compute
 	/// </summary>
 	/// <param name="Name">Path to extract the files to</param>
 	/// <param name="Locator">Locator for the tree to extract</param>
-	public record struct UploadFilesMessage(string Name, NodeLocator Locator);
+	public record struct UploadFilesMessage(string Name, BundleNodeLocator Locator);
 
 	/// <summary>
 	/// Deletes files or directories in the remote
@@ -338,7 +338,7 @@ namespace EpicGames.Horde.Compute
 		/// <param name="locator">Location of a <see cref="DirectoryNode"/> describing contents of the sandbox</param>
 		/// <param name="storage">Storage for the sandbox data</param>
 		/// <param name="cancellationToken">Cancellation token for the operation</param>
-		public static async Task UploadFilesAsync(this AgentMessageChannel channel, string path, NodeLocator locator, IStorageClient storage, CancellationToken cancellationToken = default)
+		public static async Task UploadFilesAsync(this AgentMessageChannel channel, string path, BundleNodeLocator locator, IStorageClient storage, CancellationToken cancellationToken = default)
 		{
 			using (IAgentMessageBuilder request = await channel.CreateMessageAsync(AgentMessageType.WriteFiles, cancellationToken))
 			{
@@ -360,7 +360,7 @@ namespace EpicGames.Horde.Compute
 		public static UploadFilesMessage ParseUploadFilesMessage(this AgentMessage message)
 		{
 			string name = message.ReadString();
-			NodeLocator locator = message.ReadNodeLocator();
+			BundleNodeLocator locator = message.ReadNodeLocator();
 			return new UploadFilesMessage(name, locator);
 		}
 
diff --git a/Engine/Source/Programs/Shared/EpicGames.Horde/Compute/AgentMessageHandler.cs b/Engine/Source/Programs/Shared/EpicGames.Horde/Compute/AgentMessageHandler.cs
--- a/Engine/Source/Programs/Shared/EpicGames.Horde/Compute/AgentMessageHandler.cs
+++ b/Engine/Source/Programs/Shared/EpicGames.Horde/Compute/AgentMessageHandler.cs
@@ -116,7 +116,7 @@ namespace EpicGames.Horde.Compute
 			}
 		}
 
-		async Task WriteFilesAsync(AgentMessageChannel channel, string path, NodeLocator locator, CancellationToken cancellationToken)
+		async Task WriteFilesAsync(AgentMessageChannel channel, string path, BundleNodeLocator locator, CancellationToken cancellationToken)
 		{
 			using AgentStorageClient store = new AgentStorageClient(channel);
 			BundleReader reader = new BundleReader(store, _memoryCache, _logger);
diff --git a/Engine/Source/Programs/Shared/EpicGames.Horde/Storage/Backends/FileStorageClient.cs b/Engine/Source/Programs/Shared/EpicGames.Horde/Storage/Backends/FileStorageClient.cs
--- a/Engine/Source/Programs/Shared/EpicGames.Horde/Storage/Backends/FileStorageClient.cs
+++ b/Engine/Source/Programs/Shared/EpicGames.Horde/Storage/Backends/FileStorageClient.cs
@@ -43,7 +43,7 @@ namespace EpicGames.Horde.Storage.Backends
 		public async ValueTask<BlobHandle> ReadRefAsync(FileReference file)
 		{
 			string text = await FileReference.ReadAllTextAsync(file);
-			return new FlushedNodeHandle(TreeReader, NodeLocator.Parse(text));
+			return new FlushedNodeHandle(TreeReader, BundleNodeLocator.Parse(text));
 		}
 
 		FileReference GetRefFile(RefName name) => FileReference.Combine(_rootDir, name.ToString() + ".ref");
@@ -128,7 +128,7 @@ namespace EpicGames.Horde.Storage.Backends
 
 			_logger.LogInformation("Reading {File}", file);
 			string text = await FileReference.ReadAllTextAsync(file, cancellationToken);
-			return new FlushedNodeHandle(TreeReader, NodeLocator.Parse(text));
+			return new FlushedNodeHandle(TreeReader, BundleNodeLocator.Parse(text));
 		}
 
 		/// <inheritdoc/>
@@ -138,7 +138,7 @@ namespace EpicGames.Horde.Storage.Backends
 			DirectoryReference.CreateDirectory(file.Directory);
 			_logger.LogInformation("Writing {File}", file);
 
-			NodeLocator locator = await target.FlushAsync(cancellationToken);
+			BundleNodeLocator locator = await target.FlushAsync(cancellationToken);
 			for (int attempt = 0; ; attempt++)
 			{
 				try
diff --git a/Engine/Source/Programs/Shared/EpicGames.Horde/Storage/Backends/HttpStorageClient.cs b/Engine/Source/Programs/Shared/EpicGames.Horde/Storage/Backends/HttpStorageClient.cs
--- a/Engine/Source/Programs/Shared/EpicGames.Horde/Storage/Backends/HttpStorageClient.cs
+++ b/Engine/Source/Programs/Shared/EpicGames.Horde/Storage/Backends/HttpStorageClient.cs
@@ -220,7 +220,7 @@ namespace EpicGames.Horde.Storage.Backends
 						FindNodesResponse? message = await response.Content.ReadFromJsonAsync<FindNodesResponse>(cancellationToken: cancellationToken);
 						foreach (FindNodeResponse node in message!.Nodes)
 						{
-							yield return new FlushedNodeHandle(TreeReader, new NodeLocator(node.Hash, node.Blob, node.ExportIdx));
+							yield return new FlushedNodeHandle(TreeReader, new BundleNodeLocator(node.Hash, node.Blob, node.ExportIdx));
 						}
 					}
 				}
@@ -276,7 +276,7 @@ namespace EpicGames.Horde.Storage.Backends
 							response.EnsureSuccessStatusCode();
 							ReadRefResponse? data = await response.Content.ReadFromJsonAsync<ReadRefResponse>(cancellationToken: cancellationToken);
 							_logger.LogDebug("Read ref {RefName} -> {Blob}#{ExportIdx}", name, data!.Blob, data!.ExportIdx);
-							return new FlushedNodeHandle(TreeReader, new NodeLocator(data.Hash, data!.Blob, data!.ExportIdx));
+							return new FlushedNodeHandle(TreeReader, new BundleNodeLocator(data.Hash, data!.Blob, data!.ExportIdx));
 						}
 					}
 				}
@@ -289,7 +289,7 @@ namespace EpicGames.Horde.Storage.Backends
 			_logger.LogDebug("Writing ref {RefName} -> {RefTarget}", name, target);
 			using (HttpClient httpClient = _createClient())
 			{
-				NodeLocator locator = await target.FlushAsync(cancellationToken);
+				BundleNodeLocator locator = await target.FlushAsync(cancellationToken);
 				using (HttpResponseMessage response = await httpClient.PutAsync($"refs/{name}", new { blob = locator.Blob, exportIdx = locator.ExportIdx, options }, cancellationToken))
 				{
 					response.EnsureSuccessStatusCode();
diff --git a/Engine/Source/Programs/Shared/EpicGames.Horde/Storage/Backends/MemoryStorageClient.cs b/Engine/Source/Programs/Shared/EpicGames.Horde/Storage/Backends/MemoryStorageClient.cs
--- a/Engine/Source/Programs/Shared/EpicGames.Horde/Storage/Backends/MemoryStorageClient.cs
+++ b/Engine/Source/Programs/Shared/EpicGames.Horde/Storage/Backends/MemoryStorageClient.cs
@@ -30,7 +30,7 @@ namespace EpicGames.Horde.Storage.Backends
 		/// <summary>
 		/// Map of ref name to ref data
 		/// </summary>
-		readonly ConcurrentDictionary<RefName, NodeLocator> _refs = new ConcurrentDictionary<RefName, NodeLocator>();
+		readonly ConcurrentDictionary<RefName, BundleNodeLocator> _refs = new ConcurrentDictionary<RefName, BundleNodeLocator>();
 
 		/// <summary>
 		/// Content addressed data lookup
@@ -41,7 +41,7 @@ namespace EpicGames.Horde.Storage.Backends
 		public IReadOnlyDictionary<BlobLocator, Bundle> Blobs => _blobs;
 
 		/// <inheritdoc cref="_refs"/>
-		public IReadOnlyDictionary<RefName, NodeLocator> Refs => _refs;
+		public IReadOnlyDictionary<RefName, BundleNodeLocator> Refs => _refs;
 
 		/// <summary>
 		/// Constructor
@@ -118,7 +118,7 @@ namespace EpicGames.Horde.Storage.Backends
 		/// <inheritdoc/>
 		public override Task<BlobHandle?> TryReadRefTargetAsync(RefName name, RefCacheTime cacheTime = default, CancellationToken cancellationToken = default)
 		{
-			NodeLocator hashedLocator;
+			BundleNodeLocator hashedLocator;
 			if (_refs.TryGetValue(name, out hashedLocator))
 			{
 				return Task.FromResult<BlobHandle?>(new FlushedNodeHandle(TreeReader, hashedLocator)); 
diff --git a/Engine/Source/Programs/Shared/EpicGames.Horde/Storage/BlobHandle.cs b/Engine/Source/Programs/Shared/EpicGames.Horde/Storage/BlobHandle.cs
--- a/Engine/Source/Programs/Shared/EpicGames.Horde/Storage/BlobHandle.cs
+++ b/Engine/Source/Programs/Shared/EpicGames.Horde/Storage/BlobHandle.cs
@@ -33,7 +33,7 @@ namespace EpicGames.Horde.Storage
 		/// Gets the node locator. May throw if the node has not been written to storage yet.
 		/// </summary>
 		/// <returns>Locator for the node</returns>
-		public abstract NodeLocator GetLocator();
+		public abstract BundleNodeLocator GetLocator();
 
 		/// <summary>
 		/// Adds a callback to be executed once the node has been written. Triggers immediately if the node has already been written.
@@ -96,7 +96,7 @@ namespace EpicGames.Horde.Storage
 		/// Flush the node to storage and retrieve its locator
 		/// </summary>
 		/// <param name="cancellationToken">Cancellation token for the operation</param>
-		public abstract ValueTask<NodeLocator> FlushAsync(CancellationToken cancellationToken = default);
+		public abstract ValueTask<BundleNodeLocator> FlushAsync(CancellationToken cancellationToken = default);
 
 		/// <inheritdoc/>
 		public override string ToString() => HasLocator()? GetLocator().ToString() : Hash.ToString();
diff --git a/Engine/Source/Programs/Shared/EpicGames.Horde/Storage/BundleNodeLocator.cs b/Engine/Source/Programs/Shared/EpicGames.Horde/Storage/BundleNodeLocator.cs
new file mode 100644
--- /dev/null
+++ b/Engine/Source/Programs/Shared/EpicGames.Horde/Storage/BundleNodeLocator.cs
@@ -0,0 +1,179 @@
+// Copyright Epic Games, Inc. All Rights Reserved.
+
+using System;
+using System.ComponentModel;
+using System.Diagnostics.CodeAnalysis;
+using System.Globalization;
+using System.Text;
+using System.Text.Json;
+using System.Text.Json.Serialization;
+using EpicGames.Core;
+using EpicGames.Serialization;
+
+namespace EpicGames.Horde.Storage
+{
+	/// <summary>
+	/// Locates a node in storage
+	/// </summary>
+	[JsonConverter(typeof(BundleNodeLocatorJsonConverter))]
+	[TypeConverter(typeof(BundleNodeLocatorTypeConverter))]
+	[CbConverter(typeof(BundleNodeLocatorCbConverter))]
+	public struct BundleNodeLocator : IEquatable<BundleNodeLocator>
+	{
+		/// <summary>
+		/// Hash of the referenced node
+		/// </summary>
+		public IoHash Hash { get; }
+
+		/// <summary>
+		/// Location of the blob containing this node
+		/// </summary>
+		public BlobLocator Blob { get; }
+
+		/// <summary>
+		/// Index of the export within the blob
+		/// </summary>
+		public int ExportIdx { get; }
+
+		/// <summary>
+		/// Constructor
+		/// </summary>
+		public BundleNodeLocator(IoHash hash, BlobLocator blob, int exportIdx)
+		{
+			Hash = hash;
+			Blob = blob;
+			ExportIdx = exportIdx;
+		}
+
+		/// <summary>
+		/// Determines if this locator points to a valid entry
+		/// </summary>
+		public bool IsValid() => Blob.IsValid();
+
+		/// <summary>
+		/// Parse a string as a node locator
+		/// </summary>
+		/// <param name="text">Text to parse</param>
+		/// <returns></returns>
+		public static BundleNodeLocator Parse(ReadOnlySpan<char> text)
+		{
+			int hashLength = IoHash.NumBytes * 2;
+
+			IoHash hash;
+			if (text.Length == hashLength && IoHash.TryParse(text, out hash))
+			{
+				return new BundleNodeLocator(hash, default, 0);
+			}
+
+			if (text[hashLength] == '@')
+			{
+				hash = IoHash.Parse(text.Slice(0, hashLength));
+				text = text.Slice(hashLength + 1);
+			}
+			else
+			{
+				// For legacy locators that don't have a hash, hash the path instead. This is obviously incorrect, but we never validate hashes and will serve as a stable id during migration.
+				hash = IoHash.Compute(Encoding.UTF8.GetBytes(text.ToString()));
+			}
+
+			int hashIdx = text.IndexOf('#');
+			if (hashIdx == -1)
+			{
+				throw new ArgumentException("Invalid node locator", nameof(text));
+			}
+
+			int exportIdx = Int32.Parse(text.Slice(hashIdx + 1), NumberStyles.None, CultureInfo.InvariantCulture);
+			BlobLocator blobLocator = new BlobLocator(new Utf8String(text.Slice(0, hashIdx)));
+			return new BundleNodeLocator(hash, blobLocator, exportIdx);
+		}
+
+		/// <inheritdoc/>
+		public override bool Equals([NotNullWhen(true)] object? obj) => obj is BundleNodeLocator locator && Equals(locator);
+
+		/// <inheritdoc/>
+		public bool Equals(BundleNodeLocator other) => Blob == other.Blob && ExportIdx == other.ExportIdx;
+
+		/// <inheritdoc/>
+		public override int GetHashCode() => HashCode.Combine(Blob, ExportIdx);
+
+		/// <inheritdoc/>
+		public override string ToString() => $"{Hash}@{Blob}#{ExportIdx}";
+
+		/// <inheritdoc/>
+		public static bool operator ==(BundleNodeLocator left, BundleNodeLocator right) => left.Equals(right);
+
+		/// <inheritdoc/>
+		public static bool operator !=(BundleNodeLocator left, BundleNodeLocator right) => !left.Equals(right);
+	}
+
+	/// <summary>
+	/// Type converter for <see cref="BundleNodeLocator"/> to and from JSON
+	/// </summary>
+	sealed class BundleNodeLocatorJsonConverter : JsonConverter<BundleNodeLocator>
+	{
+		/// <inheritdoc/>
+		public override BundleNodeLocator Read(ref Utf8JsonReader reader, Type typeToConvert, JsonSerializerOptions options) => BundleNodeLocator.Parse(reader.GetString() ?? String.Empty);
+
+		/// <inheritdoc/>
+		public override void Write(Utf8JsonWriter writer, BundleNodeLocator value, JsonSerializerOptions options) => writer.WriteStringValue(value.ToString());
+	}
+
+	/// <summary>
+	/// Type converter from strings to <see cref="BundleNodeLocator"/> objects
+	/// </summary>
+	sealed class BundleNodeLocatorTypeConverter : TypeConverter
+	{
+		/// <inheritdoc/>
+		public override bool CanConvertFrom(ITypeDescriptorContext? context, Type sourceType)
+		{
+			return sourceType == typeof(string);
+		}
+
+		/// <inheritdoc/>
+		public override object ConvertFrom(ITypeDescriptorContext? context, CultureInfo? culture, object? value)
+		{
+			return BundleNodeLocator.Parse((string)value!);
+		}
+	}
+
+	/// <summary>
+	/// Type converter to compact binary
+	/// </summary>
+	sealed class BundleNodeLocatorCbConverter : CbConverterBase<BundleNodeLocator>
+	{
+		/// <inheritdoc/>
+		public override BundleNodeLocator Read(CbField field) => BundleNodeLocator.Parse(field.AsString());
+
+		/// <inheritdoc/>
+		public override void Write(CbWriter writer, BundleNodeLocator value) => writer.WriteUtf8StringValue(value.ToString());
+
+		/// <inheritdoc/>
+		public override void WriteNamed(CbWriter writer, Utf8String name, BundleNodeLocator value) => writer.WriteUtf8String(name, value.ToString());
+	}
+
+	/// <summary>
+	/// Extension methods for node locators
+	/// </summary>
+	public static class NodeLocatorExtensions
+	{
+		/// <summary>
+		/// Deserialize a node locator
+		/// </summary>
+		/// <param name="reader">Reader to deserialize from</param>
+		/// <returns>The node id that was read</returns>
+		public static BundleNodeLocator ReadNodeLocator(this IMemoryReader reader)
+		{
+			return BundleNodeLocator.Parse(reader.ReadString());
+		}
+
+		/// <summary>
+		/// Serialize a node locator
+		/// </summary>
+		/// <param name="writer">Writer to serialize to</param>
+		/// <param name="value">Value to serialize</param>
+		public static void WriteNodeLocator(this IMemoryWriter writer, BundleNodeLocator value)
+		{
+			writer.WriteString(value.ToString());
+		}
+	}
+}
diff --git a/Engine/Source/Programs/Shared/EpicGames.Horde/Storage/BundleReader.cs b/Engine/Source/Programs/Shared/EpicGames.Horde/Storage/BundleReader.cs
--- a/Engine/Source/Programs/Shared/EpicGames.Horde/Storage/BundleReader.cs
+++ b/Engine/Source/Programs/Shared/EpicGames.Horde/Storage/BundleReader.cs
@@ -539,7 +539,7 @@ namespace EpicGames.Horde.Storage
 		/// <param name="locator">Locator for the node</param>
 		/// <param name="cancellationToken">Cancellation token for the operation</param>
 		/// <returns>Node data read from the given bundle</returns>
-		public async ValueTask<BlobData> ReadNodeDataAsync(NodeLocator locator, CancellationToken cancellationToken = default)
+		public async ValueTask<BlobData> ReadNodeDataAsync(BundleNodeLocator locator, CancellationToken cancellationToken = default)
 		{
 			BundleInfo bundleInfo = await GetBundleInfoAsync(locator.Blob, cancellationToken);
 			BundleExport export = bundleInfo.Header.Exports[locator.ExportIdx];
@@ -557,7 +557,7 @@ namespace EpicGames.Horde.Storage
 					importBlob = bundleInfo.Header.Imports[reference.ImportIdx];
 				}
 				Debug.Assert(importBlob.IsValid());
-				refs.Add(new FlushedNodeHandle(this, new NodeLocator(reference.Hash, importBlob, reference.NodeIdx)));
+				refs.Add(new FlushedNodeHandle(this, new BundleNodeLocator(reference.Hash, importBlob, reference.NodeIdx)));
 			}
 
 			ReadOnlyMemory<byte> nodeData = ReadOnlyMemory<byte>.Empty;
@@ -577,7 +577,7 @@ namespace EpicGames.Horde.Storage
 		/// <param name="locator">Locator for the node</param>
 		/// <param name="cancellationToken">Cancellation token for the operation</param>
 		/// <returns>Node data read from the given bundle</returns>
-		public async ValueTask<Node> ReadNodeAsync(NodeLocator locator, CancellationToken cancellationToken = default)
+		public async ValueTask<Node> ReadNodeAsync(BundleNodeLocator locator, CancellationToken cancellationToken = default)
 		{
 			BlobData nodeData = await ReadNodeDataAsync(locator, cancellationToken);
 			return Node.Deserialize(nodeData);
@@ -589,6 +589,6 @@ namespace EpicGames.Horde.Storage
 		/// <param name="locator">Locator for the node</param>
 		/// <param name="cancellationToken">Cancellation token for the operation</param>
 		/// <returns>Node data read from the given bundle</returns>
-		public async ValueTask<TNode> ReadNodeAsync<TNode>(NodeLocator locator, CancellationToken cancellationToken = default) where TNode : Node => (TNode)await ReadNodeAsync(locator, cancellationToken);
+		public async ValueTask<TNode> ReadNodeAsync<TNode>(BundleNodeLocator locator, CancellationToken cancellationToken = default) where TNode : Node => (TNode)await ReadNodeAsync(locator, cancellationToken);
 	}
 }
diff --git a/Engine/Source/Programs/Shared/EpicGames.Horde/Storage/BundleWriter.cs b/Engine/Source/Programs/Shared/EpicGames.Horde/Storage/BundleWriter.cs
--- a/Engine/Source/Programs/Shared/EpicGames.Horde/Storage/BundleWriter.cs
+++ b/Engine/Source/Programs/Shared/EpicGames.Horde/Storage/BundleWriter.cs
@@ -57,12 +57,12 @@ namespace EpicGames.Horde.Storage
 	public sealed class FlushedNodeHandle : BlobHandle
 	{
 		readonly BundleReader _reader;
-		readonly NodeLocator _locator;
+		readonly BundleNodeLocator _locator;
 
 		/// <summary>
 		/// Constructor
 		/// </summary>
-		public FlushedNodeHandle(BundleReader reader, NodeLocator locator)
+		public FlushedNodeHandle(BundleReader reader, BundleNodeLocator locator)
 			: base(locator.Hash)
 		{
 			_reader = reader;
@@ -73,7 +73,7 @@ namespace EpicGames.Horde.Storage
 		public override bool HasLocator() => true;
 
 		/// <inheritdoc/>
-		public override NodeLocator GetLocator() => _locator;
+		public override BundleNodeLocator GetLocator() => _locator;
 
 		/// <inheritdoc/>
 		public override void AddWriteCallback(BlobWriteCallback callback) => callback.OnWrite();
@@ -82,7 +82,7 @@ namespace EpicGames.Horde.Storage
 		public override ValueTask<BlobData> ReadAsync(CancellationToken cancellationToken = default) => _reader.ReadNodeDataAsync(_locator, cancellationToken);
 
 		/// <inheritdoc/>
-		public override ValueTask<NodeLocator> FlushAsync(CancellationToken cancellationToken = default) => new ValueTask<NodeLocator>(_locator);
+		public override ValueTask<BundleNodeLocator> FlushAsync(CancellationToken cancellationToken = default) => new ValueTask<BundleNodeLocator>(_locator);
 	}
 
 	/// <summary>
@@ -157,7 +157,7 @@ namespace EpicGames.Horde.Storage
 
 			object LockObject => _reader;
 
-			NodeLocator _locator;
+			BundleNodeLocator _locator;
 			PendingBundle? _pendingBundle;
 
 			public readonly NodeKey Key;
@@ -186,7 +186,7 @@ namespace EpicGames.Horde.Storage
 			public override bool HasLocator() => _locator.IsValid();
 
 			/// <inheritdoc/>
-			public override NodeLocator GetLocator()
+			public override BundleNodeLocator GetLocator()
 			{
 				if (!_locator.IsValid())
 				{
@@ -195,7 +195,7 @@ namespace EpicGames.Horde.Storage
 				return _locator;
 			}
 
-			public void MarkAsWritten(NodeLocator locator)
+			public void MarkAsWritten(BundleNodeLocator locator)
 			{
 				lock (LockObject)
 				{
@@ -238,7 +238,7 @@ namespace EpicGames.Horde.Storage
 			}
 
 			/// <inheritdoc/>
-			public override async ValueTask<NodeLocator> FlushAsync(CancellationToken cancellationToken = default)
+			public override async ValueTask<BundleNodeLocator> FlushAsync(CancellationToken cancellationToken = default)
 			{
 				if (_locator.IsValid())
 				{
@@ -539,7 +539,7 @@ namespace EpicGames.Horde.Storage
 
 					for (int idx = 0; idx < _queue.Count; idx++)
 					{
-						NodeLocator nodeLocator = new NodeLocator(_queue[idx].Key.Hash, locator, idx);
+						BundleNodeLocator nodeLocator = new BundleNodeLocator(_queue[idx].Key.Hash, locator, idx);
 						traceLogger?.LogInformation("Updated pending node {Hash} with locator {Locator}", _queue[idx].Key.Hash, nodeLocator);
 						_queue[idx].MarkAsWritten(nodeLocator);
 					}
@@ -601,7 +601,7 @@ namespace EpicGames.Horde.Storage
 						BundleExportRef exportRef;
 						if (!nodeHandleToExportRef.TryGetValue(handle, out exportRef))
 						{
-							NodeLocator locator = handle.GetLocator();
+							BundleNodeLocator locator = handle.GetLocator();
 
 							int importIdx = FindOrAddItemIndex(locator.Blob, imports, importToIndex);
 							exportRef = new BundleExportRef(importIdx, locator.ExportIdx, handle.Hash);
diff --git a/Engine/Source/Programs/Shared/EpicGames.Horde/Storage/NodeLocator.cs b/Engine/Source/Programs/Shared/EpicGames.Horde/Storage/NodeLocator.cs
deleted file mode 100644
diff --git a/Engine/Source/Programs/UnrealCloudDDC/Jupiter/Controllers/BundlesController.cs b/Engine/Source/Programs/UnrealCloudDDC/Jupiter/Controllers/BundlesController.cs
--- a/Engine/Source/Programs/UnrealCloudDDC/Jupiter/Controllers/BundlesController.cs
+++ b/Engine/Source/Programs/UnrealCloudDDC/Jupiter/Controllers/BundlesController.cs
@@ -328,7 +328,7 @@ namespace Jupiter.Controllers
                 return result;
             }
             IStorageClientJupiter client = await _storageService.GetClientAsync(namespaceId, cancellationToken);
-            NodeLocator target = new NodeLocator(request.Hash, request.Blob, request.ExportIdx);
+            BundleNodeLocator target = new BundleNodeLocator(request.Hash, request.Blob, request.ExportIdx);
             await client.WriteRefTargetAsync(refName, target, request.Options, cancellationToken);
 
             return Ok();
@@ -466,7 +466,7 @@ namespace Jupiter.Controllers
 
             object content;
 
-            BlobData nodeData = await reader.ReadNodeDataAsync(new NodeLocator(IoHash.Zero, locator, exportIdx), cancellationToken);
+            BlobData nodeData = await reader.ReadNodeDataAsync(new BundleNodeLocator(IoHash.Zero, locator, exportIdx), cancellationToken);
 
             Node node = Node.Deserialize(nodeData);
             switch (node)
diff --git a/Engine/Source/Programs/UnrealCloudDDC/Jupiter/Implementation/Bundles/StorageService.cs b/Engine/Source/Programs/UnrealCloudDDC/Jupiter/Implementation/Bundles/StorageService.cs
--- a/Engine/Source/Programs/UnrealCloudDDC/Jupiter/Implementation/Bundles/StorageService.cs
+++ b/Engine/Source/Programs/UnrealCloudDDC/Jupiter/Implementation/Bundles/StorageService.cs
@@ -21,7 +21,7 @@ public interface IStorageClientJupiter : IStorageClient
     Task<(BlobLocator Locator, Uri UploadUrl)?> GetWriteRedirectAsync(string prefix, CancellationToken cancellationToken);
     bool SupportsRedirects { get; set; }
     Task<Uri?> GetReadRedirectAsync(BlobLocator locator, CancellationToken cancellationToken);
-    Task WriteRefTargetAsync(RefName refName, NodeLocator target, RefOptions? requestOptions, CancellationToken cancellationToken);
+    Task WriteRefTargetAsync(RefName refName, BundleNodeLocator target, RefOptions? requestOptions, CancellationToken cancellationToken);
 }
 
 public interface IStorageService
@@ -176,7 +176,7 @@ public class StorageClient : IStorageClientJupiter
             BlobLocator blobLocator = new BlobLocator(inlinePayload.BlobLocator);
             int exportId = inlinePayload.ExportId;
 
-            return new FlushedNodeHandle(_treeReader, new NodeLocator(nodeHash, blobLocator, exportId));
+            return new FlushedNodeHandle(_treeReader, new BundleNodeLocator(nodeHash, blobLocator, exportId));
         }
         catch (ObjectNotFoundException )
         {
@@ -187,7 +187,7 @@ public class StorageClient : IStorageClientJupiter
     public async Task<BlobHandle> WriteRefAsync(RefName name, Bundle bundle, int exportIdx, Utf8String prefix = default, RefOptions? options = null, CancellationToken cancellationToken = default)
     {
         BlobLocator locator = await this.WriteBundleAsync(bundle, prefix, cancellationToken);
-        BlobHandle target = new FlushedNodeHandle(_treeReader, new NodeLocator(bundle.Header.Exports[exportIdx].Hash, locator, exportIdx));
+        BlobHandle target = new FlushedNodeHandle(_treeReader, new BundleNodeLocator(bundle.Header.Exports[exportIdx].Hash, locator, exportIdx));
         await WriteRefTargetAsync(name, target, options, cancellationToken);
 
         return target;
@@ -198,7 +198,7 @@ public class StorageClient : IStorageClientJupiter
         return WriteRefTargetAsync(refName, target.GetLocator(), requestOptions, cancellationToken);
     }
 
-    public async Task WriteRefTargetAsync(RefName refName, NodeLocator target, RefOptions? requestOptions, CancellationToken cancellationToken)
+    public async Task WriteRefTargetAsync(RefName refName, BundleNodeLocator target, RefOptions? requestOptions, CancellationToken cancellationToken)
     {
         BlobIdentifier bundleBlob = BlobIdentifier.FromBlobLocator(target.Blob);
         IoHashKey refKey = IoHashKey.FromName(refName.ToString());
