Horde: Inform AWS auto-scaling groups which agents are ready for termination

diff --git a/Engine/Source/Programs/Horde/Horde.Server.Tests/Fleet/AwsAutoScalingLifecycleTest.cs b/Engine/Source/Programs/Horde/Horde.Server.Tests/Fleet/AwsAutoScalingLifecycleTest.cs
--- a/Engine/Source/Programs/Horde/Horde.Server.Tests/Fleet/AwsAutoScalingLifecycleTest.cs
+++ b/Engine/Source/Programs/Horde/Horde.Server.Tests/Fleet/AwsAutoScalingLifecycleTest.cs
@@ -2,6 +2,7 @@
 
 using System;
 using System.Collections.Generic;
+using System.Linq;
 using System.Net;
 using System.Text.Json;
 using System.Threading;
@@ -10,11 +11,13 @@ using Amazon.AutoScaling;
 using Amazon.AutoScaling.Model;
 using Horde.Server.Agents;
 using Horde.Server.Agents.Fleet;
+using Horde.Server.Agents.Leases;
 using Horde.Server.Agents.Pools;
 using HordeCommon;
 using Microsoft.Extensions.DependencyInjection;
 using Microsoft.Extensions.Logging;
 using Microsoft.VisualStudio.TestTools.UnitTesting;
+using MongoDB.Bson;
 using Moq;
 
 namespace Horde.Server.Tests.Fleet;
@@ -134,6 +137,50 @@ public class AwsAutoScalingLifecycleServiceTest : TestSetup
 		AssertLifecycleUpdate(lae, AwsAutoScalingLifecycleService.ActionAbandon, _lifecycleUpdates[0]);
 	}
 	
+	[TestMethod]
+	public async Task GetInstancesAvailableForTermination_IdleAgent_ReturnsInstanceId()
+	{
+		// Arrange
+		IAgent agent = await CreateAgentAsync(new PoolId("pool1"), awsInstanceId: "i-1000");
+		TerminationPolicyEvent e = CreateTerminationPolicyEvent("i-1000");
+
+		// Act
+		List<string> instanceIds = await _asgLifecycleService.GetInstancesAvailableForTermination(e, CancellationToken.None);
+
+		// Assert
+		CollectionAssert.AreEqual(new List<string> () { "i-1000"}, instanceIds);
+	}
+	
+	[TestMethod]
+	public async Task GetInstancesAvailableForTermination_IdleAgentsInMixedAsgs_OnlyReturnInstanceIdFromSameAsg()
+	{
+		// Arrange
+		IAgent agent1 = await CreateAgentAsync(new PoolId("pool1"), awsInstanceId: "i-1000");
+		IAgent agent2 = await CreateAgentAsync(new PoolId("pool1"), awsInstanceId: "i-2000");
+		TerminationPolicyEvent e = CreateTerminationPolicyEvent("i-1000");
+
+		// Act
+		List<string> instanceIds = await _asgLifecycleService.GetInstancesAvailableForTermination(e, CancellationToken.None);
+
+		// Assert
+		CollectionAssert.AreEqual(new List<string> () { "i-1000"}, instanceIds);
+	}
+	
+	[TestMethod]
+	public async Task GetInstancesAvailableForTermination_AgentRunningJob_ReturnsNoInstanceId()
+	{
+		// Arrange
+		AgentLease lease = new(new LeaseId(ObjectId.GenerateNewId()), "test-lease", null, null, null, LeaseState.Active, null, false, null);
+		IAgent agent = await CreateAgentAsync(new PoolId("pool1"), awsInstanceId: "i-1000", lease: lease);
+		TerminationPolicyEvent e = CreateTerminationPolicyEvent("i-1000");
+
+		// Act
+		List<string> instanceIds = await _asgLifecycleService.GetInstancesAvailableForTermination(e, CancellationToken.None);
+
+		// Assert
+		Assert.AreEqual(0, instanceIds.Count);
+	}
+	
 	[TestMethod]
 	public void DeserializeLifecycleActionEvent()
 	{
@@ -161,6 +208,59 @@ public class AwsAutoScalingLifecycleServiceTest : TestSetup
 		Assert.AreEqual("autoscaling:EC2_INSTANCE_TERMINATING", ev.LifecycleTransition);
 		Assert.AreEqual("AutoScalingGroup", ev.Origin);
 	}
+	
+	[TestMethod]
+	public void DeserializeTerminationPolicyEvent()
+	{
+		string rawText = @"
+{
+	""AutoScalingGroupARN"": ""my-asg-arn"",
+	""AutoScalingGroupName"": ""my-asg-name"",
+	""CapacityToTerminate"": [{
+		""AvailabilityZone"": ""us-east-1a"",
+		""Capacity"": 5,
+		""InstanceMarketOption"": ""on-demand""
+	}],
+	""Instances"": [{
+		""AvailabilityZone"": ""us-east-1b"",
+		""InstanceId"": ""i-10001"",
+		""InstanceType"": ""m5d.large"",
+		""InstanceMarketOption"": ""on-demand""
+	}, {
+		""AvailabilityZone"": ""us-east-1c"",
+		""InstanceId"": ""i-10002"",
+		""InstanceType"": ""m6i.large"",
+		""InstanceMarketOption"": ""spot""
+	}],
+	""Cause"": ""SCALE_IN""
+}
+";
+
+		TerminationPolicyEvent? ev = JsonSerializer.Deserialize<TerminationPolicyEvent>(rawText);
+		Assert.AreEqual("my-asg-arn", ev!.AutoScalingGroupArn);
+		Assert.AreEqual("my-asg-name", ev.AutoScalingGroupName);
+		Assert.AreEqual("SCALE_IN", ev.Cause);
+
+		{
+			Assert.AreEqual(1, ev.CapacityToTerminate.Count);
+			Assert.AreEqual("us-east-1a", ev.CapacityToTerminate[0].AvailabilityZone);
+			Assert.AreEqual(5, ev.CapacityToTerminate[0].Capacity);
+			Assert.AreEqual("on-demand", ev.CapacityToTerminate[0].InstanceMarketOption);
+		}
+		
+		{
+			Assert.AreEqual(2, ev.Instances.Count);
+			Assert.AreEqual("us-east-1b", ev.Instances[0].AvailabilityZone);
+			Assert.AreEqual("i-10001", ev.Instances[0].InstanceId);
+			Assert.AreEqual("m5d.large", ev.Instances[0].InstanceType);
+			Assert.AreEqual("on-demand", ev.Instances[0].InstanceMarketOption);
+			
+			Assert.AreEqual("us-east-1c", ev.Instances[1].AvailabilityZone);
+			Assert.AreEqual("i-10002", ev.Instances[1].InstanceId);
+			Assert.AreEqual("m6i.large", ev.Instances[1].InstanceType);
+			Assert.AreEqual("spot", ev.Instances[1].InstanceMarketOption);
+		}
+	}
 
 	private static void AssertLifecycleUpdate(LifecycleActionEvent expectedEvent, string expectedResult, CompleteLifecycleActionRequest actual)
 	{
@@ -170,4 +270,14 @@ public class AwsAutoScalingLifecycleServiceTest : TestSetup
 		Assert.AreEqual(expectedEvent.LifecycleHookName, actual.LifecycleHookName);
 		Assert.AreEqual(expectedEvent.AutoScalingGroupName, actual.AutoScalingGroupName);
 	}
+	
+	private static TerminationPolicyEvent CreateTerminationPolicyEvent(params string[] instanceIds)
+	{
+		return new()
+		{
+			AutoScalingGroupArn = "test-asg-arn", AutoScalingGroupName = "test-asg-name", Cause = "SCALE_IN",
+			CapacityToTerminate = new () { new (1) },
+			Instances = new(instanceIds.Select(x => new TerminationPolicyInstance(x)))
+		};
+	}
 }
diff --git a/Engine/Source/Programs/Horde/Horde.Server.Tests/TestSetup.cs b/Engine/Source/Programs/Horde/Horde.Server.Tests/TestSetup.cs
--- a/Engine/Source/Programs/Horde/Horde.Server.Tests/TestSetup.cs
+++ b/Engine/Source/Programs/Horde/Horde.Server.Tests/TestSetup.cs
@@ -358,20 +358,45 @@ namespace Horde.Server.Tests
 			return CreateAgentAsync(pool.Id, enabled, requestShutdown, properties, adjustClockBy);
 		}
 		
-		public async Task<IAgent> CreateAgentAsync(PoolId poolId, bool enabled = true, bool requestShutdown = false, List<string>? properties = null, TimeSpan? adjustClockBy = null)
+		/// <summary>
+		/// Helper function for setting up agents to be used in tests
+		/// </summary>
+		/// <param name="poolId">Pool ID which the agent should belong to</param>
+		/// <param name="enabled">Whether set the agent as enabled</param>
+		/// <param name="requestShutdown">Mark it with a request for shutdown</param>
+		/// <param name="properties">Any properties to assign</param>
+		/// <param name="adjustClockBy">Time span to temporarily skew the clock when creating the agent</param>
+		/// <param name="awsInstanceId">AWS instance ID for the agent (will be set in properties)</param>
+		/// <param name="lease">A lease to assign the agent</param>
+		/// <returns>A new agent</returns>
+		public async Task<IAgent> CreateAgentAsync(
+			PoolId poolId, bool enabled = true, bool requestShutdown = false, List<string>? properties = null,
+			TimeSpan? adjustClockBy = null, string? awsInstanceId = null, AgentLease? lease = null)
 		{
 			DateTime now = Clock.UtcNow;
 			if (adjustClockBy != null)
 			{
 				Clock.UtcNow = now + adjustClockBy.Value;
 			}
+
+			Dictionary<string,int> resources = new ();
+			List<string> tempProps = new(properties ?? new List<string>());
+			if (awsInstanceId != null)
+			{
+				tempProps.Add(KnownPropertyNames.AwsInstanceId + "=" + awsInstanceId);
+			}
 			
 			IAgent agent = await AgentService.CreateAgentAsync("TestAgent" + s_agentIdCounter++, enabled, new List<PoolId> { poolId });
-			agent = await AgentService.CreateSessionAsync(agent, AgentStatus.Ok, properties ?? new List<string>(), new Dictionary<string, int>(), null);
+			agent = await AgentService.CreateSessionAsync(agent, AgentStatus.Ok, tempProps, resources, null);
 			if (requestShutdown)
 			{
 				await AgentCollection.TryUpdateSettingsAsync(agent, requestShutdown: true);
 			}
+
+			if (lease != null)
+			{
+				await AgentCollection.TryAddLeaseAsync(agent, lease);
+			}
 			
 			Clock.UtcNow = now;
 			return agent;
diff --git a/Engine/Source/Programs/Horde/Horde.Server/Agents/Fleet/AwsAutoScalingLifecycle.cs b/Engine/Source/Programs/Horde/Horde.Server/Agents/Fleet/AwsAutoScalingLifecycle.cs
--- a/Engine/Source/Programs/Horde/Horde.Server/Agents/Fleet/AwsAutoScalingLifecycle.cs
+++ b/Engine/Source/Programs/Horde/Horde.Server/Agents/Fleet/AwsAutoScalingLifecycle.cs
@@ -2,6 +2,7 @@
 
 using System;
 using System.Collections.Generic;
+using System.Diagnostics.CodeAnalysis;
 using System.IO;
 using System.Linq;
 using System.Text.Json;
@@ -15,7 +16,11 @@ using Amazon.SQS;
 using Amazon.SQS.Model;
 using EpicGames.Core;
 using Horde.Server.Server;
+using Horde.Server.Utilities;
 using HordeCommon;
+using Microsoft.AspNetCore.Authorization;
+using Microsoft.AspNetCore.Http;
+using Microsoft.AspNetCore.Mvc;
 using Microsoft.Extensions.DependencyInjection;
 using Microsoft.Extensions.Hosting;
 using Microsoft.Extensions.Logging;
@@ -248,6 +253,49 @@ public sealed class AwsAutoScalingLifecycleService : IHostedService, IDisposable
 			return false;
 		}
 	}
+	
+	/// <summary>
+	/// Get instances available for termination.
+	/// An AWS Lambda function is set to handle termination policy queries from an auto-scaling group
+	/// That function is configured to call Horde server and invoke this method
+	/// </summary>
+	/// <param name="e">Event as received from AWS API</param>
+	/// <param name="cancellationToken">Cancellation token</param>
+	/// <returns>A list of instance IDs that can be terminated</returns>
+	public async Task<List<string>> GetInstancesAvailableForTermination(TerminationPolicyEvent e, CancellationToken cancellationToken)
+	{
+		using TelemetrySpan span = _tracer.StartActiveSpan($"{nameof(AwsAutoScalingLifecycleService)}.{nameof(GetInstancesAvailableForTermination)}");
+		span.SetAttribute("asgName", e.AutoScalingGroupName);
+		span.SetAttribute("instancesCount", e.Instances.Count);
+		span.SetAttribute("capacityCount", e.CapacityToTerminate.Count);
+		
+		List<string> asgInstanceIds = e.Instances
+			.Select(x => x.InstanceId)
+			.OfType<string>()
+			.ToList();
+		
+		bool IsAgentSuggestedByAsg(IAgent agent, [NotNullWhen(true)] out string? instanceId)
+		{
+			instanceId = asgInstanceIds.FirstOrDefault(asgInstanceId => agent.HasProperty($"{KnownPropertyNames.AwsInstanceId}={asgInstanceId}"));
+			return instanceId != null;
+		}
+		
+		List<string> validInstanceIds = new();
+		List<IAgent> agents = await _agentService.FindAgentsAsync(null, null, null, null, null);
+		foreach (IAgent agent in agents)
+		{
+			if (IsAgentSuggestedByAsg(agent, out string? instanceId))
+			{
+				if (agent.Leases.Count == 0)
+				{
+					validInstanceIds.Add(instanceId);
+				}
+			}
+		}
+
+		span.SetAttribute("validInstanceIds", validInstanceIds.Count);
+		return validInstanceIds;
+	}
 
 	private async Task TrackAgentLifecycleAsync(AgentId agentId, LifecycleActionEvent e)
 	{
@@ -368,6 +416,42 @@ public sealed class AwsAutoScalingLifecycleService : IHostedService, IDisposable
 	}
 }
 
+/// <summary>
+/// Controller handling callbacks for AWS auto-scaling
+/// </summary>
+[ApiController]
+[Authorize]
+[Route("[controller]")]
+public class AwsAutoScalingLifecycleController : HordeControllerBase
+{
+	private readonly AwsAutoScalingLifecycleService _lifecycleService;
+
+	/// <summary>
+	/// Constructor
+	/// </summary>
+	public AwsAutoScalingLifecycleController(AwsAutoScalingLifecycleService lifecycleService)
+	{
+		_lifecycleService = lifecycleService;
+	}
+	
+	/// <summary>
+	/// Called by AWS auto-scaling group to get which instance IDs are valid for termination
+	/// <see cref="AwsAutoScalingLifecycleService.GetInstancesAvailableForTermination" />
+	/// </summary>
+	/// <param name="tpe">Event</param>
+	/// <param name="cancellationToken">Cancellation token</param>
+	/// <returns>List of instance IDs the AWS auto-scaling group is allowed to terminate</returns>
+	[HttpPost]
+	[Route("/api/v1/aws/asg-termination-policy")]
+	[ProducesResponseType(StatusCodes.Status200OK)]
+	public async Task<ActionResult> TerminationPolicyAsync([FromBody] TerminationPolicyEvent tpe, CancellationToken cancellationToken)
+	{
+		List<string> instanceIds = await _lifecycleService.GetInstancesAvailableForTermination(tpe, cancellationToken);
+		return new JsonResult(instanceIds);
+	}
+}
+
+
 /// <summary>
 /// Lifecycle action event from AWS auto-scaling group
 /// The property names are explicitly set to highlight these are not controlled by Horde but sent from AWS API.
@@ -404,3 +488,102 @@ public class LifecycleActionEvent
 	/// </summary>
 	[JsonPropertyName("Origin")] public string Origin { get; set; } = "";
 }
+
+/// <summary>
+/// Instance referenced by <see cref="TerminationPolicyEvent" />
+/// </summary>
+public class TerminationPolicyInstance
+{
+	/// <summary>
+	/// Availability zone of the instance (such as 'us-east-1c')
+	/// </summary>
+	[JsonPropertyName("AvailabilityZone")] public string AvailabilityZone { get; set; } = "";
+	
+	/// <summary>
+	/// Capacity to terminate (number of instances)
+	/// </summary>
+	[JsonPropertyName("Capacity")] public int? Capacity { get; set; }
+	
+	/// <summary>
+	/// Instance ID (such as 'i-123456789')
+	/// </summary>
+	[JsonPropertyName("InstanceId")] public string? InstanceId { get; set; }
+	
+	/// <summary>
+	/// Instance type (such as 'm5.xlarge')
+	/// </summary>
+	[JsonPropertyName("InstanceType")] public string? InstanceType { get; set; }
+	
+	/// <summary>
+	/// Instance market options (such as 'on-demand' or 'spot')
+	/// </summary>
+	[JsonPropertyName("InstanceMarketOption")] public string InstanceMarketOption { get; set; } = "";
+
+	/// <summary>
+	/// Constructor
+	/// </summary>
+	public TerminationPolicyInstance()
+	{
+	}
+
+	/// <summary>
+	/// Constructor
+	/// </summary>
+	/// <param name="capacity"></param>
+	/// <param name="availabilityZone"></param>
+	/// <param name="instanceMarketOption"></param>
+	public TerminationPolicyInstance(int capacity, string availabilityZone = "us-east-1a", string instanceMarketOption = "on-demand")
+	{
+		AvailabilityZone = availabilityZone;
+		Capacity = capacity;
+		InstanceMarketOption = instanceMarketOption;
+	}
+	
+	/// <summary>
+	/// Constructor
+	/// </summary>
+	/// <param name="instanceId"></param>
+	/// <param name="availabilityZone"></param>
+	/// <param name="instanceType"></param>
+	/// <param name="instanceMarketOption"></param>
+	public TerminationPolicyInstance(string instanceId, string availabilityZone = "us-east-1a", string instanceType = "m5.large", string instanceMarketOption = "on-demand")
+	{
+		AvailabilityZone = availabilityZone;
+		InstanceId = instanceId;
+		InstanceType = instanceType;
+		InstanceMarketOption = instanceMarketOption;
+	}
+}
+
+/// <summary>
+/// Termination policy event from AWS auto-scaling group
+/// The property names are explicitly set to highlight these are not controlled by Horde but sent from AWS API.
+/// </summary>
+public class TerminationPolicyEvent
+{
+	/// <summary>
+	/// ARN of the auto-scaling group
+	/// </summary>
+	[JsonPropertyName("AutoScalingGroupARN")] public string AutoScalingGroupArn { get; set; } = "";
+	
+	/// <summary>
+	/// Name of the auto-scaling group
+	/// </summary>
+	[JsonPropertyName("AutoScalingGroupName")] public string AutoScalingGroupName { get; set; } = "";
+	
+	/// <summary>
+	/// Capacity that's been requested to be terminated
+	/// </summary>
+	[JsonPropertyName("CapacityToTerminate")] public List<TerminationPolicyInstance> CapacityToTerminate { get; set; } = new ();
+	
+	/// <summary>
+	/// Instances available to terminate
+	/// </summary>
+	[JsonPropertyName("Instances")] public List<TerminationPolicyInstance> Instances { get; set; } = new ();
+	
+	/// <summary>
+	/// Cause of the termination (such as 'SCALE_IN' or 'INSTANCE_REFRESH')
+	/// </summary>
+	[JsonPropertyName("Cause")] public string Cause { get; set; } = "";
+}
+
