Horde: Surface 404 job errors for missing streams
#skipci
#rnx

diff --git a/Engine/Source/Programs/Horde/HordeDashboard/src/backend/Backend.ts b/Engine/Source/Programs/Horde/HordeDashboard/src/backend/Backend.ts
--- a/Engine/Source/Programs/Horde/HordeDashboard/src/backend/Backend.ts
+++ b/Engine/Source/Programs/Horde/HordeDashboard/src/backend/Backend.ts
@@ -251,12 +251,13 @@ export class Backend {
 
     }
 
-    getJob(id: string, query?: JobQuery, includeGraph = true): Promise<JobData> {
+    getJob(id: string, query?: JobQuery, includeGraph = true, show404Error = false): Promise<JobData> {
 
         return new Promise<JobData>((resolve, reject) => {
 
             this.backend.get(`/api/v1/jobs/${id}`, {
-                params: query
+                params: query,
+                show404Error: show404Error
             }).then((value) => {
 
                 const response = value.data as JobData;
diff --git a/Engine/Source/Programs/Horde/HordeDashboard/src/backend/Fetch.ts b/Engine/Source/Programs/Horde/HordeDashboard/src/backend/Fetch.ts
--- a/Engine/Source/Programs/Horde/HordeDashboard/src/backend/Fetch.ts
+++ b/Engine/Source/Programs/Horde/HordeDashboard/src/backend/Fetch.ts
@@ -13,7 +13,10 @@ export interface FetchRequestConfig {
     formData?: boolean;
     responseBlob?: boolean;
     params?: Record<string, string | string[] | number | boolean | undefined>;
+    // does not report 404 error
     suppress404?: boolean;
+    // shows error dialog for 404, in addition to logging it
+    show404Error?: boolean;
 }
 
 export type FetchResponse = {
@@ -53,6 +56,25 @@ export class Fetch {
                         });
                         return null;
                     }
+
+
+                    if (!response.ok && response.status === 404 && config?.show404Error) {
+                        response.json().then(o => {
+                            handleError({
+                                response: response,
+                                title: "404 Not Found",
+                                reason: "404 Not Found",
+                                mode: "GET",
+                                url: url,
+                                message: o?.message ?? "Malformed json on response object"
+                            }, true);            
+                        }).catch(reason => {
+                            console.error("Unable to parse response json on 404: ", reason)
+                        });
+                        return null;
+                    }
+
+
                     return response;
                 });
 
diff --git a/Engine/Source/Programs/Horde/HordeDashboard/src/components/jobDetailsV2/JobDetailViewV2.tsx b/Engine/Source/Programs/Horde/HordeDashboard/src/components/jobDetailsV2/JobDetailViewV2.tsx
--- a/Engine/Source/Programs/Horde/HordeDashboard/src/components/jobDetailsV2/JobDetailViewV2.tsx
+++ b/Engine/Source/Programs/Horde/HordeDashboard/src/components/jobDetailsV2/JobDetailViewV2.tsx
@@ -53,11 +53,7 @@ const JobBreadCrumbs: React.FC<{ jobDetails: JobDetailsV2 }> = observer(({ jobDe
          return <Breadcrumbs items={[{ text: "Loading Job" }]} title={"Loading Job"} spinner={true} />
       }
       else {
-         ErrorHandler.set({
-            reason: `Error loading job data`,
-            title: `Unable to Load Job`,
-            message: `Job ${jobDetails.jobId} could not be loaded.\n\nPlease check that you are on the network and the job stream exists.`
-         }, true);
+         console.error(`Unable to load job ${jobDetails.jobId}: ${jobDetails.jobError}` )
       }
 
       return null;
diff --git a/Engine/Source/Programs/Horde/HordeDashboard/src/components/jobDetailsV2/JobDetailsViewCommon.tsx b/Engine/Source/Programs/Horde/HordeDashboard/src/components/jobDetailsV2/JobDetailsViewCommon.tsx
--- a/Engine/Source/Programs/Horde/HordeDashboard/src/components/jobDetailsV2/JobDetailsViewCommon.tsx
+++ b/Engine/Source/Programs/Horde/HordeDashboard/src/components/jobDetailsV2/JobDetailsViewCommon.tsx
@@ -615,7 +615,7 @@ export class JobDetailsV2 extends PollBase {
 
       const initialRequest = !this.jobData;
 
-      requests.push(backend.getJob(this.jobId));
+      requests.push(backend.getJob(this.jobId, undefined, true, true));
 
       let results: any;
 
@@ -624,6 +624,10 @@ export class JobDetailsV2 extends PollBase {
          this.jobError = reason;
       });
 
+      if (!this.jobError && !results?.length) {
+         this.jobError = "Not Found";
+      }
+
       if (this.jobError) {
          this.setRootUpdated();
          return;
