Horde: Change direct blob read methods to return bundle objects.

diff --git a/Engine/Source/Programs/Horde/Horde.Server/Commands/Bundles/BundleWriteTest.cs b/Engine/Source/Programs/Horde/Horde.Server/Commands/Bundles/BundleWriteTest.cs
--- a/Engine/Source/Programs/Horde/Horde.Server/Commands/Bundles/BundleWriteTest.cs
+++ b/Engine/Source/Programs/Horde/Horde.Server/Commands/Bundles/BundleWriteTest.cs
@@ -28,13 +28,13 @@ namespace Horde.Server.Commands.Bundles
 			}
 
 			public Task DeleteRefAsync(RefName name, CancellationToken cancellationToken = default) => Task.CompletedTask;
-			public Task<Stream> ReadBlobAsync(BlobLocator locator, CancellationToken cancellationToken = default) => throw new NotImplementedException();
-			public Task<Stream> ReadBlobRangeAsync(BlobLocator locator, int offset, int length, CancellationToken cancellationToken = default) => throw new NotImplementedException();
+			public Task<Bundle> ReadBundleAsync(BlobLocator locator, CancellationToken cancellationToken = default) => throw new NotImplementedException();
+			public Task<ReadOnlyMemory<byte>> ReadBundleRangeAsync(BlobLocator locator, int offset, int length, CancellationToken cancellationToken = default) => throw new NotImplementedException();
 			public Task AddAliasAsync(Utf8String name, BlobHandle handle, CancellationToken cancellationToken = default) => throw new NotImplementedException();
 			public Task RemoveAliasAsync(Utf8String name, BlobHandle handle, CancellationToken cancellationToken = default) => throw new NotImplementedException();
 			public IAsyncEnumerable<BlobHandle> FindNodesAsync(Utf8String alias, CancellationToken cancellationToken = default) => throw new NotImplementedException();
 			public Task<BlobHandle?> TryReadRefTargetAsync(RefName name, RefCacheTime cacheTime = default, CancellationToken cancellationToken = default) => throw new NotImplementedException();
-			public Task<BlobLocator> WriteBlobAsync(Stream stream, Utf8String prefix = default, CancellationToken cancellationToken = default) => Task.FromResult(BlobLocator.CreateUnique(Utf8String.Empty));
+			public Task<BlobLocator> WriteBundleAsync(Bundle bundle, Utf8String prefix = default, CancellationToken cancellationToken = default) => Task.FromResult(BlobLocator.CreateUnique(Utf8String.Empty));
 			public Task WriteRefTargetAsync(RefName name, BlobHandle target, RefOptions? options = null, CancellationToken cancellationToken = default) => Task.CompletedTask;
 
 			public BundleWriter CreateWriter(RefName refName = default, BundleOptions? options = null) => new BundleWriter(this, _reader, refName, options);
diff --git a/Engine/Source/Programs/Horde/Horde.Server/Storage/StorageController.cs b/Engine/Source/Programs/Horde/Horde.Server/Storage/StorageController.cs
--- a/Engine/Source/Programs/Horde/Horde.Server/Storage/StorageController.cs
+++ b/Engine/Source/Programs/Horde/Horde.Server/Storage/StorageController.cs
@@ -231,11 +231,14 @@ namespace Horde.Server.Storage
 			}
 			else
 			{
+				Bundle bundle;
 				using (Stream stream = file.OpenReadStream())
 				{
-					BlobLocator locator = await storageClient.WriteBlobAsync(stream, prefix: (prefix == null) ? Utf8String.Empty : new Utf8String(prefix), cancellationToken: cancellationToken);
-					return new WriteBlobResponse { Blob = locator, SupportsRedirects = storageClientImpl?.SupportsRedirects };
+					bundle = await Bundle.FromStreamAsync(stream, cancellationToken);
 				}
+
+				BlobLocator locator = await storageClient.WriteBundleAsync(bundle, prefix: (prefix == null) ? Utf8String.Empty : new Utf8String(prefix), cancellationToken: cancellationToken);
+				return new WriteBlobResponse { Blob = locator, SupportsRedirects = storageClientImpl?.SupportsRedirects };
 			}
 		}
 
@@ -292,11 +295,13 @@ namespace Horde.Server.Storage
 			Stream stream;
 			if (offset == null && length == null)
 			{
-				stream = await storageClient.ReadBlobAsync(locator, cancellationToken);
+				Bundle bundle = await storageClient.ReadBundleAsync(locator, cancellationToken);
+				stream = new ReadOnlySequenceStream(bundle.AsSequence());
 			}
 			else if (offset != null && length != null)
 			{
-				stream = await storageClient.ReadBlobRangeAsync(locator, offset.Value, length.Value, cancellationToken);
+				ReadOnlyMemory<byte> memory = await storageClient.ReadBundleRangeAsync(locator, offset.Value, length.Value, cancellationToken);
+				stream = new ReadOnlyMemoryStream(memory);
 			}
 			else
 			{
diff --git a/Engine/Source/Programs/Horde/Horde.Server/Storage/StorageService.cs b/Engine/Source/Programs/Horde/Horde.Server/Storage/StorageService.cs
--- a/Engine/Source/Programs/Horde/Horde.Server/Storage/StorageService.cs
+++ b/Engine/Source/Programs/Horde/Horde.Server/Storage/StorageService.cs
@@ -127,7 +127,7 @@ namespace Horde.Server.Storage
 			#region Blobs
 
 			/// <inheritdoc/>
-			public override async Task<Stream> ReadBlobAsync(BlobLocator locator, CancellationToken cancellationToken = default)
+			public override async Task<Bundle> ReadBundleAsync(BlobLocator locator, CancellationToken cancellationToken = default)
 			{
 				string path = GetBlobPath(locator);
 
@@ -137,14 +137,14 @@ namespace Horde.Server.Storage
 					throw new StorageException($"Unable to read data from {path}");
 				}
 
-				return stream;
+				return await Bundle.FromStreamAsync(stream, cancellationToken);
 			}
 
 			/// <inheritdoc/>
 			public ValueTask<Uri?> GetReadRedirectAsync(BlobLocator locator, CancellationToken cancellationToken = default) => Backend.TryGetReadRedirectAsync(GetBlobPath(locator), cancellationToken);
 
 			/// <inheritdoc/>
-			public override async Task<Stream> ReadBlobRangeAsync(BlobLocator locator, int offset, int length, CancellationToken cancellationToken = default)
+			public override async Task<ReadOnlyMemory<byte>> ReadBundleRangeAsync(BlobLocator locator, int offset, int length, CancellationToken cancellationToken = default)
 			{
 				string path = GetBlobPath(locator);
 
@@ -154,42 +154,20 @@ namespace Horde.Server.Storage
 					throw new StorageException($"Unable to read data from {path}");
 				}
 
-				return stream;
+				return await stream.ReadAllBytesAsync(cancellationToken);
 			}
 
 			/// <inheritdoc/>
-			public override async Task<BlobLocator> WriteBlobAsync(Stream stream, Utf8String prefix = default, CancellationToken cancellationToken = default)
+			public override async Task<BlobLocator> WriteBundleAsync(Bundle bundle, Utf8String prefix = default, CancellationToken cancellationToken = default)
 			{
-				BlobLocator locator;
-				if (Config.EnableAliases)
-				{
-					using (MemoryStream memoryStream = new MemoryStream())
-					{
-						// Read the blob into memory
-						await stream.CopyToAsync(memoryStream, cancellationToken);
-
-						// Reset the position back to zero and read the header
-						memoryStream.Position = 0;
-						BundleHeader header = await BundleHeader.FromStreamAsync(memoryStream, cancellationToken);
-
-						// Add the blob record
-						locator = await _outer.AddBlobAsync(NamespaceId, prefix, null, cancellationToken);
+				// Add the blob record
+				BlobLocator locator = await _outer.AddBlobAsync(NamespaceId, prefix, null, cancellationToken);
 
-						// Write it to the backend
-						string path = GetBlobPath(locator);
-						memoryStream.Position = 0;
-						await Backend.WriteAsync(path, memoryStream, cancellationToken);
-					}
-				}
-				else
-				{
-					// Add the blob record
-					locator = await _outer.AddBlobAsync(NamespaceId, prefix, null, cancellationToken);
+				// Write it to the backend
+				string path = GetBlobPath(locator);
+				using ReadOnlySequenceStream memoryStream = new ReadOnlySequenceStream(bundle.AsSequence());
+				await Backend.WriteAsync(path, memoryStream, cancellationToken);
 
-					// Write it to the backend
-					string path = GetBlobPath(locator);
-					await Backend.WriteAsync(path, stream, cancellationToken);
-				}
 				return locator;
 			}
 
@@ -1041,40 +1019,23 @@ namespace Horde.Server.Storage
 			int fetchSize = 64 * 1024;
 			for (; ; )
 			{
-				using (IMemoryOwner<byte> owner = MemoryPool<byte>.Shared.Rent(fetchSize))
+				// Read the start of the blob
+				ReadOnlyMemory<byte> memory = await store.ReadBundleRangeAsync(locator, 0, fetchSize, cancellationToken);
+				if (memory.Length < BundleHeader.PreludeLength)
 				{
-					using (Stream? stream = await store.ReadBlobRangeAsync(locator, 0, fetchSize, cancellationToken))
-					{
-						if (stream == null)
-						{
-							_logger.LogError("Unable to read blob {Blob}", locator);
-							return null;
-						}
-
-						// Read the start of the blob
-						Memory<byte> memory = owner.Memory.Slice(0, fetchSize);
-
-						int length = await stream.ReadGreedyAsync(memory, cancellationToken);
-						if (length < BundleHeader.PreludeLength)
-						{
-							_logger.LogError("Blob {Blob} does not have a valid prelude", locator);
-							return null;
-						}
-
-						memory = memory.Slice(0, length);
-
-						// Make sure it's large enough to hold the header
-						int headerSize = BundleHeader.ReadPrelude(memory.Span);
-						if (headerSize <= fetchSize)
-						{
-							byte[] data = memory.ToArray(); // Need to copy data since rented memory will be disposed
-							return BundleHeader.Read(data);
-						}
+					_logger.LogError("Blob {Blob} does not have a valid prelude", locator);
+					return null;
+				}
 
-						// Increase the fetch size and retry
-						fetchSize = headerSize;
-					}
+				// Make sure it's large enough to hold the header
+				int headerSize = BundleHeader.ReadPrelude(memory.Span);
+				if (headerSize <= fetchSize)
+				{
+					return BundleHeader.Read(memory);
 				}
+
+				// Increase the fetch size and retry
+				fetchSize = headerSize;
 			}
 		}
 
diff --git a/Engine/Source/Programs/Shared/EpicGames.Core/StreamUtils.cs b/Engine/Source/Programs/Shared/EpicGames.Core/StreamUtils.cs
--- a/Engine/Source/Programs/Shared/EpicGames.Core/StreamUtils.cs
+++ b/Engine/Source/Programs/Shared/EpicGames.Core/StreamUtils.cs
@@ -53,6 +53,19 @@ namespace EpicGames.Core
 			}
 		}
 
+		/// <summary>
+		/// Read the contents of a file async using double buffering
+		/// </summary>
+		/// <param name="stream">Data to compute the hash for</param>
+		/// <param name="cancellationToken">Cancellation token used to terminate processing</param>
+		/// <returns>New content hash instance containing the hash of the data</returns>
+		public static async Task<byte[]> ReadAllBytesAsync(this Stream stream, CancellationToken cancellationToken = default)
+		{
+			using MemoryStream memoryStream = new MemoryStream();
+			await stream.CopyToAsync(memoryStream, cancellationToken);
+			return memoryStream.ToArray();
+		}
+
 		/// <summary>
 		/// Read the contents of a file async using double buffering
 		/// </summary>
diff --git a/Engine/Source/Programs/Shared/EpicGames.Horde.Tests/BundleStoreTests.cs b/Engine/Source/Programs/Shared/EpicGames.Horde.Tests/BundleStoreTests.cs
--- a/Engine/Source/Programs/Shared/EpicGames.Horde.Tests/BundleStoreTests.cs
+++ b/Engine/Source/Programs/Shared/EpicGames.Horde.Tests/BundleStoreTests.cs
@@ -94,7 +94,7 @@ namespace EpicGames.Horde.Tests
 			MemoryStorageClient blobStore = new MemoryStorageClient();
 			await TestTreeAsync(blobStore, new BundleOptions { MaxBlobSize = 1024 * 1024 });
 
-			Assert.AreEqual(1, blobStore.Blobs.Count);
+			Assert.AreEqual(1, blobStore.Bundles.Count);
 			Assert.AreEqual(1, blobStore.Refs.Count);
 		}
 
@@ -106,7 +106,7 @@ namespace EpicGames.Horde.Tests
 			MemoryStorageClient blobStore = new MemoryStorageClient();
 			await TestTreeAsync(blobStore, new BundleOptions { MaxBlobSize = 1 });
 
-			Assert.AreEqual(5, blobStore.Blobs.Count);
+			Assert.AreEqual(5, blobStore.Bundles.Count);
 			Assert.AreEqual(1, blobStore.Refs.Count);
 		}
 
@@ -387,7 +387,7 @@ namespace EpicGames.Horde.Tests
 
 				NodeRef<ChunkedDataNode> file = root.GetFileEntry("test");
 
-				long uniqueSize = store.Blobs.Values.SelectMany(x => x.Header.Packets).Sum(x => x.DecodedLength);
+				long uniqueSize = store.Bundles.Values.SelectMany(x => x.Header.Packets).Sum(x => x.DecodedLength);
 				Assert.IsTrue(uniqueSize < data.Length / 3); // random fraction meaning "lots of dedupe happened"
 			}
 		}
diff --git a/Engine/Source/Programs/Shared/EpicGames.Horde.Tests/BundleTests.cs b/Engine/Source/Programs/Shared/EpicGames.Horde.Tests/BundleTests.cs
--- a/Engine/Source/Programs/Shared/EpicGames.Horde.Tests/BundleTests.cs
+++ b/Engine/Source/Programs/Shared/EpicGames.Horde.Tests/BundleTests.cs
@@ -213,7 +213,7 @@ namespace EpicGames.Horde.Tests
 
 			// Should be stored inline
 			Assert.AreEqual(1, store.Refs.Count);
-			Assert.AreEqual(1, store.Blobs.Count);
+			Assert.AreEqual(1, store.Bundles.Count);
 
 			// Check the ref
 			Storage.BlobHandle refTarget = await store.ReadRefTargetAsync(refName);
@@ -242,7 +242,7 @@ namespace EpicGames.Horde.Tests
 		public async Task DedupTests()
 		{
 			Assert.AreEqual(0, _storage.Refs.Count);
-			Assert.AreEqual(0, _storage.Blobs.Count);
+			Assert.AreEqual(0, _storage.Bundles.Count);
 
 			BundleOptions options = new BundleOptions();
 			options.MaxBlobSize = 1;
@@ -263,7 +263,7 @@ namespace EpicGames.Horde.Tests
 			}
 
 			Assert.AreEqual(1, _storage.Refs.Count);
-			Assert.AreEqual(2, _storage.Blobs.Count);
+			Assert.AreEqual(2, _storage.Bundles.Count);
 		}
 
 		[TestMethod]
@@ -288,7 +288,7 @@ namespace EpicGames.Horde.Tests
 				}
 
 				Assert.AreEqual(1, _storage.Refs.Count);
-				Assert.AreEqual(5, _storage.Blobs.Count);
+				Assert.AreEqual(5, _storage.Bundles.Count);
 			}
 
 			{
diff --git a/Engine/Source/Programs/Shared/EpicGames.Horde/Compute/AgentMessage.cs b/Engine/Source/Programs/Shared/EpicGames.Horde/Compute/AgentMessage.cs
--- a/Engine/Source/Programs/Shared/EpicGames.Horde/Compute/AgentMessage.cs
+++ b/Engine/Source/Programs/Shared/EpicGames.Horde/Compute/AgentMessage.cs
@@ -502,7 +502,7 @@ namespace EpicGames.Horde.Compute
 		/// <param name="length">Length of data to return</param>
 		/// <param name="cancellationToken">Cancellation token for the operation</param>
 		/// <returns>Stream containing the blob data</returns>
-		public static async Task<Stream> ReadBlobAsync(this AgentMessageChannel channel, BlobLocator locator, int offset, int length, CancellationToken cancellationToken = default)
+		public static async Task<ReadOnlyMemory<byte>> ReadBlobAsync(this AgentMessageChannel channel, BlobLocator locator, int offset, int length, CancellationToken cancellationToken = default)
 		{
 			using (IAgentMessageBuilder request = await channel.CreateMessageAsync(AgentMessageType.ReadBlob, cancellationToken))
 			{
@@ -528,13 +528,6 @@ namespace EpicGames.Horde.Compute
 					int chunkLength = response.Data.Length - 8;
 					int totalLength = BinaryPrimitives.ReadInt32LittleEndian(response.Data.Span.Slice(4, 4));
 
-					if (chunkOffset == 0 && chunkLength == totalLength)
-					{
-						BlobDataStream stream = new BlobDataStream(response);
-						response = null;
-						return stream;
-					}
-
 					buffer ??= new byte[totalLength];
 					response.Data.Slice(8).CopyTo(buffer.AsMemory(chunkOffset));
 
@@ -550,7 +543,7 @@ namespace EpicGames.Horde.Compute
 				}
 			}
 
-			return new ReadOnlyMemoryStream(buffer);
+			return buffer;
 		}
 
 		/// <summary>
@@ -576,35 +569,27 @@ namespace EpicGames.Horde.Compute
 		/// <param name="cancellationToken">Cancellation token for the operation</param>
 		public static async Task SendBlobDataAsync(this AgentMessageChannel channel, BlobLocator locator, int offset, int length, IStorageClient storage, CancellationToken cancellationToken = default)
 		{
-			byte[] data;
+			ReadOnlySequence<byte> data;
 			if (offset == 0 && length == 0)
 			{
-				using (Stream stream = await storage.ReadBlobAsync(locator, cancellationToken))
-				{
-					using MemoryStream target = new MemoryStream();
-					await stream.CopyToAsync(target, cancellationToken);
-					data = target.ToArray();
-				}
+				Bundle bundle = await storage.ReadBundleAsync(locator, cancellationToken);
+				data = bundle.AsSequence();
 			}
 			else
 			{
-				using (Stream stream = await storage.ReadBlobRangeAsync(locator, offset, length, cancellationToken))
-				{
-					using MemoryStream target = new MemoryStream();
-					await stream.CopyToAsync(target, cancellationToken);
-					data = target.ToArray();
-				}
+				ReadOnlyMemory<byte> range = await storage.ReadBundleRangeAsync(locator, offset, length, cancellationToken);
+				data = new ReadOnlySequence<byte>(range);
 			}
 
 			const int MaxChunkSize = 512 * 1024;
 			for (int chunkOffset = 0; chunkOffset < data.Length;)
 			{
-				int chunkLength = Math.Min(data.Length - chunkOffset, MaxChunkSize);
+				int chunkLength = (int)Math.Min(data.Length - chunkOffset, MaxChunkSize);
 				using (IAgentMessageBuilder response = await channel.CreateMessageAsync(AgentMessageType.ReadBlobResponse, chunkLength + 128, cancellationToken))
 				{
 					response.WriteInt32(chunkOffset);
-					response.WriteInt32(data.Length);
-					response.WriteFixedLengthBytes(data.AsSpan(chunkOffset, chunkLength));
+					response.WriteInt32((int)data.Length);
+					response.WriteFixedLengthBytes(data.Slice(chunkOffset, chunkLength));
 					response.Send();
 				}
 				chunkOffset += chunkLength;
diff --git a/Engine/Source/Programs/Shared/EpicGames.Horde/Compute/AgentStorageClient.cs b/Engine/Source/Programs/Shared/EpicGames.Horde/Compute/AgentStorageClient.cs
--- a/Engine/Source/Programs/Shared/EpicGames.Horde/Compute/AgentStorageClient.cs
+++ b/Engine/Source/Programs/Shared/EpicGames.Horde/Compute/AgentStorageClient.cs
@@ -39,13 +39,14 @@ namespace EpicGames.Horde.Compute
 		#region Blobs
 
 		/// <inheritdoc/>
-		public override async Task<Stream> ReadBlobAsync(BlobLocator locator, CancellationToken cancellationToken = default)
+		public override async Task<Bundle> ReadBundleAsync(BlobLocator locator, CancellationToken cancellationToken = default)
 		{
-			return await ReadBlobRangeAsync(locator, 0, 0, cancellationToken);
+			ReadOnlyMemory<byte> data = await ReadBundleRangeAsync(locator, 0, 0, cancellationToken);
+			return new Bundle(data);
 		}
 
 		/// <inheritdoc/>
-		public override async Task<Stream> ReadBlobRangeAsync(BlobLocator locator, int offset, int length, CancellationToken cancellationToken = default)
+		public override async Task<ReadOnlyMemory<byte>> ReadBundleRangeAsync(BlobLocator locator, int offset, int length, CancellationToken cancellationToken = default)
 		{
 			await _semaphore.WaitAsync(cancellationToken);
 			try
@@ -59,7 +60,7 @@ namespace EpicGames.Horde.Compute
 		}
 
 		/// <inheritdoc/>
-		public override Task<BlobLocator> WriteBlobAsync(Stream stream, Utf8String prefix = default, CancellationToken cancellationToken = default)
+		public override Task<BlobLocator> WriteBundleAsync(Bundle bundle, Utf8String prefix = default, CancellationToken cancellationToken = default)
 		{
 			throw new NotSupportedException();
 		}
diff --git a/Engine/Source/Programs/Shared/EpicGames.Horde/Storage/Backends/FileStorageClient.cs b/Engine/Source/Programs/Shared/EpicGames.Horde/Storage/Backends/FileStorageClient.cs
--- a/Engine/Source/Programs/Shared/EpicGames.Horde/Storage/Backends/FileStorageClient.cs
+++ b/Engine/Source/Programs/Shared/EpicGames.Horde/Storage/Backends/FileStorageClient.cs
@@ -1,6 +1,7 @@
 // Copyright Epic Games, Inc. All Rights Reserved.
 
 using System;
+using System.Buffers;
 using System.Collections.Generic;
 using System.Diagnostics;
 using System.Globalization;
@@ -52,23 +53,31 @@ namespace EpicGames.Horde.Storage.Backends
 		#region Blobs
 
 		/// <inheritdoc/>
-		public override Task<Stream> ReadBlobAsync(BlobLocator id, CancellationToken cancellationToken = default)
+		public override async Task<Bundle> ReadBundleAsync(BlobLocator id, CancellationToken cancellationToken = default)
 		{
 			FileReference file = GetBlobFile(id);
 			_logger.LogInformation("Reading {File}", file);
-			return Task.FromResult<Stream>(FileReference.Open(file, FileMode.Open, FileAccess.Read, FileShare.Read));
+
+			byte[] data = await FileReference.ReadAllBytesAsync(file, cancellationToken);
+			return new Bundle(data);
 		}
 
 		/// <inheritdoc/>
-		public override async Task<Stream> ReadBlobRangeAsync(BlobLocator id, int offset, int length, CancellationToken cancellationToken = default)
+		public override async Task<ReadOnlyMemory<byte>> ReadBundleRangeAsync(BlobLocator id, int offset, int length, CancellationToken cancellationToken = default)
 		{
-			Stream stream = await ReadBlobAsync(id, cancellationToken);
-			stream.Seek(offset, SeekOrigin.Begin);
-			return stream;
+			Bundle bundle = await ReadBundleAsync(id, cancellationToken);
+
+			ReadOnlySequence<byte> sequence = bundle.AsSequence().Slice(offset);
+			if (sequence.Length > length)
+			{
+				sequence = sequence.Slice(0, length);
+			}
+
+			return sequence.AsSingleSegment();
 		}
 
 		/// <inheritdoc/>
-		public override async Task<BlobLocator> WriteBlobAsync(Stream stream, Utf8String prefix = default, CancellationToken cancellationToken = default)
+		public override async Task<BlobLocator> WriteBundleAsync(Bundle bundle, Utf8String prefix = default, CancellationToken cancellationToken = default)
 		{
 			BlobLocator id = BlobLocator.CreateUnique(prefix);
 			FileReference file = GetBlobFile(id);
@@ -77,7 +86,10 @@ namespace EpicGames.Horde.Storage.Backends
 
 			using (FileStream fileStream = FileReference.Open(file, FileMode.Create, FileAccess.Write, FileShare.ReadWrite))
 			{
-				await stream.CopyToAsync(fileStream, cancellationToken);
+				foreach (ReadOnlyMemory<byte> segment in bundle.AsSequence())
+				{
+					await fileStream.WriteAsync(segment, cancellationToken);
+				}
 			}
 
 			return id;
diff --git a/Engine/Source/Programs/Shared/EpicGames.Horde/Storage/Backends/HttpStorageClient.cs b/Engine/Source/Programs/Shared/EpicGames.Horde/Storage/Backends/HttpStorageClient.cs
--- a/Engine/Source/Programs/Shared/EpicGames.Horde/Storage/Backends/HttpStorageClient.cs
+++ b/Engine/Source/Programs/Shared/EpicGames.Horde/Storage/Backends/HttpStorageClient.cs
@@ -102,7 +102,7 @@ namespace EpicGames.Horde.Storage.Backends
 		#region Blobs
 
 		/// <inheritdoc/>
-		public override async Task<Stream> ReadBlobAsync(BlobLocator locator, CancellationToken cancellationToken = default)
+		public override async Task<Bundle> ReadBundleAsync(BlobLocator locator, CancellationToken cancellationToken = default)
 		{
 			_logger.LogDebug("Reading {Locator}", locator);
 			using (HttpClient httpClient = _createClient())
@@ -111,13 +111,15 @@ namespace EpicGames.Horde.Storage.Backends
 				{
 					HttpResponseMessage response = await httpClient.SendAsync(request, cancellationToken);
 					response.EnsureSuccessStatusCode();
-					return await response.Content.ReadAsStreamAsync(cancellationToken);
+
+					byte[] data = await response.Content.ReadAsByteArrayAsync(cancellationToken);
+					return new Bundle(data);
 				}
 			}
 		}
 
 		/// <inheritdoc/>
-		public override async Task<Stream> ReadBlobRangeAsync(BlobLocator locator, int offset, int length, CancellationToken cancellationToken = default)
+		public override async Task<ReadOnlyMemory<byte>> ReadBundleRangeAsync(BlobLocator locator, int offset, int length, CancellationToken cancellationToken = default)
 		{
 			_logger.LogDebug("Reading {Locator} ({Offset}+{Length})", locator, offset, length);
 			using (HttpClient httpClient = _createClient())
@@ -126,15 +128,15 @@ namespace EpicGames.Horde.Storage.Backends
 				{
 					HttpResponseMessage response = await httpClient.SendAsync(request, cancellationToken);
 					response.EnsureSuccessStatusCode();
-					return await response.Content.ReadAsStreamAsync(cancellationToken);
+					return await response.Content.ReadAsByteArrayAsync(cancellationToken);
 				}
 			}
 		}
 
 		/// <inheritdoc/>
-		public override async Task<BlobLocator> WriteBlobAsync(Stream stream, Utf8String prefix = default, CancellationToken cancellationToken = default)
+		public override async Task<BlobLocator> WriteBundleAsync(Bundle bundle, Utf8String prefix = default, CancellationToken cancellationToken = default)
 		{
-			using StreamContent streamContent = new StreamContent(stream);
+			using StreamContent streamContent = new StreamContent(new ReadOnlySequenceStream(bundle.AsSequence()));
 
 			if (_supportsUploadRedirects)
 			{
diff --git a/Engine/Source/Programs/Shared/EpicGames.Horde/Storage/Backends/MemoryStorageClient.cs b/Engine/Source/Programs/Shared/EpicGames.Horde/Storage/Backends/MemoryStorageClient.cs
--- a/Engine/Source/Programs/Shared/EpicGames.Horde/Storage/Backends/MemoryStorageClient.cs
+++ b/Engine/Source/Programs/Shared/EpicGames.Horde/Storage/Backends/MemoryStorageClient.cs
@@ -25,7 +25,7 @@ namespace EpicGames.Horde.Storage.Backends
 		/// <summary>
 		/// Map of blob id to blob data
 		/// </summary>
-		readonly ConcurrentDictionary<BlobLocator, Bundle> _blobs = new ConcurrentDictionary<BlobLocator, Bundle>();
+		readonly ConcurrentDictionary<BlobLocator, Bundle> _bundles = new ConcurrentDictionary<BlobLocator, Bundle>();
 
 		/// <summary>
 		/// Map of ref name to ref data
@@ -37,8 +37,8 @@ namespace EpicGames.Horde.Storage.Backends
 		/// </summary>
 		readonly ConcurrentDictionary<Utf8String, ExportEntry> _exports = new ConcurrentDictionary<Utf8String, ExportEntry>();
 
-		/// <inheritdoc cref="_blobs"/>
-		public IReadOnlyDictionary<BlobLocator, Bundle> Blobs => _blobs;
+		/// <inheritdoc cref="_bundles"/>
+		public IReadOnlyDictionary<BlobLocator, Bundle> Bundles => _bundles;
 
 		/// <inheritdoc cref="_refs"/>
 		public IReadOnlyDictionary<RefName, BundleNodeLocator> Refs => _refs;
@@ -54,27 +54,29 @@ namespace EpicGames.Horde.Storage.Backends
 		#region Blobs
 
 		/// <inheritdoc/>
-		public override Task<Stream> ReadBlobAsync(BlobLocator locator, CancellationToken cancellationToken = default)
+		public override Task<Bundle> ReadBundleAsync(BlobLocator locator, CancellationToken cancellationToken = default)
 		{
-			Stream stream = new ReadOnlySequenceStream(_blobs[locator].AsSequence());
-			return Task.FromResult(stream);
+			Bundle bundle = _bundles[locator];
+			return Task.FromResult(bundle);
 		}
 
 		/// <inheritdoc/>
-		public override Task<Stream> ReadBlobRangeAsync(BlobLocator locator, int offset, int length, CancellationToken cancellationToken = default)
+		public override Task<ReadOnlyMemory<byte>> ReadBundleRangeAsync(BlobLocator locator, int offset, int length, CancellationToken cancellationToken = default)
 		{
-			Stream stream = new ReadOnlySequenceStream(_blobs[locator].AsSequence().Slice(offset));
-			return Task.FromResult(stream);
+			ReadOnlySequence<byte> sequence = _bundles[locator].AsSequence().Slice(offset);
+			if (sequence.Length > length)
+			{
+				sequence = sequence.Slice(0, length);
+			}
+			return Task.FromResult(sequence.AsSingleSegment());
 		}
 
 		/// <inheritdoc/>
-		public override async Task<BlobLocator> WriteBlobAsync(Stream stream, Utf8String prefix = default, CancellationToken cancellationToken = default)
+		public override Task<BlobLocator> WriteBundleAsync(Bundle bundle, Utf8String prefix = default, CancellationToken cancellationToken = default)
 		{
 			BlobLocator locator = BlobLocator.CreateUnique(prefix);
-			Bundle bundle = await Bundle.FromStreamAsync(stream, cancellationToken);
-			_blobs[locator] = bundle;
-
-			return locator;
+			_bundles[locator] = new Bundle(bundle.AsSequence().ToArray());
+			return Task.FromResult(locator);
 		}
 
 		#endregion
diff --git a/Engine/Source/Programs/Shared/EpicGames.Horde/Storage/BundleReader.cs b/Engine/Source/Programs/Shared/EpicGames.Horde/Storage/BundleReader.cs
--- a/Engine/Source/Programs/Shared/EpicGames.Horde/Storage/BundleReader.cs
+++ b/Engine/Source/Programs/Shared/EpicGames.Horde/Storage/BundleReader.cs
@@ -232,7 +232,7 @@ namespace EpicGames.Horde.Storage
 				{
 					// Read the prefetch size from the blob
 					Memory<byte> memory = owner.Memory.Slice(0, prefetchSize);
-					memory = await _store.ReadBlobRangeAsync(queuedHeader.Blob, 0, memory, cancellationToken);
+					memory = await _store.ReadBundleRangeAsync(queuedHeader.Blob, 0, memory, cancellationToken);
 
 					// Make sure it's large enough to hold the header
 					int headerSize = BundleHeader.ReadPrelude(memory.Span);
@@ -330,7 +330,7 @@ namespace EpicGames.Horde.Storage
 			using (IMemoryOwner<byte> owner = MemoryPool<byte>.Shared.Rent(readLength))
 			{
 				Memory<byte> buffer = owner.Memory.Slice(0, readLength);
-				buffer = await _store.ReadBlobRangeAsync(bundleInfo.Locator, bundleInfo.HeaderLength + bundleInfo.Header.Packets[minPacketIdx].EncodedOffset, buffer, cancellationToken);
+				buffer = await _store.ReadBundleRangeAsync(bundleInfo.Locator, bundleInfo.HeaderLength + bundleInfo.Header.Packets[minPacketIdx].EncodedOffset, buffer, cancellationToken);
 
 				// Copy all the packets that have been read into separate buffers, so we can cache them individually.
 				ReadOnlyMemory<byte>?[] packets = new ReadOnlyMemory<byte>?[maxPacketIdx - minPacketIdx];
@@ -429,9 +429,6 @@ namespace EpicGames.Horde.Storage
 		/// <returns>The decoded data</returns>
 		async ValueTask<ReadOnlyMemory<byte>> ReadBundlePacketAsync(BundleInfo bundleInfo, int packetIdx, CancellationToken cancellationToken)
 		{
-			char c = bundleInfo.Locator.ToString()[0];
-			Debug.Assert(Char.IsLetterOrDigit(c));
-
 			if (packetIdx < 0 || packetIdx >= bundleInfo.Header.Packets.Count)
 			{
 				throw new ArgumentException("Packet index is out of range", nameof(packetIdx));
@@ -496,9 +493,6 @@ namespace EpicGames.Horde.Storage
 		/// <returns>The encoded packet data</returns>
 		async ValueTask<ReadOnlyMemory<byte>> ReadEncodedPacketAsync(BundleInfo bundleInfo, int packetIdx)
 		{
-			char c = bundleInfo.Locator.ToString()[0];
-			Debug.Assert(Char.IsLetterOrDigit(c));
-
 			if (packetIdx < 0 || packetIdx >= bundleInfo.Header.Packets.Count)
 			{
 				throw new ArgumentException("Packet index is out of range", nameof(packetIdx));
diff --git a/Engine/Source/Programs/Shared/EpicGames.Horde/Storage/BundleStorageClient.cs b/Engine/Source/Programs/Shared/EpicGames.Horde/Storage/BundleStorageClient.cs
--- a/Engine/Source/Programs/Shared/EpicGames.Horde/Storage/BundleStorageClient.cs
+++ b/Engine/Source/Programs/Shared/EpicGames.Horde/Storage/BundleStorageClient.cs
@@ -32,13 +32,13 @@ namespace EpicGames.Horde.Storage
 		#region Blobs
 
 		/// <inheritdoc/>
-		public abstract Task<Stream> ReadBlobAsync(BlobLocator locator, CancellationToken cancellationToken = default);
+		public abstract Task<Bundle> ReadBundleAsync(BlobLocator locator, CancellationToken cancellationToken = default);
 
 		/// <inheritdoc/>
-		public abstract Task<Stream> ReadBlobRangeAsync(BlobLocator locator, int offset, int length, CancellationToken cancellationToken = default);
+		public abstract Task<ReadOnlyMemory<byte>> ReadBundleRangeAsync(BlobLocator locator, int offset, int length, CancellationToken cancellationToken = default);
 
 		/// <inheritdoc/>
-		public abstract Task<BlobLocator> WriteBlobAsync(Stream stream, Utf8String prefix = default, CancellationToken cancellationToken = default);
+		public abstract Task<BlobLocator> WriteBundleAsync(Bundle bundle, Utf8String prefix = default, CancellationToken cancellationToken = default);
 
 		#endregion
 
@@ -81,43 +81,4 @@ namespace EpicGames.Horde.Storage
 
 		#endregion
 	}
-
-	/// <summary>
-	/// Extension methods for serializing bundles
-	/// </summary>
-	public static class BundleStorageClientExtensions
-	{
-		#region Bundles
-
-		/// <summary>
-		/// Reads a bundle from the given blob id, or retrieves it from the cache
-		/// </summary>
-		/// <param name="store">The store instance to read from</param>
-		/// <param name="locator"></param>
-		/// <param name="cancellationToken"></param>
-		/// <returns></returns>
-		public static async Task<Bundle> ReadBundleAsync(this IStorageClient store, BlobLocator locator, CancellationToken cancellationToken = default)
-		{
-			using (Stream stream = await store.ReadBlobAsync(locator, cancellationToken))
-			{
-				return await Bundle.FromStreamAsync(stream, cancellationToken);
-			}
-		}
-
-		/// <summary>
-		/// Writes a new bundle to the store
-		/// </summary>
-		/// <param name="store">The store instance to write to</param>
-		/// <param name="bundle">Bundle data</param>
-		/// <param name="cancellationToken">Cancellation token for the operation</param>
-		/// <param name="prefix">Prefix for blob names. While the returned BlobId is guaranteed to be unique, this name can be used as a prefix to aid debugging.</param>
-		/// <returns>Unique identifier for the blob</returns>
-		public static async Task<BlobLocator> WriteBundleAsync(this IStorageClient store, Bundle bundle, Utf8String prefix = default, CancellationToken cancellationToken = default)
-		{
-			using ReadOnlySequenceStream stream = new ReadOnlySequenceStream(bundle.AsSequence());
-			return await store.WriteBlobAsync(stream, prefix, cancellationToken);
-		}
-
-		#endregion
-	}
 }
diff --git a/Engine/Source/Programs/Shared/EpicGames.Horde/Storage/IStorageClient.cs b/Engine/Source/Programs/Shared/EpicGames.Horde/Storage/IStorageClient.cs
--- a/Engine/Source/Programs/Shared/EpicGames.Horde/Storage/IStorageClient.cs
+++ b/Engine/Source/Programs/Shared/EpicGames.Horde/Storage/IStorageClient.cs
@@ -73,7 +73,7 @@ namespace EpicGames.Horde.Storage
 	/// </summary>
 	public interface IStorageClient
 	{
-		#region Blobs
+		#region Bundles
 
 		/// <summary>
 		/// Reads raw data for a blob from the store
@@ -81,7 +81,7 @@ namespace EpicGames.Horde.Storage
 		/// <param name="locator">The blob locator</param>
 		/// <param name="cancellationToken">Cancellation token for the operation</param>
 		/// <returns>Stream containing the data</returns>
-		Task<Stream> ReadBlobAsync(BlobLocator locator, CancellationToken cancellationToken = default);
+		Task<Bundle> ReadBundleAsync(BlobLocator locator, CancellationToken cancellationToken = default);
 
 		/// <summary>
 		/// Reads a ranged chunk from a blob
@@ -90,16 +90,16 @@ namespace EpicGames.Horde.Storage
 		/// <param name="offset">Starting offset for the data to read</param>
 		/// <param name="length">Length of the data</param>
 		/// <param name="cancellationToken">Cancellation token for the operation</param>
-		Task<Stream> ReadBlobRangeAsync(BlobLocator locator, int offset, int length, CancellationToken cancellationToken = default);
+		Task<ReadOnlyMemory<byte>> ReadBundleRangeAsync(BlobLocator locator, int offset, int length, CancellationToken cancellationToken = default);
 
 		/// <summary>
 		/// Writes a new blob to the store
 		/// </summary>
-		/// <param name="stream">Blob data</param>
+		/// <param name="bundle">The bundle to write</param>
 		/// <param name="cancellationToken">Cancellation token for the operation</param>
 		/// <param name="prefix">Prefix for blob names. While the returned BlobId is guaranteed to be unique, this name can be used as a prefix to aid debugging.</param>
 		/// <returns>Unique identifier for the blob</returns>
-		Task<BlobLocator> WriteBlobAsync(Stream stream, Utf8String prefix = default, CancellationToken cancellationToken = default);
+		Task<BlobLocator> WriteBundleAsync(Bundle bundle, Utf8String prefix = default, CancellationToken cancellationToken = default);
 
 		#endregion
 
@@ -287,22 +287,11 @@ namespace EpicGames.Horde.Storage
 		/// <param name="memory">Buffer to read into</param>
 		/// <param name="cancellationToken">Cancellation token for the operation</param>
 		/// <returns>The data that was read</returns>
-		public static async Task<Memory<byte>> ReadBlobRangeAsync(this IStorageClient store, BlobLocator locator, int offset, Memory<byte> memory, CancellationToken cancellationToken = default)
+		public static async Task<Memory<byte>> ReadBundleRangeAsync(this IStorageClient store, BlobLocator locator, int offset, Memory<byte> memory, CancellationToken cancellationToken = default)
 		{
-			using (Stream stream = await store.ReadBlobRangeAsync(locator, offset, memory.Length, cancellationToken))
-			{
-				int length = 0;
-				while (length < memory.Length)
-				{
-					int readBytes = await stream.ReadAsync(memory.Slice(length), cancellationToken);
-					if (readBytes == 0)
-					{
-						break;
-					}
-					length += readBytes;
-				}
-				return memory.Slice(0, length);
-			}
+			ReadOnlyMemory<byte> buffer = await store.ReadBundleRangeAsync(locator, offset, memory.Length, cancellationToken);
+			buffer.CopyTo(memory);
+			return memory.Slice(0, buffer.Length);
 		}
 
 		#endregion
diff --git a/Engine/Source/Programs/UnrealCloudDDC/Jupiter/Controllers/BundlesController.cs b/Engine/Source/Programs/UnrealCloudDDC/Jupiter/Controllers/BundlesController.cs
--- a/Engine/Source/Programs/UnrealCloudDDC/Jupiter/Controllers/BundlesController.cs
+++ b/Engine/Source/Programs/UnrealCloudDDC/Jupiter/Controllers/BundlesController.cs
@@ -226,7 +226,8 @@ namespace Jupiter.Controllers
             {
                 using (Stream stream = file.OpenReadStream())
                 {
-                    BlobLocator locator = await client.WriteBlobAsync(stream, prefix: (prefix == null) ? Utf8String.Empty : new Utf8String(prefix), cancellationToken: cancellationToken);
+                    Bundle bundle = await Bundle.FromStreamAsync(stream, cancellationToken);
+                    BlobLocator locator = await client.WriteBundleAsync(bundle, prefix: (prefix == null) ? Utf8String.Empty : new Utf8String(prefix), cancellationToken: cancellationToken);
                     return new WriteBlobResponse { Blob = locator, SupportsRedirects = client.SupportsRedirects? (bool?)true : null };
                 }
             }
@@ -263,19 +264,19 @@ namespace Jupiter.Controllers
             Stream stream;
             if (offset == null && length == null)
             {
-                stream = await client.ReadBlobAsync(locator, cancellationToken);
+                Bundle bundle = await client.ReadBundleAsync(locator, cancellationToken);
+                stream = new ReadOnlySequenceStream(bundle.AsSequence());
             }
             else if (offset != null && length != null)
             {
-                stream = await client.ReadBlobRangeAsync(locator, offset.Value, length.Value, cancellationToken);
-
-                return File(stream, MediaTypeNames.Application.Octet);
+                ReadOnlyMemory<byte> memory = await client.ReadBundleRangeAsync(locator, offset.Value, length.Value, cancellationToken);
+                stream = new ReadOnlyMemoryStream(memory);
             }
             else
             {
                 return BadRequest("Offset and length must both be specified as query parameters for ranged reads");
             }
-            return File(stream, "application/octet-stream");
+            return File(stream, MediaTypeNames.Application.Octet);
 #pragma warning restore CA2000 // Dispose objects before losing scope
         }
 
diff --git a/Engine/Source/Programs/UnrealCloudDDC/Jupiter/Implementation/Bundles/StorageService.cs b/Engine/Source/Programs/UnrealCloudDDC/Jupiter/Implementation/Bundles/StorageService.cs
--- a/Engine/Source/Programs/UnrealCloudDDC/Jupiter/Implementation/Bundles/StorageService.cs
+++ b/Engine/Source/Programs/UnrealCloudDDC/Jupiter/Implementation/Bundles/StorageService.cs
@@ -1,6 +1,7 @@
 // Copyright Epic Games, Inc. All Rights Reserved.
 
 using System;
+using System.Buffers;
 using System.Collections.Concurrent;
 using System.Collections.Generic;
 using System.IO;
@@ -78,39 +79,22 @@ public class StorageClient : IStorageClientJupiter
         return (locator, redirectUri);
     }
 
-    public async Task<BlobLocator> WriteBlobAsync(Stream stream, Utf8String prefix, CancellationToken cancellationToken)
+    public async Task<BlobLocator> WriteBundleAsync(Bundle bundle, Utf8String prefix, CancellationToken cancellationToken)
     {
         BlobLocator locator = BlobLocator.CreateUnique(prefix);
         BlobIdentifier blobIdentifier = BlobIdentifier.FromBlobLocator(locator);
-        await using MemoryStream ms = new MemoryStream();
-        await stream.CopyToAsync(ms, cancellationToken);
-        byte[] blob = ms.ToArray();
-        await _blobService.PutObject(_namespaceId, blob, blobIdentifier);
+        await _blobService.PutObject(_namespaceId, bundle.AsSequence().ToArray(), blobIdentifier);
 
-        bool isBundle;
-        try
-        {
-            BundleHeader.ReadPrelude(blob);
-            isBundle = true;
-        }
-        catch (InvalidDataException)
+        await using ReadOnlySequenceStream bundleStream = new ReadOnlySequenceStream(bundle.AsSequence());
+        BundleHeader bundleHeader = await BundleHeader.FromStreamAsync(bundleStream, cancellationToken);
+        List<Task> addReferencesTasks = new List<Task>();
+        foreach (BlobLocator import in bundleHeader.Imports)
         {
-            isBundle = false;
+            BlobIdentifier dependentBlob = BlobIdentifier.FromBlobLocator(import);
+            addReferencesTasks.Add(_blobIndex.AddBlobReferences(_namespaceId, dependentBlob, blobIdentifier));
         }
 
-        if (isBundle)
-        {
-            await using MemoryStream bundleStream = new MemoryStream(blob);
-            BundleHeader bundleHeader = await BundleHeader.FromStreamAsync(bundleStream, cancellationToken);
-            List<Task> addReferencesTasks = new List<Task>();
-            foreach (BlobLocator import in bundleHeader.Imports)
-            {
-                BlobIdentifier dependentBlob = BlobIdentifier.FromBlobLocator(import);
-                addReferencesTasks.Add(_blobIndex.AddBlobReferences(_namespaceId, dependentBlob, blobIdentifier));
-            }
-
-            await Task.WhenAll(addReferencesTasks);
-        }
+        await Task.WhenAll(addReferencesTasks);
         
         return locator;
     }
@@ -141,22 +125,24 @@ public class StorageClient : IStorageClientJupiter
         return redirectUri;
     }
 
-    public async Task<Stream> ReadBlobAsync(BlobLocator locator, CancellationToken cancellationToken)
+    public async Task<Bundle> ReadBundleAsync(BlobLocator locator, CancellationToken cancellationToken)
     {
         BlobIdentifier blobIdentifier = BlobIdentifier.FromBlobLocator(locator);
         BlobContents blobContents = await _blobService.GetObject(_namespaceId, blobIdentifier);
-        return blobContents.Stream;
+        return await Bundle.FromStreamAsync(blobContents.Stream, cancellationToken);
     }
 
-    public async Task<Stream> ReadBlobRangeAsync(BlobLocator locator, int offset, int length, CancellationToken cancellationToken)
+    public async Task<ReadOnlyMemory<byte>> ReadBundleRangeAsync(BlobLocator locator, int offset, int length, CancellationToken cancellationToken)
     {
-        await using Stream stream = await ReadBlobAsync(locator, cancellationToken);
-        using BinaryReader br = new BinaryReader(stream);
-        // skip in the stream, TODO: will fail if stream is very large
-        br.ReadBytes(offset);
-        byte[] data = br.ReadBytes(length);
-        return new MemoryStream(data);
+        ReadOnlySequence<byte> sequence = await ReadBundleAsync(locator, cancellationToken);
+		sequence = sequence.Slice(offset);
+
+		if(sequence.Length > length)
+		{
+			sequence = sequence.Slice(0, length);
+		}
 
+        return sequence.AsSingleSegment();
     }
 
     public async Task<BlobHandle?> TryReadRefTargetAsync(RefName name, RefCacheTime cacheTime = default, CancellationToken cancellationToken = default)
