Horde: Adding device reporting service
#rnx

diff --git a/Engine/Source/Programs/Horde/Horde.Server.Tests/NotificationServiceTests.cs b/Engine/Source/Programs/Horde/Horde.Server.Tests/NotificationServiceTests.cs
--- a/Engine/Source/Programs/Horde/Horde.Server.Tests/NotificationServiceTests.cs
+++ b/Engine/Source/Programs/Horde/Horde.Server.Tests/NotificationServiceTests.cs
@@ -42,6 +42,7 @@ namespace Horde.Server.Tests
 		public Task NotifyConfigUpdateFailureAsync(string errorMessage, string fileName, int? change = null, IUser? author = null, string? description = null) { throw new NotImplementedException(); }
 		public Task NotifyDeviceServiceAsync(string message, IDevice? device = null, IDevicePool? pool = null, StreamConfig? stream = null, IJob? job = null, IJobStep? step = null, INode? node = null, IUser? user = null) { throw new NotImplementedException(); }
 		public Task SendIssueReportAsync(IssueReportGroup report) => throw new NotImplementedException();
+		public Task SendDeviceIssueReportAsync(DeviceIssueReport report) => throw new NotImplementedException();
 	}
 
 	[TestClass]
diff --git a/Engine/Source/Programs/Horde/Horde.Server/Dashboard/DashboardController.cs b/Engine/Source/Programs/Horde/Horde.Server/Dashboard/DashboardController.cs
--- a/Engine/Source/Programs/Horde/Horde.Server/Dashboard/DashboardController.cs
+++ b/Engine/Source/Programs/Horde/Horde.Server/Dashboard/DashboardController.cs
@@ -140,6 +140,8 @@ namespace Horde.Server.Dashboard
 			dashboardConfigResponse.HelpEmailAddress = _settings.HelpEmailAddress;
 			dashboardConfigResponse.HelpSlackChannel = _settings.HelpSlackChannel;
 
+			dashboardConfigResponse.DeviceProblemCooldownMinutes = _settings.DeviceProblemCooldownMinutes;
+
 			return dashboardConfigResponse;
 		}
 
diff --git a/Engine/Source/Programs/Horde/Horde.Server/Dashboard/DashboardMessages.cs b/Engine/Source/Programs/Horde/Horde.Server/Dashboard/DashboardMessages.cs
--- a/Engine/Source/Programs/Horde/Horde.Server/Dashboard/DashboardMessages.cs
+++ b/Engine/Source/Programs/Horde/Horde.Server/Dashboard/DashboardMessages.cs
@@ -34,6 +34,11 @@ namespace Horde.Server.Dashboard
 		/// Help slack channel that users can use for issues
 		/// </summary>
 		public string? HelpSlackChannel { get; set; }
+
+		/// <summary>
+		/// Device problem cooldown in minutes
+		/// </summary>
+		public int DeviceProblemCooldownMinutes { get; set; }
 	}
 
 	/// <summary>
diff --git a/Engine/Source/Programs/Horde/Horde.Server/Devices/DeviceCollection.cs b/Engine/Source/Programs/Horde/Horde.Server/Devices/DeviceCollection.cs
--- a/Engine/Source/Programs/Horde/Horde.Server/Devices/DeviceCollection.cs
+++ b/Engine/Source/Programs/Horde/Horde.Server/Devices/DeviceCollection.cs
@@ -794,7 +794,7 @@ namespace Horde.Server.Devices
 		}
 
 		/// <inheritdoc/>
-		public async Task<IDeviceReservation?> TryAddReservationAsync(DevicePoolId poolId, List<DeviceRequestData> request, string? hostname, string? reservationDetails, string? streamId, string? jobId, string? stepId, string? jobName, string? stepName)
+		public async Task<IDeviceReservation?> TryAddReservationAsync(DevicePoolId poolId, List<DeviceRequestData> request, int problemCooldown, string? hostname, string? reservationDetails, string? streamId, string? jobId, string? stepId, string? jobName, string? stepName)
 		{
 
 			if (request.Count == 0)
@@ -823,7 +823,7 @@ namespace Horde.Server.Devices
 				x.MaintenanceTimeUtc == null).ToListAsync();
 
 			// filter out problem devices
-			poolDevices = poolDevices.FindAll(x => (x.ProblemTimeUtc == null || ((reservationTimeUtc - x.ProblemTimeUtc).Value.TotalMinutes > 30)));
+			poolDevices = poolDevices.FindAll(x => (x.ProblemTimeUtc == null || ((reservationTimeUtc - x.ProblemTimeUtc).Value.TotalMinutes > problemCooldown)));
 
 			// filter out currently reserved devices
 			poolDevices = poolDevices.FindAll(x => poolReservations.FirstOrDefault(p => p.Devices.Contains(x.Id)) == null);
@@ -1109,7 +1109,7 @@ namespace Horde.Server.Devices
 		/// Create a device pool telemetry snapshot
 		/// </summary>
 		/// <returns></returns>
-		public async Task CreatePoolTelemetrySnapshot()
+		public async Task CreatePoolTelemetrySnapshot(int problemCooldown)
 		{
 			List<IDevice> devices = await FindAllDevicesAsync();
 			List<IDevicePool> pools = await FindAllPoolsAsync();			
@@ -1129,7 +1129,7 @@ namespace Horde.Server.Devices
 			List<IDevice> reservedDevices = devices.Where(x => reservations.FirstOrDefault(r => r.Devices.Contains(x.Id)) != null).ToList();
 			List<IDevice> maintenanceDevices = devices.Where(x => x.MaintenanceTimeUtc != null).ToList();
 			List<IDevice> disabledDevices  = devices.Where(x => !x.Enabled).ToList();
-			List<IDevice> problemDevices = devices.Where(x => (x.ProblemTimeUtc != null && ((now - x.ProblemTimeUtc).Value.TotalMinutes < 30))).ToList();
+			List<IDevice> problemDevices = devices.Where(x => (x.ProblemTimeUtc != null && ((now - x.ProblemTimeUtc).Value.TotalMinutes < problemCooldown))).ToList();
 
 			Dictionary<DevicePoolId, List<DevicePlatformTelemetryDocument>> poolTelemetry = new Dictionary<DevicePoolId, List<DevicePlatformTelemetryDocument>>();
 
diff --git a/Engine/Source/Programs/Horde/Horde.Server/Devices/DeviceReportService.cs b/Engine/Source/Programs/Horde/Horde.Server/Devices/DeviceReportService.cs
new file mode 100644
--- /dev/null
+++ b/Engine/Source/Programs/Horde/Horde.Server/Devices/DeviceReportService.cs
@@ -0,0 +1,316 @@
+// Copyright Epic Games, Inc. All Rights Reserved.
+
+using System;
+using System.Collections.Generic;
+using System.Linq;
+using System.Text.RegularExpressions;
+using System.Threading;
+using System.Threading.Tasks;
+using EpicGames.Core;
+using EpicGames.Slack;
+using Horde.Server.Issues;
+using Horde.Server.Jobs;
+using Horde.Server.Jobs.Graphs;
+using Horde.Server.Notifications;
+using Horde.Server.Server;
+using Horde.Server.Streams;
+using Horde.Server.Utilities;
+using HordeCommon;
+using Microsoft.AspNetCore.DataProtection.KeyManagement;
+using Microsoft.CodeAnalysis.CSharp.Syntax;
+using Microsoft.Extensions.Hosting;
+using Microsoft.Extensions.Logging;
+using Microsoft.Extensions.Options;
+using MongoDB.Bson.Serialization.Attributes;
+using MongoDB.Bson.Serialization.Options;
+using Swashbuckle.AspNetCore.SwaggerGen;
+
+namespace Horde.Server.Devices
+{
+#pragma warning disable CS1591 // Missing XML comment for publicly visible type or member
+	public class DeviceReport
+	{
+		public string DeviceId { get; }
+		public string DeviceName { get; }
+		public string DeviceAddress { get; }
+
+		public string PlatformId { get; }
+		public string PlatformName { get; }
+
+		public string PoolId { get; }
+		public string PoolName { get; }
+
+		/// <summary>
+		/// Problems reported since the last report
+		/// </summary>
+		public int ProblemDelta { get; set; } = 0;
+
+		/// <summary>
+		/// The last problem Description
+		/// </summary>
+		public string? LastProblemDesc { get; set; }
+
+		/// <summary>
+		/// The last problem encountered
+		/// </summary>
+		public string? LastProblemURL { get; set; }
+
+
+		public List<IDeviceTelemetry> Telemetry { get; }
+
+		public DeviceReport(string platformId, string platformName, string deviceId, string deviceName, string deviceAddress, string poolId, string poolName, List<IDeviceTelemetry> telemetry)
+		{
+			PlatformId = platformId;
+			PlatformName = platformName;
+			DeviceId = deviceId;
+			DeviceName = deviceName;
+			DeviceAddress = deviceAddress;
+			PoolId = poolId;
+			PoolName = poolName;
+			Telemetry = telemetry;			
+		}
+	}
+
+	public class DevicePlatformReport
+	{
+		public string PlatformId { get; }
+		public string PlatformName { get; }
+		public List<DeviceReport> DeviceReports { get; set; } = new List<DeviceReport>();
+
+		public DevicePlatformReport(string platformId, string platformName)
+		{
+			PlatformId = platformId;
+			PlatformName = platformName;			
+		}
+	}
+
+	public class DevicePoolMetrics
+	{
+		public string PlatformId { get; }
+		public string PlatformName { get; }
+		public int Total { get; set; } = 0;
+		public int Disabled { get; set; } = 0;
+		public int Maintenance { get; set; } = 0;
+		public int Problems { get; set; } = 0;
+
+		public DevicePoolMetrics(string platformId, string platformName)
+		{
+			PlatformId = platformId;
+			PlatformName = platformName;
+		}
+	}
+
+	public class DevicePoolReport
+	{
+		public string PoolId { get; }
+		public string PoolName { get; }
+
+		public List<DevicePoolMetrics> Metrics { get; } = new List<DevicePoolMetrics>();
+
+		public DevicePoolReport(string poolId, string poolName)
+		{
+			PoolId = poolId;
+			PoolName = poolName;		
+		}
+	}
+
+	public class DeviceIssueReport
+	{
+		public string Channel { get; }
+
+		public List<DevicePlatformReport> PlatformReports { get; set;  } = new List<DevicePlatformReport>();
+
+		public List<DevicePoolReport> PoolReports { get; set; } = new List<DevicePoolReport>();
+
+		public DeviceIssueReport(string channel)
+		{
+			Channel = channel;
+		}
+	}
+
+#pragma warning restore CS1591 // Missing XML comment for publicly visible type or member
+
+	[SingletonDocument("device-report-state", "6491bc0ce39f8b70b1ef2feb")]
+	class DeviceReportState : SingletonBase
+	{
+		public DateTime ReportTime { get; set; } = DateTime.MinValue;
+	}
+
+	/// <summary>
+	/// Posts summaries for all the open issues in different streams to Slack channels
+	/// </summary>
+	public class DeviceReportService : IHostedService
+	{
+
+		readonly SingletonDocument<DeviceReportState> _state;
+		readonly DeviceService _deviceService;
+		readonly INotificationService _notificationService;
+		readonly IClock _clock;
+		readonly ITicker _ticker;
+		readonly IOptionsMonitor<GlobalConfig> _globalConfig;
+		readonly ILogger<DeviceReportService> _logger;
+		
+		readonly int _reportIntervalMinutes = 15;
+
+		/// <summary>
+		/// Constructor
+		/// </summary>
+		public DeviceReportService(MongoService mongoService, DeviceService deviceService, INotificationService notificationService, IClock clock, IOptionsMonitor<GlobalConfig> globalConfig, ILogger<DeviceReportService> logger)
+		{
+			_state = new SingletonDocument<DeviceReportState>(mongoService);
+			_deviceService = deviceService;
+			_notificationService = notificationService;
+			_clock = clock;
+			_ticker = clock.AddSharedTicker<DeviceReportService>(TimeSpan.FromMinutes(5.0), TickAsync, logger);
+			_globalConfig = globalConfig;
+			_logger = logger;
+		}
+
+		/// <inheritdoc/>
+		public async Task StartAsync(CancellationToken cancellationToken)
+		{
+			await _ticker.StartAsync();
+		}
+
+		/// <inheritdoc/>
+		public async Task StopAsync(CancellationToken cancellationToken)
+		{
+			await _ticker.StopAsync();
+		}
+
+		async ValueTask TickAsync(CancellationToken cancellationToken)
+		{
+			GlobalConfig globalConfig = _globalConfig.CurrentValue;
+
+			if (String.IsNullOrEmpty(globalConfig.ServerSettings.DeviceReportChannel))
+			{
+				return;
+			}
+
+			DeviceReportState state = await _state.GetAsync();
+			DateTime currentTime = _clock.UtcNow;
+			DateTime lastReportTime = state.ReportTime == DateTime.MinValue ? DateTime.Now : state.ReportTime;
+
+			if ((currentTime - lastReportTime).TotalMinutes <= _reportIntervalMinutes)
+			{
+				return;
+			}
+
+			_logger.LogInformation("Creating device report");
+
+			List<IDevicePool> pools = await _deviceService.GetPoolsAsync();
+			List<IDevicePlatform> platforms = await _deviceService.GetPlatformsAsync();
+			List<IDevice> devices = await _deviceService.GetDevicesAsync();
+
+			devices = devices.OrderBy(d => d.PlatformId.ToString()).ThenBy(d => d.PoolId.ToString()).ToList();
+
+			List<IDeviceTelemetry> deviceTelemetry = await _deviceService.GetDeviceTelemetryAsync(null, lastReportTime);
+			
+			DeviceIssueReport issueReport = new DeviceIssueReport(globalConfig.ServerSettings.DeviceReportChannel);
+
+			devices.ForEach(device =>
+			{
+				IDevicePool? pool = pools.Find(p => p.Id == device.PoolId);
+				IDevicePlatform? platform = platforms.Find(p => p.Id == device.PlatformId);
+
+				if (pool == null || platform == null)
+				{
+					return;
+				}
+
+				if (pool.PoolType != DevicePoolType.Automation)
+				{
+					return;
+				}
+
+				DevicePoolReport? poolReport = issueReport.PoolReports.Find(r => r.PoolId == pool.Id.ToString());
+				if (poolReport == null)
+				{
+					poolReport = new DevicePoolReport(pool.Id.ToString(), pool.Name.ToString());
+					issueReport.PoolReports.Add(poolReport);
+				}
+
+				DevicePoolMetrics? metrics = poolReport.Metrics.Find(m => m.PlatformId == device.PlatformId.ToString());
+				if (metrics == null) 
+				{
+					metrics = new DevicePoolMetrics(platform.Id.ToString(), platform.Name.ToString());
+					poolReport.Metrics.Add(metrics);
+				}
+
+				metrics.Total++;
+				if (!device.Enabled)
+				{
+					metrics.Disabled++;
+				}
+				else if (device.MaintenanceTimeUtc != null)
+				{
+					metrics.Maintenance++;
+				}
+
+				if (!device.Enabled || device.MaintenanceTimeUtc != null)
+				{
+					return;
+				}
+
+				// Report device issues, though only if not since disabled or put into maintenance
+
+				List<IDeviceTelemetry> telemetry = deviceTelemetry.Where(t => t.DeviceId == device.Id).ToList();
+				List<IDeviceTelemetry> problems = telemetry.Where(t => t.ProblemTimeUtc != null && t.ProblemTimeUtc! > lastReportTime).OrderByDescending(t => t.ProblemTimeUtc!).ToList();
+
+				if (problems.Count == 0)
+				{
+					return;
+				}
+
+				// mark that this device reported a problem
+				metrics.Problems += 1;
+
+				IDeviceTelemetry lastProblem = problems[0];
+
+				DeviceReport deviceReport = new DeviceReport(platform.Id.ToString(), platform.Name, device.Id.ToString(), device.Name, device.Address ?? "Unknown Address", pool.Id.ToString(), pool.Name, telemetry);
+
+				deviceReport.ProblemDelta = problems.Count;
+				
+				if (lastProblem.JobName != null && lastProblem.StepName != null && lastProblem.JobId != null && lastProblem.StepId != null) 
+				{
+					deviceReport.LastProblemDesc = $"{lastProblem.JobName} - {lastProblem.StepName}";
+					deviceReport.LastProblemURL = new Uri($"{globalConfig.ServerSettings.DashboardUrl}job/{lastProblem.JobId}?step={lastProblem.StepId}").ToString();
+				}
+
+				DevicePlatformReport? platformReport = issueReport.PlatformReports.Find(p => p.PlatformId == device.PlatformId.ToString());
+				if (platformReport == null)
+				{
+					platformReport = new DevicePlatformReport(platform.Id.ToString(), platform.Name);
+					issueReport.PlatformReports.Add(platformReport);
+				}
+
+				platformReport.DeviceReports.Add(deviceReport);
+			});
+
+			// order everything nicely
+			issueReport.PoolReports = issueReport.PoolReports.OrderBy(r => r.PoolName).ToList();
+			issueReport.PlatformReports = issueReport.PlatformReports.OrderBy(r => r.PlatformName).ToList();
+			issueReport.PlatformReports.ForEach(r => r.DeviceReports = r.DeviceReports.OrderBy(r => r.PoolName).ThenBy(r => r.DeviceName).ToList());
+
+			if (issueReport.PoolReports.Count > 0 || issueReport.PlatformReports.Count > 0) 
+			{
+				try
+				{
+					_logger.LogInformation("Sending device report notification");
+
+					await _notificationService.SendDeviceIssueReportAsync(issueReport);
+				}
+				catch (Exception ex)
+				{
+					_logger.LogError(ex, "Error sending device health notification");
+				}
+			}
+
+			state = await _state.UpdateAsync(s => s.ReportTime = currentTime);
+
+		}
+	}
+}
+
+
+// CL 22278596 - Has the device channel notification stuff
\ No newline at end of file
diff --git a/Engine/Source/Programs/Horde/Horde.Server/Devices/DeviceService.cs b/Engine/Source/Programs/Horde/Horde.Server/Devices/DeviceService.cs
--- a/Engine/Source/Programs/Horde/Horde.Server/Devices/DeviceService.cs
+++ b/Engine/Source/Programs/Horde/Horde.Server/Devices/DeviceService.cs
@@ -186,7 +186,7 @@ namespace Horde.Server.Devices
 			{
 				using TelemetrySpan span = _tracer.StartActiveSpan($"{nameof(DeviceService)}.{nameof(TickTelemetryAsync)}");
 				_logger.LogInformation("Updating pool telemetry");
-				await _devices.CreatePoolTelemetrySnapshot();
+				await _devices.CreatePoolTelemetrySnapshot(_settings.CurrentValue.DeviceProblemCooldownMinutes);
 			}
 		}
 
@@ -441,7 +441,7 @@ namespace Horde.Server.Devices
 				}
 			}
 
-			return await _devices.TryAddReservationAsync(pool, request, hostname, reservationDetails, streamId, jobId, stepId, jobName, stepName);
+			return await _devices.TryAddReservationAsync(pool, request, _settings.CurrentValue.DeviceProblemCooldownMinutes, hostname, reservationDetails, streamId, jobId, stepId, jobName, stepName);
 		}
 
 		/// <summary>
diff --git a/Engine/Source/Programs/Horde/Horde.Server/Devices/IDeviceCollection.cs b/Engine/Source/Programs/Horde/Horde.Server/Devices/IDeviceCollection.cs
--- a/Engine/Source/Programs/Horde/Horde.Server/Devices/IDeviceCollection.cs
+++ b/Engine/Source/Programs/Horde/Horde.Server/Devices/IDeviceCollection.cs
@@ -176,6 +176,7 @@ namespace Horde.Server.Devices
 		/// </summary>
 		/// <param name="poolId">The pool of devices to use for the new reservation</param>
 		/// <param name="request">The requested devices for the reservation</param>
+		/// <param name="problemCooldown">The configured problem device cooldown in minutes</param>
 		/// <param name="hostname">The hostname of the machine making the reservation</param>
 		/// <param name="reservationDetails">The details of the reservation</param>
 		/// <param name="streamId">The Stream Id associated with the job</param>
@@ -183,7 +184,7 @@ namespace Horde.Server.Devices
 		/// <param name="stepId">The Step Id associated with the job</param>
 		/// <param name="jobName">The Job name associated with the job</param>
 		/// <param name="stepName">The Step name associated with the job</param>		
-		Task<IDeviceReservation?> TryAddReservationAsync(DevicePoolId poolId, List<DeviceRequestData> request, string? hostname, string? reservationDetails, string? streamId, string? jobId, string? stepId, string? jobName, string? stepName);
+		Task<IDeviceReservation?> TryAddReservationAsync(DevicePoolId poolId, List<DeviceRequestData> request, int problemCooldown, string? hostname, string? reservationDetails, string? streamId, string? jobId, string? stepId, string? jobName, string? stepName);
 
 		/// <summary>
 		/// Gets a reservation by guid for legacy clients
@@ -246,7 +247,7 @@ namespace Horde.Server.Devices
 		/// Creates a device pool telemetry snapshot
 		/// </summary>
 		/// <returns></returns>
-		public Task CreatePoolTelemetrySnapshot();
+		public Task CreatePoolTelemetrySnapshot(int poolCooldown);
 
 		/// <summary>
 		/// Gets pool telemetry for an optional date range
diff --git a/Engine/Source/Programs/Horde/Horde.Server/Notifications/INotificationService.cs b/Engine/Source/Programs/Horde/Horde.Server/Notifications/INotificationService.cs
--- a/Engine/Source/Programs/Horde/Horde.Server/Notifications/INotificationService.cs
+++ b/Engine/Source/Programs/Horde/Horde.Server/Notifications/INotificationService.cs
@@ -150,6 +150,12 @@ namespace Horde.Server.Notifications
 		/// <param name="user"></param>
 		void NotifyDeviceService(string message, IDevice? device = null, IDevicePool? pool = null, StreamConfig? streamConfig = null, IJob? job = null, IJobStep? step = null, INode? node = null, IUser? user = null);
 
+		/// <summary>
+		/// Post a notification for device issues
+		/// </summary>
+		/// <param name="report">The report data to send</param>
+		Task SendDeviceIssueReportAsync(DeviceIssueReport report);
+
 		/// <summary>
 		/// Post a notification for the open issues in a stream
 		/// </summary>
diff --git a/Engine/Source/Programs/Horde/Horde.Server/Notifications/INotificationSink.cs b/Engine/Source/Programs/Horde/Horde.Server/Notifications/INotificationSink.cs
--- a/Engine/Source/Programs/Horde/Horde.Server/Notifications/INotificationSink.cs
+++ b/Engine/Source/Programs/Horde/Horde.Server/Notifications/INotificationSink.cs
@@ -118,6 +118,14 @@ namespace Horde.Server.Notifications
 		/// <param name="user"></param>
 		/// <returns></returns>
 		Task NotifyDeviceServiceAsync(string message, IDevice? device = null, IDevicePool? pool = null, StreamConfig? streamConfig = null, IJob? job = null, IJobStep? step = null, INode? node = null, IUser? user = null);
-    }
+
+		/// <summary>
+		/// Notify the status of devices
+		/// </summary>
+		/// <param name="report"></param>
+		/// <returns></returns>
+		Task SendDeviceIssueReportAsync(DeviceIssueReport report);
+
+	}
 }
 
diff --git a/Engine/Source/Programs/Horde/Horde.Server/Notifications/NotificationService.cs b/Engine/Source/Programs/Horde/Horde.Server/Notifications/NotificationService.cs
--- a/Engine/Source/Programs/Horde/Horde.Server/Notifications/NotificationService.cs
+++ b/Engine/Source/Programs/Horde/Horde.Server/Notifications/NotificationService.cs
@@ -298,6 +298,23 @@ namespace Horde.Server.Notifications
 			EnqueueTasks(sink => sink.NotifyDeviceServiceAsync(message, device, pool, streamConfig, job, step, node, user));
 		}
 
+		/// <inheritdoc/>
+		public async Task SendDeviceIssueReportAsync(DeviceIssueReport report)
+		{
+			foreach (INotificationSink sink in _sinks)
+			{
+				try
+				{
+					await sink.SendDeviceIssueReportAsync(report);
+				}
+				catch (Exception e)
+				{
+					_logger.LogError(e, "Failed sending issue report to {Channel}", report.Channel);
+				}
+			}
+		}
+
+
 		/// <summary>
 		/// Enqueues an async task
 		/// </summary>
diff --git a/Engine/Source/Programs/Horde/Horde.Server/Notifications/Sinks/SlackNotificationSink.cs b/Engine/Source/Programs/Horde/Horde.Server/Notifications/Sinks/SlackNotificationSink.cs
--- a/Engine/Source/Programs/Horde/Horde.Server/Notifications/Sinks/SlackNotificationSink.cs
+++ b/Engine/Source/Programs/Horde/Horde.Server/Notifications/Sinks/SlackNotificationSink.cs
@@ -14,14 +14,12 @@ using System.Text.Json.Serialization;
 using System.Text.RegularExpressions;
 using System.Threading;
 using System.Threading.Tasks;
-using System.Timers;
 using EpicGames.Core;
 using EpicGames.Redis;
 using EpicGames.Redis.Utility;
 using EpicGames.Slack;
 using EpicGames.Slack.Blocks;
 using EpicGames.Slack.Elements;
-using Horde.Server.Agents;
 using Horde.Server.Configuration;
 using Horde.Server.Devices;
 using Horde.Server.Issues;
@@ -35,6 +33,7 @@ using Horde.Server.Users;
 using Horde.Server.Utilities;
 using HordeCommon;
 using Microsoft.AspNetCore.Hosting;
+using Microsoft.CodeAnalysis;
 using Microsoft.Extensions.Caching.Memory;
 using Microsoft.Extensions.Hosting;
 using Microsoft.Extensions.Logging;
@@ -2119,6 +2118,110 @@ namespace Horde.Server.Notifications.Sinks
 
 		#region Device notifications
 
+		/// <inheritdoc/>
+		public async Task SendDeviceIssueReportAsync(DeviceIssueReport report)
+		{			
+			if (report.PoolReports.Count > 0)
+			{
+				SlackMessage headerMessage = new SlackMessage();
+				headerMessage.AddHeader($"Device Pool Health Summary");
+				await SendMessageAsync(report.Channel, headerMessage);
+
+				foreach (DevicePoolReport pool in report.PoolReports)
+				{
+					List<SlackAttachment> attachments = new List<SlackAttachment>();
+					foreach (DevicePoolMetrics metrics in pool.Metrics)
+					{
+						
+						double totalPct = (((double)metrics.Disabled + (double)metrics.Maintenance + (double)metrics.Problems) / (double)metrics.Total) * 100.0;
+
+						StringBuilder builder = new StringBuilder("   ");
+
+						string color = "#eeeeee";
+
+						if (totalPct >= 40)
+						{
+							color = "#ff0000";
+						}
+						else if (totalPct >= 20)
+						{
+							color = "#ffff00";
+						}
+						else
+						{
+							continue;
+						}
+
+						builder.Append($"*{metrics.PlatformName}* - ");
+
+						if (metrics.Problems > 0)
+						{
+							builder.Append($"Problems: { metrics.Problems}, ");
+						}
+
+						builder.Append($"Devices: {metrics.Total}, Disabled: {metrics.Disabled}, Maintenance: {metrics.Maintenance}");
+
+						SlackAttachment attachment = new SlackAttachment();
+						attachment.FallbackText = builder.ToString();
+						attachment.Color = color;
+						attachment.AddSection(builder.ToString());
+						attachments.Add(attachment);
+					}
+					if (attachments.Count > 0)
+					{
+						SlackMessage message = new SlackMessage();
+						message.Attachments.AddRange(attachments);
+						message.Markdown = true;
+						message.Text = $"*{pool.PoolName}*";
+						await SendMessageAsync(report.Channel, message);
+					}					
+				}
+
+				SlackMessage divider = new SlackMessage();
+				divider.AddDivider();
+				await SendMessageAsync(report.Channel, divider);
+
+			}
+
+			if (report.PlatformReports.Count > 0)
+			{
+				SlackMessage headerMessage = new SlackMessage();
+				headerMessage.AddHeader($"Device Problems Since Last Update");
+				await SendMessageAsync(report.Channel, headerMessage);
+
+				foreach (DevicePlatformReport platform in report.PlatformReports)
+				{
+					SlackMessage message = new SlackMessage();
+					message.Markdown = true;
+					message.Text = $"*{platform.PlatformName}*";					
+
+					SlackAttachment attachment = new SlackAttachment();
+
+					foreach (DeviceReport device in platform.DeviceReports.OrderByDescending(r => r.ProblemDelta))
+					{
+						StringBuilder builder = new StringBuilder($"   {device.DeviceName} / {device.DeviceAddress}");
+
+						builder.Append($" / Problems: {device.ProblemDelta}");
+						
+						if (!String.IsNullOrEmpty(device.LastProblemURL))
+						{						
+							builder.Append($" - <{device.LastProblemURL}|{device.LastProblemDesc ?? "??"}>");
+						}
+
+						attachment.FallbackText = builder.ToString();
+						attachment.AddSection(builder.ToString());
+					}
+
+					message.Attachments.Add(attachment);
+					await SendMessageAsync(report.Channel, message);
+				}
+
+				SlackMessage divider = new SlackMessage();
+				divider.AddDivider();
+				await SendMessageAsync(report.Channel, divider);
+			}
+		}
+
 		/// <inheritdoc/>
 		public async Task NotifyDeviceServiceAsync(string message, IDevice? device = null, IDevicePool? pool = null, StreamConfig? streamConfig = null, IJob? job = null, IJobStep? step = null, INode? node = null, IUser? user = null)
 		{
diff --git a/Engine/Source/Programs/Horde/Horde.Server/ServerSettings.cs b/Engine/Source/Programs/Horde/Horde.Server/ServerSettings.cs
--- a/Engine/Source/Programs/Horde/Horde.Server/ServerSettings.cs
+++ b/Engine/Source/Programs/Horde/Horde.Server/ServerSettings.cs
@@ -736,6 +736,17 @@ namespace Horde.Server
 		/// </summary>
 		public int SharedDeviceCheckoutDays { get; set; } = 3;
 
+		/// <summary>
+		/// The number of cooldown minutes for device problems
+		/// </summary>
+		public int DeviceProblemCooldownMinutes { get; set; } = 10;
+
+		/// <summary>
+		/// Channel to send device reports to
+		/// </summary>
+		public string? DeviceReportChannel { get; set; }
+
+
 		/// <summary>
 		/// Default agent pool sizing strategy for pools that doesn't have one explicitly configured
 		/// </summary>
diff --git a/Engine/Source/Programs/Horde/Horde.Server/Startup.cs b/Engine/Source/Programs/Horde/Horde.Server/Startup.cs
--- a/Engine/Source/Programs/Horde/Horde.Server/Startup.cs
+++ b/Engine/Source/Programs/Horde/Horde.Server/Startup.cs
@@ -701,6 +701,7 @@ namespace Horde.Server
 				services.AddHostedService<PoolUpdateService>();
 				services.AddHostedService<TelemetryService>();
 				services.AddHostedService(provider => provider.GetRequiredService<DeviceService>());
+				services.AddHostedService<DeviceReportService>();
 				services.AddHostedService(provider => provider.GetRequiredService<TestDataService>());
 				
 				if (settings.Commits.ReplicateMetadata)
