Horde: Additional resolution on preflight redirector
#rnx
#skipci

diff --git a/Engine/Source/Programs/Horde/HordeDashboard/src/components/Preflight.tsx b/Engine/Source/Programs/Horde/HordeDashboard/src/components/Preflight.tsx
--- a/Engine/Source/Programs/Horde/HordeDashboard/src/components/Preflight.tsx
+++ b/Engine/Source/Programs/Horde/HordeDashboard/src/components/Preflight.tsx
@@ -1,7 +1,7 @@
 // Copyright Epic Games, Inc. All Rights Reserved.
 
 import React, { useState } from 'react';
-import { useNavigate, Navigate } from 'react-router-dom';
+import { useNavigate } from 'react-router-dom';
 import backend, { useBackend } from '../backend';
 import { useQuery } from './JobDetailCommon';
 import ErrorHandler from './ErrorHandler';
@@ -56,6 +56,15 @@ export const PreflightRedirector: React.FC = () => {
 
    let stream = projectStore.streamByFullname(streamName);
 
+   if (!stream) {      
+
+      stream = projectStore.streamById(streamName?.replace("//", "").replaceAll("/", "-").toLowerCase());
+
+      if (!stream) {
+         stream = projectStore.streamByFullname(streamName + "-VS");   
+      }      
+      
+   }
 
    if (!stream) {
       setError(`Unable to resolve stream with name ${streamName}`);
