Horde: Change device service report interval
#rnx

diff --git a/Engine/Source/Programs/Horde/Horde.Server/Devices/DeviceReportService.cs b/Engine/Source/Programs/Horde/Horde.Server/Devices/DeviceReportService.cs
--- a/Engine/Source/Programs/Horde/Horde.Server/Devices/DeviceReportService.cs
+++ b/Engine/Source/Programs/Horde/Horde.Server/Devices/DeviceReportService.cs
@@ -150,7 +150,7 @@ namespace Horde.Server.Devices
 		readonly IOptionsMonitor<GlobalConfig> _globalConfig;
 		readonly ILogger<DeviceReportService> _logger;
 		
-		readonly int _reportIntervalMinutes = 15;
+		readonly int _reportIntervalMinutes = 180;
 
 		/// <summary>
 		/// Constructor
