Horde: Make device problem cooldown configurable
#rnx
#skipci

diff --git a/Engine/Source/Programs/Horde/HordeDashboard/src/backend/Api.ts b/Engine/Source/Programs/Horde/HordeDashboard/src/backend/Api.ts
--- a/Engine/Source/Programs/Horde/HordeDashboard/src/backend/Api.ts
+++ b/Engine/Source/Programs/Horde/HordeDashboard/src/backend/Api.ts
@@ -393,6 +393,10 @@ export type GetDashboardConfigResponse = {
 
 	/** Help slack channel that users can use for issues */
 	helpSlackChannel?: string;
+
+	/** Device problem cooldown in minutes */ 
+	deviceProblemCooldownMinutes?: number
+
 }
 
 /**Parameters to register a new agent */
diff --git a/Engine/Source/Programs/Horde/HordeDashboard/src/backend/Dashboard.ts b/Engine/Source/Programs/Horde/HordeDashboard/src/backend/Dashboard.ts
--- a/Engine/Source/Programs/Horde/HordeDashboard/src/backend/Dashboard.ts
+++ b/Engine/Source/Programs/Horde/HordeDashboard/src/backend/Dashboard.ts
@@ -158,6 +158,10 @@ export class Dashboard {
         return this.config?.helpSlackChannel;
     }
 
+    get deviceProblemCooldownMinutes(): number {
+        return this.config?.deviceProblemCooldownMinutes ?? 30;
+    }
+
     get swarmUrl(): string | undefined {
         return this.config?.perforceSwarmUrl;
     }
diff --git a/Engine/Source/Programs/Horde/HordeDashboard/src/components/DeviceEditor.tsx b/Engine/Source/Programs/Horde/HordeDashboard/src/components/DeviceEditor.tsx
--- a/Engine/Source/Programs/Horde/HordeDashboard/src/components/DeviceEditor.tsx
+++ b/Engine/Source/Programs/Horde/HordeDashboard/src/components/DeviceEditor.tsx
@@ -133,7 +133,7 @@ export class DeviceHandler extends PollBase {
             const d = moment.duration(end.diff(moment(device.problemTime)));
 
             // note this must match reservation selection for problem time in backend
-            if (d.asMinutes() < 30) {
+            if (d.asMinutes() < dashboard.deviceProblemCooldownMinutes) {
                 return DeviceStatus.Problem;
             }
         }
