Horde: Add a more specific error when a step is aborted due to an unhealthy server connection.

diff --git a/Engine/Source/Programs/Horde/Horde.Agent/Leases/Handlers/JobHandler.cs b/Engine/Source/Programs/Horde/Horde.Agent/Leases/Handlers/JobHandler.cs
--- a/Engine/Source/Programs/Horde/Horde.Agent/Leases/Handlers/JobHandler.cs
+++ b/Engine/Source/Programs/Horde/Horde.Agent/Leases/Handlers/JobHandler.cs
@@ -165,7 +165,14 @@ namespace Horde.Agent.Leases.Handlers
 			{
 				if (cancellationToken.IsCancellationRequested && ex.IsCancellationException())
 				{
-					logger.LogError("Step was aborted");
+					if (session.RpcConnection.Healthy)
+					{
+						logger.LogError("Step was aborted");
+					}
+					else
+					{
+						logger.LogError("Connection to the server was lost; step aborted.");
+					}
 					throw;
 				}
 				else
@@ -337,7 +344,14 @@ namespace Horde.Agent.Leases.Handlers
 			{
 				if (cancellationToken.IsCancellationRequested && ex.IsCancellationException())
 				{
-					batchLogger.LogError("Step was aborted");
+					if (session.RpcConnection.Healthy)
+					{
+						batchLogger.LogError("Lease was aborted");
+					}
+					else
+					{
+						batchLogger.LogError("Connection to the server was lost; lease aborted.");
+					}
 				}
 				else
 				{
