Horde: Add a more explicit error when cancelling leases due to a timeout.

diff --git a/Engine/Source/Programs/Horde/Horde.Agent/Leases/Handlers/ComputeHandler.cs b/Engine/Source/Programs/Horde/Horde.Agent/Leases/Handlers/ComputeHandler.cs
--- a/Engine/Source/Programs/Horde/Horde.Agent/Leases/Handlers/ComputeHandler.cs
+++ b/Engine/Source/Programs/Horde/Horde.Agent/Leases/Handlers/ComputeHandler.cs
@@ -55,6 +55,8 @@ namespace Horde.Agent.Leases.Handlers
 			}
 		}
 
+		const int NoDataTimeoutMinutes = 10;
+
 		readonly ComputeListenerService _listenerService;
 		readonly IMemoryCache _memoryCache;
 		readonly ILogger _logger;
@@ -92,23 +94,31 @@ namespace Horde.Agent.Leases.Handlers
 				using (CancellationTokenSource cts = CancellationTokenSource.CreateLinkedTokenSource(cancellationToken))
 				{
 					await using BackgroundTask timeoutTask = BackgroundTask.StartNew(ctx => TickTimeoutAsync(transport, cts, ctx));
-					await using (ComputeSocket socket = new ComputeSocket(transport, ComputeSocketEndpoint.Local, _logger))
+					try
 					{
-						DirectoryReference sandboxDir = DirectoryReference.Combine(session.WorkingDir, "Sandbox", leaseId);
-						try
-						{
-							DirectoryReference.CreateDirectory(sandboxDir);
-
-							AgentMessageHandler worker = new AgentMessageHandler(sandboxDir, _memoryCache, _logger);
-							await worker.RunAsync(socket, cts.Token);
-							await socket.CloseAsync(cts.Token);
-							return LeaseResult.Success;
-						}
-						finally
+						await using (ComputeSocket socket = new ComputeSocket(transport, ComputeSocketEndpoint.Local, _logger))
 						{
-							FileUtils.ForceDeleteDirectory(sandboxDir);
+							DirectoryReference sandboxDir = DirectoryReference.Combine(session.WorkingDir, "Sandbox", leaseId);
+							try
+							{
+								DirectoryReference.CreateDirectory(sandboxDir);
+
+								AgentMessageHandler worker = new AgentMessageHandler(sandboxDir, _memoryCache, _logger);
+								await worker.RunAsync(socket, cts.Token);
+								await socket.CloseAsync(cts.Token);
+								return LeaseResult.Success;
+							}
+							finally
+							{
+								FileUtils.ForceDeleteDirectory(sandboxDir);
+							}
 						}
 					}
+					catch (OperationCanceledException ex) when (cts.IsCancellationRequested && transport.TimeSinceActivity > TimeSpan.FromMinutes(NoDataTimeoutMinutes))
+					{
+						_logger.LogError(ex, "Lease was terminated due to no data being received for {Time} minutes", NoDataTimeoutMinutes);
+						return LeaseResult.Failed;
+					}
 				}
 			}
 			catch (Exception ex)
@@ -126,7 +136,7 @@ namespace Horde.Agent.Leases.Handlers
 		{
 			while(!cancellationToken.IsCancellationRequested)
 			{
-				if (transport.TimeSinceActivity > TimeSpan.FromMinutes(10))
+				if (transport.TimeSinceActivity > TimeSpan.FromMinutes(NoDataTimeoutMinutes))
 				{
 					_logger.LogWarning("Terminating compute task due to timeout (last tick at {Time})", DateTime.UtcNow - transport.TimeSinceActivity);
 					cts.Cancel();
