Horde: Fix conversion of high freq counter ticks to system ticks. This ratio is ~0.4 on my PC, causing compute timeouts to fire at incorrect intervals.

diff --git a/Engine/Source/Programs/Horde/Horde.Agent/Leases/Handlers/ComputeHandler.cs b/Engine/Source/Programs/Horde/Horde.Agent/Leases/Handlers/ComputeHandler.cs
--- a/Engine/Source/Programs/Horde/Horde.Agent/Leases/Handlers/ComputeHandler.cs
+++ b/Engine/Source/Programs/Horde/Horde.Agent/Leases/Handlers/ComputeHandler.cs
@@ -26,6 +26,8 @@ namespace Horde.Agent.Leases.Handlers
 			readonly TcpTransport _inner;
 			long _lastPingTicks;
 
+			static readonly double s_ticksToSystemTicks = (double)TimeSpan.TicksPerSecond / Stopwatch.Frequency;
+
 			public long Position => _inner.Position;
 
 			public TcpTransportWithTimeout(Socket socket)
@@ -34,7 +36,7 @@ namespace Horde.Agent.Leases.Handlers
 				_lastPingTicks = Stopwatch.GetTimestamp();
 			}
 
-			public TimeSpan TimeSinceActivity => TimeSpan.FromTicks(Stopwatch.GetTimestamp() - Interlocked.CompareExchange(ref _lastPingTicks, 0, 0));
+			public TimeSpan TimeSinceActivity => TimeSpan.FromTicks((long)((Stopwatch.GetTimestamp() - Interlocked.CompareExchange(ref _lastPingTicks, 0, 0)) * s_ticksToSystemTicks));
 
 			public ValueTask MarkCompleteAsync(CancellationToken cancellationToken) => _inner.MarkCompleteAsync(cancellationToken);
 
