Horde: Remove methods for reading and writing bundles from IStorageClient. This is now an implementation detail specific to types derived from BundleStorageClient.

diff --git a/Engine/Source/Programs/Horde/Horde.Agent/Execution/JobExecutor.cs b/Engine/Source/Programs/Horde/Horde.Agent/Execution/JobExecutor.cs
--- a/Engine/Source/Programs/Horde/Horde.Agent/Execution/JobExecutor.cs
+++ b/Engine/Source/Programs/Horde/Horde.Agent/Execution/JobExecutor.cs
@@ -792,7 +792,6 @@ namespace Horde.Agent.Execution
 
 			using MemoryCache cache = new MemoryCache(new MemoryCacheOptions { });
 			IStorageClient storage = _storageFactory.CreateStorageClient(_session, _namespaceId, _token);
-			BundleReader reader = new BundleReader(storage, cache, _logger);
 
 			// Create the mapping of tag names to file sets
 			Dictionary<string, HashSet<FileReference>> tagNameToFileSet = new Dictionary<string, HashSet<FileReference>>();
diff --git a/Engine/Source/Programs/Horde/Horde.Server.Tests/BlobStoreTests.cs b/Engine/Source/Programs/Horde/Horde.Server.Tests/BlobStoreTests.cs
--- a/Engine/Source/Programs/Horde/Horde.Server.Tests/BlobStoreTests.cs
+++ b/Engine/Source/Programs/Horde/Horde.Server.Tests/BlobStoreTests.cs
@@ -19,7 +19,7 @@ namespace Horde.Server.Tests
 	[TestClass]
 	public class BlobStoreTests : TestSetup
 	{
-		async Task<IStorageClientImpl> CreateStorageClientAsync()
+		async Task<StorageClient> CreateStorageClientAsync()
 		{
 			GlobalConfig globalConfig = new GlobalConfig();
 			globalConfig.Storage.Backends.Add(new BackendConfig { Id = new BackendId("default-backend"), Type = StorageBackendType.Memory });
@@ -65,7 +65,7 @@ namespace Horde.Server.Tests
 			return new Bundle(header, new[] { data });
 		}
 
-		static async Task<Blob> ReadBlobAsync(IStorageClient store, BlobLocator locator)
+		static async Task<Blob> ReadBlobAsync(StorageClient store, BlobLocator locator)
 		{
 			Bundle bundle = await store.ReadBundleAsync(locator);
 			return ExtractBlobFromBundle(bundle);
@@ -82,7 +82,7 @@ namespace Horde.Server.Tests
 		[TestMethod]
 		public async Task LeafTest()
 		{
-			IStorageClient store = await CreateStorageClientAsync();
+			StorageClient store = await CreateStorageClientAsync();
 
 			byte[] input = CreateTestData(256, 0);
 			BlobLocator locator = await store.WriteBundleAsync(CreateTestBundle(input, Array.Empty<BlobLocator>()));
@@ -94,7 +94,7 @@ namespace Horde.Server.Tests
 		[TestMethod]
 		public async Task ReferenceTest()
 		{
-			IStorageClientImpl store = await CreateStorageClientAsync();
+			StorageClient store = await CreateStorageClientAsync();
 
 			byte[] input1 = CreateTestData(256, 1);
 			Bundle bundle1 = CreateTestBundle(input1, Array.Empty<BlobLocator>());
@@ -129,7 +129,7 @@ namespace Horde.Server.Tests
 		[TestMethod]
 		public async Task RefExpiryTest()
 		{
-			IStorageClientImpl store = await CreateStorageClientAsync();
+			StorageClient store = await CreateStorageClientAsync();
 
 			Bundle bundle1 = CreateTestBundle(new byte[] { 1, 2, 3 }, Array.Empty<BlobLocator>());
 			BlobLocator locator1 = await store.WriteBundleAsync(bundle1);
diff --git a/Engine/Source/Programs/Horde/Horde.Server.Tests/GcServiceTests.cs b/Engine/Source/Programs/Horde/Horde.Server.Tests/GcServiceTests.cs
--- a/Engine/Source/Programs/Horde/Horde.Server.Tests/GcServiceTests.cs
+++ b/Engine/Source/Programs/Horde/Horde.Server.Tests/GcServiceTests.cs
@@ -31,7 +31,7 @@ namespace Horde.Server.Tests
 			await StorageService.StartAsync(CancellationToken.None);
 
 			SetupNamespace();
-			IStorageClientImpl store = await StorageService.GetClientAsync(new NamespaceId("default"), CancellationToken.None);
+			StorageClient store = await StorageService.GetClientAsync(new NamespaceId("default"), CancellationToken.None);
 
 			Random random = new Random(0);
 			BlobLocator[] blobs = await CreateTestDataAsync(store, 30, 50, 30, 5, random);
@@ -55,14 +55,14 @@ namespace Horde.Server.Tests
 			Assert.IsTrue(remaining.All(x => nodes.Contains(x)));
 		}
 
-		async Task<HashSet<BlobLocator>> FindNodes(IStorageClient store, IEnumerable<BlobLocator> roots)
+		async Task<HashSet<BlobLocator>> FindNodes(BundleStorageClient store, IEnumerable<BlobLocator> roots)
 		{
 			HashSet<BlobLocator> nodes = new HashSet<BlobLocator>();
 			await FindNodes(store, roots, nodes);
 			return nodes;
 		}
 
-		async Task FindNodes(IStorageClient store, IEnumerable<BlobLocator> roots, HashSet<BlobLocator> nodes)
+		async Task FindNodes(BundleStorageClient store, IEnumerable<BlobLocator> roots, HashSet<BlobLocator> nodes)
 		{
 			foreach (BlobLocator root in roots)
 			{
@@ -74,7 +74,7 @@ namespace Horde.Server.Tests
 			}
 		}
 
-		static async ValueTask<BlobLocator[]> CreateTestDataAsync(IStorageClient store, int numRoots, int numInterior, int numLeaves, int avgChildren, Random random)
+		static async ValueTask<BlobLocator[]> CreateTestDataAsync(BundleStorageClient store, int numRoots, int numInterior, int numLeaves, int avgChildren, Random random)
 		{
 			int firstRoot = 0;
 			int firstInterior = firstRoot + numRoots;
diff --git a/Engine/Source/Programs/Horde/Horde.Server.Tests/StorageServiceTests.cs b/Engine/Source/Programs/Horde/Horde.Server.Tests/StorageServiceTests.cs
--- a/Engine/Source/Programs/Horde/Horde.Server.Tests/StorageServiceTests.cs
+++ b/Engine/Source/Programs/Horde/Horde.Server.Tests/StorageServiceTests.cs
@@ -22,7 +22,7 @@ namespace Horde.Server.Tests
 		public async Task BlobCollectionTest()
 		{
 			StorageService storageService = ServiceProvider.GetRequiredService<StorageService>();
-			IStorageClientImpl client = await storageService.GetClientAsync(new NamespaceId("memory"), default);
+			StorageClient client = await storageService.GetClientAsync(new NamespaceId("memory"), default);
 
 			List<BlobType> types = new List<BlobType>();
 			types.Add(new BlobType(Guid.Parse("{11C2D886-3349-4164-946F-E9D10BD12E3D}"), 0));
diff --git a/Engine/Source/Programs/Horde/Horde.Server/Commands/Bundles/BundleWriteTest.cs b/Engine/Source/Programs/Horde/Horde.Server/Commands/Bundles/BundleWriteTest.cs
--- a/Engine/Source/Programs/Horde/Horde.Server/Commands/Bundles/BundleWriteTest.cs
+++ b/Engine/Source/Programs/Horde/Horde.Server/Commands/Bundles/BundleWriteTest.cs
@@ -20,11 +20,8 @@ namespace Horde.Server.Commands.Bundles
 	{
 		class FakeStorageClient : IStorageClient
 		{
-			readonly BundleReader _reader;
-
 			public FakeStorageClient()
 			{
-				_reader = new BundleReader(this, null, NullLogger.Instance);
 			}
 
 			public Task DeleteRefAsync(RefName name, CancellationToken cancellationToken = default) => Task.CompletedTask;
@@ -37,7 +34,7 @@ namespace Horde.Server.Commands.Bundles
 			public Task<BlobLocator> WriteBundleAsync(Bundle bundle, Utf8String prefix = default, CancellationToken cancellationToken = default) => Task.FromResult(BlobLocator.CreateUnique(Utf8String.Empty));
 			public Task WriteRefTargetAsync(RefName name, BlobHandle target, RefOptions? options = null, CancellationToken cancellationToken = default) => Task.CompletedTask;
 
-			public BundleWriter CreateWriter(RefName refName = default, BundleOptions? options = null) => new BundleWriter(this, _reader, refName, options);
+			public BundleWriter CreateWriter(RefName refName = default, BundleOptions? options = null) => throw new NotImplementedException();// new BundleWriter(this, _reader, refName, options);
 
 			IStorageWriter IStorageClient.CreateWriter(EpicGames.Horde.Storage.RefName refName) => CreateWriter(refName);
 		}
diff --git a/Engine/Source/Programs/Horde/Horde.Server/Logs/LogRpcService.cs b/Engine/Source/Programs/Horde/Horde.Server/Logs/LogRpcService.cs
--- a/Engine/Source/Programs/Horde/Horde.Server/Logs/LogRpcService.cs
+++ b/Engine/Source/Programs/Horde/Horde.Server/Logs/LogRpcService.cs
@@ -50,7 +50,7 @@ namespace Horde.Server.Logs
 				throw new StructuredRpcException(StatusCode.PermissionDenied, "Access denied");
 			}
 
-			IStorageClientImpl store = await _storageService.GetClientAsync(Namespace.Logs, context.CancellationToken);
+			StorageClient store = await _storageService.GetClientAsync(Namespace.Logs, context.CancellationToken);
 			_logger.LogInformation("Updating {LogId} to node {RefTarget} (lines: {LineCount}, complete: {Complete})", request.LogId, request.Target, request.LineCount, request.Complete);
 			await store.WriteRefTargetAsync(new RefName(request.LogId), BundleNodeLocator.Parse(request.Target));
 
diff --git a/Engine/Source/Programs/Horde/Horde.Server/Logs/LogsController.cs b/Engine/Source/Programs/Horde/Horde.Server/Logs/LogsController.cs
--- a/Engine/Source/Programs/Horde/Horde.Server/Logs/LogsController.cs
+++ b/Engine/Source/Programs/Horde/Horde.Server/Logs/LogsController.cs
@@ -117,7 +117,7 @@ namespace Horde.Server.Logs
 				return Forbid();
 			}
 
-			IStorageClient storageClient = await _storageService.GetClientAsync(Namespace.Logs, cancellationToken);
+			StorageClient storageClient = await _storageService.GetClientAsync(Namespace.Logs, cancellationToken);
 			return await StorageController.WriteBlobAsync(storageClient, file, $"{logFile.RefName}", cancellationToken);
 		}
 
diff --git a/Engine/Source/Programs/Horde/Horde.Server/Perforce/ReplicationService.cs b/Engine/Source/Programs/Horde/Horde.Server/Perforce/ReplicationService.cs
--- a/Engine/Source/Programs/Horde/Horde.Server/Perforce/ReplicationService.cs
+++ b/Engine/Source/Programs/Horde/Horde.Server/Perforce/ReplicationService.cs
@@ -138,7 +138,7 @@ namespace Horde.Server.Perforce
 		{
 			RefName refName = new RefName(streamConfig.Id.ToString());
 
-			IStorageClientImpl store = await _storageService.GetClientAsync(Namespace.Perforce, cancellationToken);
+			StorageClient store = await _storageService.GetClientAsync(Namespace.Perforce, cancellationToken);
 
 			CommitNode? lastCommitNode = await store.TryReadNodeAsync<CommitNode>(refName, cancellationToken: cancellationToken);
 			ICommitCollection commits = _perforceService.GetCommits(streamConfig);
diff --git a/Engine/Source/Programs/Horde/Horde.Server/Storage/StorageController.cs b/Engine/Source/Programs/Horde/Horde.Server/Storage/StorageController.cs
--- a/Engine/Source/Programs/Horde/Horde.Server/Storage/StorageController.cs
+++ b/Engine/Source/Programs/Horde/Horde.Server/Storage/StorageController.cs
@@ -199,7 +199,7 @@ namespace Horde.Server.Storage
 				return Forbid(StorageAclAction.WriteBlobs, namespaceId);
 			}
 
-			IStorageClientImpl storageClient = await _storageService.GetClientAsync(namespaceId, cancellationToken);
+			StorageClient storageClient = await _storageService.GetClientAsync(namespaceId, cancellationToken);
 			return await WriteBlobAsync(storageClient, file, prefix, cancellationToken);
 		}
 
@@ -211,9 +211,9 @@ namespace Horde.Server.Storage
 		/// <param name="prefix">Prefix for uploaded blobs</param>
 		/// <param name="cancellationToken">Cancellation token</param>
 		/// <returns>Information about the written blob, or redirect information</returns>
-		public static async Task<ActionResult<WriteBlobResponse>> WriteBlobAsync(IStorageClient storageClient, IFormFile? file, [FromForm] string? prefix = default, CancellationToken cancellationToken = default)
+		public static async Task<ActionResult<WriteBlobResponse>> WriteBlobAsync(BundleStorageClient storageClient, IFormFile? file, [FromForm] string? prefix = default, CancellationToken cancellationToken = default)
 		{
-			IStorageClientImpl? storageClientImpl = storageClient as IStorageClientImpl;
+			StorageClient? storageClientImpl = storageClient as StorageClient;
 			if (file == null)
 			{
 				if (storageClientImpl == null)
@@ -272,16 +272,16 @@ namespace Horde.Server.Storage
 		/// </summary>
 		internal static async Task<ActionResult> ReadBlobInternalAsync(StorageService storageService, NamespaceId namespaceId, BlobLocator locator, int? offset, int? length, CancellationToken cancellationToken)
 		{
-			IStorageClientImpl client = await storageService.GetClientAsync(namespaceId, cancellationToken);
+			StorageClient client = await storageService.GetClientAsync(namespaceId, cancellationToken);
 			return await ReadBlobInternalAsync(client, locator, offset, length, cancellationToken);
 		}
 
 		/// <summary>
 		/// Reads a blob from storage, without performing namespace access checks.
 		/// </summary>
-		internal static async Task<ActionResult> ReadBlobInternalAsync(IStorageClient storageClient, BlobLocator locator, int? offset, int? length, CancellationToken cancellationToken)
+		internal static async Task<ActionResult> ReadBlobInternalAsync(BundleStorageClient storageClient, BlobLocator locator, int? offset, int? length, CancellationToken cancellationToken)
 		{
-			if (storageClient is IStorageClientImpl storageClientImpl)
+			if (storageClient is StorageClient storageClientImpl)
 			{
 				Uri? redirectUrl = await storageClientImpl.GetReadRedirectAsync(locator, cancellationToken);
 				if (redirectUrl != null)
@@ -331,7 +331,7 @@ namespace Horde.Server.Storage
 				return Forbid(StorageAclAction.ReadBlobs, namespaceId);
 			}
 
-			IStorageClientImpl client = await _storageService.GetClientAsync(namespaceId, cancellationToken);
+			StorageClient client = await _storageService.GetClientAsync(namespaceId, cancellationToken);
 
 			FindNodesResponse response = new FindNodesResponse();
 			await foreach (BlobHandle handle in client.FindNodesAsync(alias, cancellationToken))
@@ -368,7 +368,7 @@ namespace Horde.Server.Storage
 				return Forbid(StorageAclAction.WriteRefs, namespaceId);
 			}
 
-			IStorageClientImpl client = await _storageService.GetClientAsync(namespaceId, cancellationToken);
+			StorageClient client = await _storageService.GetClientAsync(namespaceId, cancellationToken);
 			BundleNodeLocator target = new BundleNodeLocator(request.Hash, request.Blob, request.ExportIdx);
 			await client.WriteRefTargetAsync(refName, target, request.Options, cancellationToken);
 
@@ -488,7 +488,7 @@ namespace Horde.Server.Storage
 				return Forbid(StorageAclAction.ReadBlobs, namespaceId);
 			}
 
-			IStorageClient storageClient = await _storageService.GetClientAsync(namespaceId, cancellationToken);
+			StorageClient storageClient = await _storageService.GetClientAsync(namespaceId, cancellationToken);
 			BundleReader reader = new BundleReader(storageClient, _memoryCache, _logger);
 
 			BundleHeader header = await reader.ReadBundleHeaderAsync(locator, cancellationToken);
@@ -568,7 +568,7 @@ namespace Horde.Server.Storage
 				return Forbid(StorageAclAction.ReadBlobs, namespaceId);
 			}
 
-			IStorageClient storageClient = await _storageService.GetClientAsync(namespaceId, cancellationToken);
+			StorageClient storageClient = await _storageService.GetClientAsync(namespaceId, cancellationToken);
 			BundleReader reader = new BundleReader(storageClient, _memoryCache, _logger);
 
 			BundleHeader header = await reader.ReadBundleHeaderAsync(locator, cancellationToken);
diff --git a/Engine/Source/Programs/Horde/Horde.Server/Storage/StorageService.cs b/Engine/Source/Programs/Horde/Horde.Server/Storage/StorageService.cs
--- a/Engine/Source/Programs/Horde/Horde.Server/Storage/StorageService.cs
+++ b/Engine/Source/Programs/Horde/Horde.Server/Storage/StorageService.cs
@@ -46,31 +46,51 @@ namespace Horde.Server.Storage
 	/// <summary>
 	/// Interface for storage clients which includes a backend implementation. Some functionality is exposed through the backend which is not part of the regular storage API (eg. enumerating).
 	/// </summary>
-	public interface IStorageClientImpl : IStorageClient
+	public abstract class StorageClient : BundleStorageClient
 	{
 		/// <summary>
 		/// Configuration for this namespace
 		/// </summary>
-		NamespaceConfig Config { get; }
+		public NamespaceConfig Config { get; }
 
 		/// <summary>
 		/// The storage backend
 		/// </summary>
-		IStorageBackend Backend { get; }
+		public IStorageBackend Backend { get; }
+
+		/// <summary>
+		/// Accessor for the namespace id
+		/// </summary>
+		public NamespaceId NamespaceId => Config.Id;
 
 		/// <summary>
 		/// Whether the backend supports redirects
 		/// </summary>
-		bool SupportsRedirects { get; }
+		public bool SupportsRedirects { get; }
+
+		/// <summary>
+		/// Constructor
+		/// </summary>
+		/// <param name="config">Namespace configuration</param>
+		/// <param name="backend">Backend store</param>
+		/// <param name="memoryCache">Memory cache</param>
+		/// <param name="logger">Logger instance</param>
+		protected StorageClient(NamespaceConfig config, IStorageBackend backend, IMemoryCache? memoryCache, ILogger logger)
+			: base(memoryCache, logger)
+		{
+			Config = config;
+			Backend = backend;
+			SupportsRedirects = backend.SupportsRedirects && !config.EnableAliases;
+		}
 
 		/// <inheritdoc cref="IStorageClient.AddAliasAsync(Utf8String, BlobHandle, CancellationToken)"/>
-		Task AddAliasAsync(Utf8String name, BundleNodeLocator target, CancellationToken cancellationToken = default);
+		public abstract Task AddAliasAsync(Utf8String name, BundleNodeLocator target, CancellationToken cancellationToken = default);
 
 		/// <inheritdoc cref="IStorageClient.RemoveAliasAsync(Utf8String, BlobHandle, CancellationToken)"/>
-		Task RemoveAliasAsync(Utf8String name, BundleNodeLocator target, CancellationToken cancellationToken = default); 
+		public abstract Task RemoveAliasAsync(Utf8String name, BundleNodeLocator target, CancellationToken cancellationToken = default);
 
 		/// <inheritdoc cref="IStorageClient.WriteRefTargetAsync(RefName, BlobHandle, RefOptions?, CancellationToken)"/>
-		Task WriteRefTargetAsync(RefName name, BundleNodeLocator handle, RefOptions? options = null, CancellationToken cancellationToken = default);
+		public abstract Task WriteRefTargetAsync(RefName name, BundleNodeLocator handle, RefOptions? options = null, CancellationToken cancellationToken = default);
 
 		/// <summary>
 		/// Gets a redirect for a read request
@@ -78,7 +98,7 @@ namespace Horde.Server.Storage
 		/// <param name="locator">Locator for the blob</param>
 		/// <param name="cancellationToken">Cancellation token for the operation</param>
 		/// <returns>Path to upload the data to</returns>
-		ValueTask<Uri?> GetReadRedirectAsync(BlobLocator locator, CancellationToken cancellationToken = default);
+		public abstract ValueTask<Uri?> GetReadRedirectAsync(BlobLocator locator, CancellationToken cancellationToken = default);
 
 		/// <summary>
 		/// Gets a redirect for a write request
@@ -86,7 +106,7 @@ namespace Horde.Server.Storage
 		/// <param name="prefix">Prefix for the new blob locator</param>
 		/// <param name="cancellationToken">Cancellation token for the operation</param>
 		/// <returns>Locator and path to upload the data to</returns>
-		ValueTask<(BlobLocator, Uri)?> GetWriteRedirectAsync(Utf8String prefix = default, CancellationToken cancellationToken = default);
+		public abstract ValueTask<(BlobLocator, Uri)?> GetWriteRedirectAsync(Utf8String prefix = default, CancellationToken cancellationToken = default);
 	}
 
 	/// <summary>
@@ -94,26 +114,15 @@ namespace Horde.Server.Storage
 	/// </summary>
 	public sealed class StorageService : IHostedService, IDisposable, IStorageClientFactory
 	{
-		sealed class StorageClient : BundleStorageClient, IStorageClientImpl
+		sealed class StorageClientImpl : StorageClient
 		{
 			readonly StorageService _outer;
 			readonly string _prefix;
 
-			public NamespaceConfig Config { get; }
-
-			public IStorageBackend Backend { get; }
-
-			NamespaceId NamespaceId => Config.Id;
-
-			public bool SupportsRedirects { get; }
-
-			public StorageClient(StorageService outer, NamespaceConfig config, IStorageBackend backend, IMemoryCache? memoryCache, ILogger logger)
-				: base(memoryCache, logger)
+			public StorageClientImpl(StorageService outer, NamespaceConfig config, IStorageBackend backend, IMemoryCache? memoryCache, ILogger logger)
+				: base(config, backend, memoryCache, logger)
 			{
 				_outer = outer;
-				Config = config;
-				Backend = backend;
-				SupportsRedirects = backend.SupportsRedirects && !config.EnableAliases;
 
 				_prefix = config.Prefix;
 				if (_prefix.Length > 0 && !_prefix.EndsWith("/", StringComparison.Ordinal))
@@ -141,7 +150,7 @@ namespace Horde.Server.Storage
 			}
 
 			/// <inheritdoc/>
-			public ValueTask<Uri?> GetReadRedirectAsync(BlobLocator locator, CancellationToken cancellationToken = default) => Backend.TryGetReadRedirectAsync(GetBlobPath(locator), cancellationToken);
+			public override ValueTask<Uri?> GetReadRedirectAsync(BlobLocator locator, CancellationToken cancellationToken = default) => Backend.TryGetReadRedirectAsync(GetBlobPath(locator), cancellationToken);
 
 			/// <inheritdoc/>
 			public override async Task<ReadOnlyMemory<byte>> ReadBundleRangeAsync(BlobLocator locator, int offset, int length, CancellationToken cancellationToken = default)
@@ -172,7 +181,7 @@ namespace Horde.Server.Storage
 			}
 
 			/// <inheritdoc/>
-			public async ValueTask<(BlobLocator, Uri)?> GetWriteRedirectAsync(Utf8String prefix = default, CancellationToken cancellationToken = default)
+			public override async ValueTask<(BlobLocator, Uri)?> GetWriteRedirectAsync(Utf8String prefix = default, CancellationToken cancellationToken = default)
 			{
 				if (!Backend.SupportsRedirects)
 				{
@@ -205,13 +214,13 @@ namespace Horde.Server.Storage
 			public override Task AddAliasAsync(Utf8String name, BlobHandle handle, CancellationToken cancellationToken = default) => _outer.AddAliasAsync(NamespaceId, name, handle.GetLocator(), cancellationToken);
 
 			/// <inheritdoc/>
-			public Task AddAliasAsync(Utf8String name, BundleNodeLocator locator, CancellationToken cancellationToken = default) => _outer.AddAliasAsync(NamespaceId, name, locator, cancellationToken);
+			public override Task AddAliasAsync(Utf8String name, BundleNodeLocator locator, CancellationToken cancellationToken = default) => _outer.AddAliasAsync(NamespaceId, name, locator, cancellationToken);
 
 			/// <inheritdoc/>
 			public override Task RemoveAliasAsync(Utf8String name, BlobHandle handle, CancellationToken cancellationToken = default) => _outer.RemoveAliasAsync(NamespaceId, name, handle.GetLocator(), cancellationToken);
 
 			/// <inheritdoc/>
-			public Task RemoveAliasAsync(Utf8String name, BundleNodeLocator locator, CancellationToken cancellationToken = default) => _outer.RemoveAliasAsync(NamespaceId, name, locator, cancellationToken);
+			public override Task RemoveAliasAsync(Utf8String name, BundleNodeLocator locator, CancellationToken cancellationToken = default) => _outer.RemoveAliasAsync(NamespaceId, name, locator, cancellationToken);
 
 			/// <inheritdoc/>
 			public override async IAsyncEnumerable<BlobHandle> FindNodesAsync(Utf8String alias, [EnumeratorCancellation] CancellationToken cancellationToken = default)
@@ -245,7 +254,7 @@ namespace Horde.Server.Storage
 			}
 
 			/// <inheritdoc/>
-			public Task WriteRefTargetAsync(RefName name, BundleNodeLocator target, RefOptions? options = null, CancellationToken cancellationToken = default) => _outer.WriteRefTargetAsync(NamespaceId, name, target, options, cancellationToken);
+			public override Task WriteRefTargetAsync(RefName name, BundleNodeLocator target, RefOptions? options = null, CancellationToken cancellationToken = default) => _outer.WriteRefTargetAsync(NamespaceId, name, target, options, cancellationToken);
 
 			/// <inheritdoc/>
 			public override Task DeleteRefAsync(RefName name, CancellationToken cancellationToken = default) => _outer.DeleteRefAsync(NamespaceId, name, cancellationToken);
@@ -256,10 +265,10 @@ namespace Horde.Server.Storage
 		class NamespaceInfo : IDisposable
 		{
 			public NamespaceConfig Config { get; }
-			public StorageClient Client { get; }
+			public StorageClientImpl Client { get; }
 			public IStorageBackend Backend { get; }
 
-			public NamespaceInfo(NamespaceConfig config, StorageClient client, IStorageBackend backend)
+			public NamespaceInfo(NamespaceConfig config, StorageClientImpl client, IStorageBackend backend)
 			{
 				Config = config;
 				Client = client;
@@ -536,7 +545,7 @@ namespace Horde.Server.Storage
 		/// <param name="namespaceId">Namespace identifier</param>
 		/// <param name="cancellationToken">Cancellation token for the operation</param>
 		/// <returns></returns>
-		public async ValueTask<IStorageClientImpl> GetClientAsync(NamespaceId namespaceId, CancellationToken cancellationToken)
+		public async ValueTask<StorageClient> GetClientAsync(NamespaceId namespaceId, CancellationToken cancellationToken)
 		{
 			NamespaceInfo namespaceInfo = await GetNamespaceInfoAsync(namespaceId, cancellationToken);
 			return namespaceInfo.Client;
@@ -561,7 +570,7 @@ namespace Horde.Server.Storage
 					foreach (NamespaceConfig namespaceConfig in storageConfig.Namespaces)
 					{
 						IStorageBackend backend = _storageBackendProvider.CreateBackend(namespaceConfig.BackendConfig);
-						StorageClient client = new StorageClient(this, namespaceConfig, backend, _cache, _logger);
+						StorageClientImpl client = new StorageClientImpl(this, namespaceConfig, backend, _cache, _logger);
 						nextState.Namespaces.Add(namespaceConfig.Id, new NamespaceInfo(namespaceConfig, client, backend));
 					}
 				}
@@ -624,7 +633,7 @@ namespace Horde.Server.Storage
 			DateTime utcNow = _clock.UtcNow;
 
 			// Cached storage clients for each namespace
-			Dictionary<NamespaceId, IStorageClient?> namespaceIdToClient = new Dictionary<NamespaceId, IStorageClient?>();
+			Dictionary<NamespaceId, StorageClient?> namespaceIdToClient = new Dictionary<NamespaceId, StorageClient?>();
 
 			// Compute missing import info, by searching for blobs with an ObjectId timestamp after the last import compute cycle
 			ObjectId latestInfoId = ObjectId.GenerateNewId(utcNow - TimeSpan.FromMinutes(30.0));
@@ -635,7 +644,7 @@ namespace Horde.Server.Storage
 					// Find imports, and add a check record for each new blob
 					foreach (BlobInfo blobInfo in cursor.Current)
 					{
-						IStorageClient? client;
+						StorageClient? client;
 						if (!namespaceIdToClient.TryGetValue(blobInfo.NamespaceId, out client))
 						{
 							NamespaceInfo? namespaceInfo = await TryGetNamespaceInfoAsync(blobInfo.NamespaceId, cancellationToken);
@@ -1014,7 +1023,7 @@ namespace Horde.Server.Storage
 			await _gcState.UpdateAsync(state => state.FindOrAddNamespace(namespaceId).LastTime = utcNow);
 		}
 
-		async Task<BundleHeader?> ReadHeaderAsync(IStorageClient store, BlobLocator locator, CancellationToken cancellationToken)
+		async Task<BundleHeader?> ReadHeaderAsync(BundleStorageClient store, BlobLocator locator, CancellationToken cancellationToken)
 		{
 			int fetchSize = 64 * 1024;
 			for (; ; )
diff --git a/Engine/Source/Programs/Horde/Horde.Server/Tools/IToolCollection.cs b/Engine/Source/Programs/Horde/Horde.Server/Tools/IToolCollection.cs
--- a/Engine/Source/Programs/Horde/Horde.Server/Tools/IToolCollection.cs
+++ b/Engine/Source/Programs/Horde/Horde.Server/Tools/IToolCollection.cs
@@ -59,7 +59,7 @@ namespace Horde.Server.Tools
 		/// <param name="tool">The tool to get a storage client for</param>
 		/// <param name="cancellationToken">Cancellation token for the operation</param>
 		/// <returns>Instance of the storage client</returns>
-		Task<IStorageClient> GetStorageClientAsync(ITool tool, CancellationToken cancellationToken);
+		Task<BundleStorageClient> GetStorageClientAsync(ITool tool, CancellationToken cancellationToken);
 
 		/// <summary>
 		/// Opens a stream to the data for a particular deployment
diff --git a/Engine/Source/Programs/Horde/Horde.Server/Tools/ToolCollection.cs b/Engine/Source/Programs/Horde/Horde.Server/Tools/ToolCollection.cs
--- a/Engine/Source/Programs/Horde/Horde.Server/Tools/ToolCollection.cs
+++ b/Engine/Source/Programs/Horde/Horde.Server/Tools/ToolCollection.cs
@@ -280,7 +280,7 @@ namespace Horde.Server.Tools
 			ToolDeploymentId deploymentId = ToolDeploymentId.GenerateNewId();
 			RefName refName = new RefName($"{tool.Id}/{deploymentId}");
 
-			IStorageClientImpl client = await _storageService.GetClientAsync(Namespace.Tools, cancellationToken);
+			StorageClient client = await _storageService.GetClientAsync(Namespace.Tools, cancellationToken);
 			await client.WriteRefTargetAsync(refName, locator, cancellationToken: cancellationToken);
 
 			return await CreateDeploymentInternalAsync(tool, deploymentId, options, refName, globalConfig, cancellationToken);
@@ -413,7 +413,7 @@ namespace Horde.Server.Tools
 		/// <param name="tool">Identifier for the tool</param>
 		/// <param name="cancellationToken">Cancellation token for the operation</param>
 		/// <returns>Storage client for the data</returns>
-		public async Task<IStorageClient> GetStorageClientAsync(ITool tool, CancellationToken cancellationToken)
+		public async Task<BundleStorageClient> GetStorageClientAsync(ITool tool, CancellationToken cancellationToken)
 		{
 			if (tool.Config is BundledToolConfig bundledConfig)
 			{
diff --git a/Engine/Source/Programs/Horde/Horde.Server/Tools/ToolsController.cs b/Engine/Source/Programs/Horde/Horde.Server/Tools/ToolsController.cs
--- a/Engine/Source/Programs/Horde/Horde.Server/Tools/ToolsController.cs
+++ b/Engine/Source/Programs/Horde/Horde.Server/Tools/ToolsController.cs
@@ -249,7 +249,7 @@ namespace Horde.Server.Tools
 				return Forbid(ToolAclAction.UploadTool, id);
 			}
 
-			IStorageClient storageClient = await _toolCollection.GetStorageClientAsync(tool, cancellationToken);
+			BundleStorageClient storageClient = await _toolCollection.GetStorageClientAsync(tool, cancellationToken);
 			return await StorageController.WriteBlobAsync(storageClient, file, cancellationToken: cancellationToken);
 		}
 
@@ -564,7 +564,7 @@ namespace Horde.Server.Tools
 				return BadRequest("Invalid blob id for tool");
 			}
 
-			IStorageClient storageClient = await _toolCollection.GetStorageClientAsync(tool, cancellationToken);
+			BundleStorageClient storageClient = await _toolCollection.GetStorageClientAsync(tool, cancellationToken);
 			return StorageController.ReadBlobInternalAsync(storageClient, locator, offset, length, cancellationToken);
 		}
 
diff --git a/Engine/Source/Programs/Horde/Horde/Commands/Bundles/BundleExtract.cs b/Engine/Source/Programs/Horde/Horde/Commands/Bundles/BundleExtract.cs
--- a/Engine/Source/Programs/Horde/Horde/Commands/Bundles/BundleExtract.cs
+++ b/Engine/Source/Programs/Horde/Horde/Commands/Bundles/BundleExtract.cs
@@ -42,7 +42,6 @@ namespace Horde.Commands.Bundles
 			}
 
 			using MemoryCache cache = new MemoryCache(new MemoryCacheOptions());
-			BundleReader reader = new BundleReader(store, cache, logger);
 
 			Stopwatch timer = Stopwatch.StartNew();
 
diff --git a/Engine/Source/Programs/Horde/Horde/Commands/Vcs/VcsLog.cs b/Engine/Source/Programs/Horde/Horde/Commands/Vcs/VcsLog.cs
--- a/Engine/Source/Programs/Horde/Horde/Commands/Vcs/VcsLog.cs
+++ b/Engine/Source/Programs/Horde/Horde/Commands/Vcs/VcsLog.cs
@@ -30,7 +30,6 @@ namespace Horde.Commands.Vcs
 			IStorageClient store = await GetStorageClientAsync();
 
 			using MemoryCache cache = new MemoryCache(new MemoryCacheOptions());
-			BundleReader reader = new BundleReader(store, cache, logger);
 
 			List<CommitNode> commits = new List<CommitNode>();
 
diff --git a/Engine/Source/Programs/Shared/EpicGames.Horde/Compute/AgentMessage.cs b/Engine/Source/Programs/Shared/EpicGames.Horde/Compute/AgentMessage.cs
--- a/Engine/Source/Programs/Shared/EpicGames.Horde/Compute/AgentMessage.cs
+++ b/Engine/Source/Programs/Shared/EpicGames.Horde/Compute/AgentMessage.cs
@@ -322,7 +322,7 @@ namespace EpicGames.Horde.Compute
 
 		#region Process
 
-		static async Task<AgentMessage> RunStorageServer(this AgentMessageChannel channel, IStorageClient storage, CancellationToken cancellationToken = default)
+		static async Task<AgentMessage> RunStorageServer(this AgentMessageChannel channel, BundleStorageClient storage, CancellationToken cancellationToken = default)
 		{
 			for (; ; )
 			{
@@ -352,7 +352,7 @@ namespace EpicGames.Horde.Compute
 		/// <param name="locator">Location of a <see cref="DirectoryNode"/> describing contents of the sandbox</param>
 		/// <param name="storage">Storage for the sandbox data</param>
 		/// <param name="cancellationToken">Cancellation token for the operation</param>
-		public static async Task UploadFilesAsync(this AgentMessageChannel channel, string path, BundleNodeLocator locator, IStorageClient storage, CancellationToken cancellationToken = default)
+		public static async Task UploadFilesAsync(this AgentMessageChannel channel, string path, BundleNodeLocator locator, BundleStorageClient storage, CancellationToken cancellationToken = default)
 		{
 			using (IAgentMessageBuilder request = await channel.CreateMessageAsync(AgentMessageType.WriteFiles, cancellationToken))
 			{
@@ -567,7 +567,7 @@ namespace EpicGames.Horde.Compute
 		/// <param name="message">The read request</param>
 		/// <param name="storage">Storage client to retrieve the blob from</param>
 		/// <param name="cancellationToken">Cancellation token for the operation</param>
-		public static Task SendBlobDataAsync(this AgentMessageChannel channel, ReadBlobMessage message, IStorageClient storage, CancellationToken cancellationToken = default)
+		public static Task SendBlobDataAsync(this AgentMessageChannel channel, ReadBlobMessage message, BundleStorageClient storage, CancellationToken cancellationToken = default)
 		{
 			return SendBlobDataAsync(channel, message.Locator, message.Offset, message.Length, storage, cancellationToken);
 		}
@@ -581,7 +581,7 @@ namespace EpicGames.Horde.Compute
 		/// <param name="length">Length of the data</param>
 		/// <param name="storage">Storage client to retrieve the blob from</param>
 		/// <param name="cancellationToken">Cancellation token for the operation</param>
-		public static async Task SendBlobDataAsync(this AgentMessageChannel channel, BlobLocator locator, int offset, int length, IStorageClient storage, CancellationToken cancellationToken = default)
+		public static async Task SendBlobDataAsync(this AgentMessageChannel channel, BlobLocator locator, int offset, int length, BundleStorageClient storage, CancellationToken cancellationToken = default)
 		{
 			ReadOnlySequence<byte> data;
 			if (offset == 0 && length == 0)
diff --git a/Engine/Source/Programs/Shared/EpicGames.Horde/Storage/BundleReader.cs b/Engine/Source/Programs/Shared/EpicGames.Horde/Storage/BundleReader.cs
--- a/Engine/Source/Programs/Shared/EpicGames.Horde/Storage/BundleReader.cs
+++ b/Engine/Source/Programs/Shared/EpicGames.Horde/Storage/BundleReader.cs
@@ -76,7 +76,7 @@ namespace EpicGames.Horde.Storage
 		// When reader is uncached, use a smaller default fetch size
 		const int DefaultUncachedFetchSize = 1 * 1024 * 1024;
 
-		readonly IStorageClient _store;
+		readonly BundleStorageClient _store;
 		readonly IMemoryCache? _cache;
 		readonly ILogger _logger;
 
@@ -92,7 +92,7 @@ namespace EpicGames.Horde.Storage
 		/// <param name="store"></param>
 		/// <param name="cache">Cache for data</param>
 		/// <param name="logger">Logger for output</param>
-		public BundleReader(IStorageClient store, IMemoryCache? cache, ILogger logger)
+		public BundleReader(BundleStorageClient store, IMemoryCache? cache, ILogger logger)
 		{
 			_store = store;
 			_cache = cache;
diff --git a/Engine/Source/Programs/Shared/EpicGames.Horde/Storage/BundleStorageClient.cs b/Engine/Source/Programs/Shared/EpicGames.Horde/Storage/BundleStorageClient.cs
--- a/Engine/Source/Programs/Shared/EpicGames.Horde/Storage/BundleStorageClient.cs
+++ b/Engine/Source/Programs/Shared/EpicGames.Horde/Storage/BundleStorageClient.cs
@@ -37,6 +37,21 @@ namespace EpicGames.Horde.Storage
 		/// <inheritdoc/>
 		public abstract Task<ReadOnlyMemory<byte>> ReadBundleRangeAsync(BlobLocator locator, int offset, int length, CancellationToken cancellationToken = default);
 
+		/// <summary>
+		/// Utility method to read a blob into a buffer
+		/// </summary>
+		/// <param name="locator">Blob location</param>
+		/// <param name="offset">Offset within the blob</param>
+		/// <param name="memory">Buffer to read into</param>
+		/// <param name="cancellationToken">Cancellation token for the operation</param>
+		/// <returns>The data that was read</returns>
+		public async Task<Memory<byte>> ReadBundleRangeAsync(BlobLocator locator, int offset, Memory<byte> memory, CancellationToken cancellationToken = default)
+		{
+			ReadOnlyMemory<byte> buffer = await ReadBundleRangeAsync(locator, offset, memory.Length, cancellationToken);
+			buffer.CopyTo(memory);
+			return memory.Slice(0, buffer.Length);
+		}
+
 		/// <inheritdoc/>
 		public abstract Task<BlobLocator> WriteBundleAsync(Bundle bundle, Utf8String prefix = default, CancellationToken cancellationToken = default);
 
diff --git a/Engine/Source/Programs/Shared/EpicGames.Horde/Storage/BundleWriter.cs b/Engine/Source/Programs/Shared/EpicGames.Horde/Storage/BundleWriter.cs
--- a/Engine/Source/Programs/Shared/EpicGames.Horde/Storage/BundleWriter.cs
+++ b/Engine/Source/Programs/Shared/EpicGames.Horde/Storage/BundleWriter.cs
@@ -504,7 +504,7 @@ namespace EpicGames.Horde.Storage
 			}
 
 			// Mark the bundle as complete
-			public void MarkAsComplete(IStorageClient store, Utf8String prefix, ILogger? traceLogger)
+			public void MarkAsComplete(BundleStorageClient store, Utf8String prefix, ILogger? traceLogger)
 			{
 				if (!IsReadOnly)
 				{
@@ -516,7 +516,7 @@ namespace EpicGames.Horde.Storage
 				}
 			}
 
-			async Task CompleteAsync(Task prevWriteTask, IStorageClient store, Utf8String prefix, ILogger? traceLogger)
+			async Task CompleteAsync(Task prevWriteTask, BundleStorageClient store, Utf8String prefix, ILogger? traceLogger)
 			{
 				try
 				{
@@ -652,7 +652,7 @@ namespace EpicGames.Horde.Storage
 
 		static readonly BundleOptions s_defaultOptions = new BundleOptions();
 
-		readonly IStorageClient _store;
+		readonly BundleStorageClient _store;
 		readonly BundleReader _reader;
 		readonly BundleOptions _options;
 		readonly RefName _refName;
@@ -681,7 +681,7 @@ namespace EpicGames.Horde.Storage
 		/// <param name="options">Options for the writer</param>
 		/// <param name="nodeCache">Cache of nodes for deduplication</param>
 		/// <param name="traceLogger">Optional logger for trace information</param>
-		public BundleWriter(IStorageClient store, BundleReader reader, RefName refName, BundleOptions? options = null, NodeCache? nodeCache = null, ILogger? traceLogger = null)
+		public BundleWriter(BundleStorageClient store, BundleReader reader, RefName refName, BundleOptions? options = null, NodeCache? nodeCache = null, ILogger? traceLogger = null)
 		{
 			_store = store;
 			_reader = reader;
diff --git a/Engine/Source/Programs/Shared/EpicGames.Horde/Storage/IStorageClient.cs b/Engine/Source/Programs/Shared/EpicGames.Horde/Storage/IStorageClient.cs
--- a/Engine/Source/Programs/Shared/EpicGames.Horde/Storage/IStorageClient.cs
+++ b/Engine/Source/Programs/Shared/EpicGames.Horde/Storage/IStorageClient.cs
@@ -73,36 +73,6 @@ namespace EpicGames.Horde.Storage
 	/// </summary>
 	public interface IStorageClient
 	{
-		#region Bundles
-
-		/// <summary>
-		/// Reads raw data for a blob from the store
-		/// </summary>
-		/// <param name="locator">The blob locator</param>
-		/// <param name="cancellationToken">Cancellation token for the operation</param>
-		/// <returns>Stream containing the data</returns>
-		Task<Bundle> ReadBundleAsync(BlobLocator locator, CancellationToken cancellationToken = default);
-
-		/// <summary>
-		/// Reads a ranged chunk from a blob
-		/// </summary>
-		/// <param name="locator">Locator for the blob</param>
-		/// <param name="offset">Starting offset for the data to read</param>
-		/// <param name="length">Length of the data</param>
-		/// <param name="cancellationToken">Cancellation token for the operation</param>
-		Task<ReadOnlyMemory<byte>> ReadBundleRangeAsync(BlobLocator locator, int offset, int length, CancellationToken cancellationToken = default);
-
-		/// <summary>
-		/// Writes a new blob to the store
-		/// </summary>
-		/// <param name="bundle">The bundle to write</param>
-		/// <param name="cancellationToken">Cancellation token for the operation</param>
-		/// <param name="prefix">Prefix for blob names. While the returned BlobId is guaranteed to be unique, this name can be used as a prefix to aid debugging.</param>
-		/// <returns>Unique identifier for the blob</returns>
-		Task<BlobLocator> WriteBundleAsync(Bundle bundle, Utf8String prefix = default, CancellationToken cancellationToken = default);
-
-		#endregion
-
 		#region Nodes
 
 		/// <summary>
@@ -276,26 +246,6 @@ namespace EpicGames.Horde.Storage
 	/// </summary>
 	public static class StorageClientExtensions
 	{
-		#region Blobs
-
-		/// <summary>
-		/// Utility method to read a blob into a buffer
-		/// </summary>
-		/// <param name="store">Store to read from</param>
-		/// <param name="locator">Blob location</param>
-		/// <param name="offset">Offset within the blob</param>
-		/// <param name="memory">Buffer to read into</param>
-		/// <param name="cancellationToken">Cancellation token for the operation</param>
-		/// <returns>The data that was read</returns>
-		public static async Task<Memory<byte>> ReadBundleRangeAsync(this IStorageClient store, BlobLocator locator, int offset, Memory<byte> memory, CancellationToken cancellationToken = default)
-		{
-			ReadOnlyMemory<byte> buffer = await store.ReadBundleRangeAsync(locator, offset, memory.Length, cancellationToken);
-			buffer.CopyTo(memory);
-			return memory.Slice(0, buffer.Length);
-		}
-
-		#endregion
-
 		#region Refs
 
 		/// <summary>
diff --git a/Engine/Source/Programs/UnrealCloudDDC/Jupiter/Controllers/BundlesController.cs b/Engine/Source/Programs/UnrealCloudDDC/Jupiter/Controllers/BundlesController.cs
--- a/Engine/Source/Programs/UnrealCloudDDC/Jupiter/Controllers/BundlesController.cs
+++ b/Engine/Source/Programs/UnrealCloudDDC/Jupiter/Controllers/BundlesController.cs
@@ -203,7 +203,7 @@ namespace Jupiter.Controllers
                 return result;
             }
 
-            IStorageClientJupiter client = await _storageService.GetClientAsync(namespaceId, cancellationToken);
+            StorageClient client = await _storageService.GetClientAsync(namespaceId, cancellationToken);
             // disable redirected responses for upload so that we can parse the blobs for references
             // TODO: Add way to parse objects after upload similar to BlobTick in Horde
             /*if (file == null)
@@ -251,7 +251,7 @@ namespace Jupiter.Controllers
                 return result;
             }
 
-            IStorageClientJupiter client = await _storageService.GetClientAsync(namespaceId, cancellationToken);
+            StorageClient client = await _storageService.GetClientAsync(namespaceId, cancellationToken);
 
             Uri? redirectUrl = await client.GetReadRedirectAsync(locator, cancellationToken);
             if (redirectUrl != null)
@@ -328,7 +328,7 @@ namespace Jupiter.Controllers
             {
                 return result;
             }
-            IStorageClientJupiter client = await _storageService.GetClientAsync(namespaceId, cancellationToken);
+            StorageClient client = await _storageService.GetClientAsync(namespaceId, cancellationToken);
             BundleNodeLocator target = new BundleNodeLocator(request.Hash, request.Blob, request.ExportIdx);
             await client.WriteRefTargetAsync(refName, target, request.Options, cancellationToken);
 
@@ -383,7 +383,7 @@ namespace Jupiter.Controllers
                 return result;
             }
 
-            IStorageClient storageClient = await _storageService.GetClientAsync(namespaceId, cancellationToken);
+            StorageClient storageClient = await _storageService.GetClientAsync(namespaceId, cancellationToken);
 
             BundleReader reader = new BundleReader(storageClient, _memoryCache, _logger);
 
@@ -459,7 +459,7 @@ namespace Jupiter.Controllers
                 return result;
             }
 
-            IStorageClient storageClient = await _storageService.GetClientAsync(namespaceId, cancellationToken);
+            StorageClient storageClient = await _storageService.GetClientAsync(namespaceId, cancellationToken);
             BundleReader reader = new BundleReader(storageClient, _memoryCache, _logger);
 
             BundleHeader header = await reader.ReadBundleHeaderAsync(locator, cancellationToken);
diff --git a/Engine/Source/Programs/UnrealCloudDDC/Jupiter/Implementation/Bundles/StorageService.cs b/Engine/Source/Programs/UnrealCloudDDC/Jupiter/Implementation/Bundles/StorageService.cs
--- a/Engine/Source/Programs/UnrealCloudDDC/Jupiter/Implementation/Bundles/StorageService.cs
+++ b/Engine/Source/Programs/UnrealCloudDDC/Jupiter/Implementation/Bundles/StorageService.cs
@@ -4,7 +4,6 @@ using System;
 using System.Buffers;
 using System.Collections.Concurrent;
 using System.Collections.Generic;
-using System.IO;
 using System.Runtime.CompilerServices;
 using System.Threading;
 using System.Threading.Tasks;
@@ -17,17 +16,9 @@ using Microsoft.Extensions.Logging.Abstractions;
 
 namespace Jupiter.Implementation.Bundles;
 
-public interface IStorageClientJupiter : IStorageClient
-{
-    Task<(BlobLocator Locator, Uri UploadUrl)?> GetWriteRedirectAsync(string prefix, CancellationToken cancellationToken);
-    bool SupportsRedirects { get; set; }
-    Task<Uri?> GetReadRedirectAsync(BlobLocator locator, CancellationToken cancellationToken);
-    Task WriteRefTargetAsync(RefName refName, BundleNodeLocator target, RefOptions? requestOptions, CancellationToken cancellationToken);
-}
-
 public interface IStorageService
 {
-    Task<IStorageClientJupiter> GetClientAsync(NamespaceId namespaceId, CancellationToken cancellationToken);
+    Task<StorageClient> GetClientAsync(NamespaceId namespaceId, CancellationToken cancellationToken);
 }
 
 public class StorageService : IStorageService
@@ -40,14 +31,14 @@ public class StorageService : IStorageService
         _provider = provider;
     }
 
-    public Task<IStorageClientJupiter> GetClientAsync(NamespaceId namespaceId, CancellationToken cancellationToken)
+    public Task<StorageClient> GetClientAsync(NamespaceId namespaceId, CancellationToken cancellationToken)
     {
-        IStorageClientJupiter storageClient = _backends.GetOrAdd(namespaceId, x => ActivatorUtilities.CreateInstance<StorageClient>(_provider, namespaceId));
+        StorageClient storageClient = _backends.GetOrAdd(namespaceId, x => ActivatorUtilities.CreateInstance<StorageClient>(_provider, namespaceId));
         return Task.FromResult(storageClient);
     }
 }
 
-public class StorageClient : IStorageClientJupiter
+public class StorageClient : BundleStorageClient
 {
     private readonly NamespaceId _namespaceId;
     private readonly IBlobService _blobService;
@@ -59,6 +50,7 @@ public class StorageClient : IStorageClientJupiter
     public bool SupportsRedirects { get; set; } = true;
 
     public StorageClient(NamespaceId namespaceId, IBlobService blobService, IReferencesStore refStore, IBlobIndex blobIndex)
+        : base(null, NullLogger.Instance)
     {
         _namespaceId = namespaceId;
         _blobService = blobService;
@@ -79,7 +71,7 @@ public class StorageClient : IStorageClientJupiter
         return (locator, redirectUri);
     }
 
-    public async Task<BlobLocator> WriteBundleAsync(Bundle bundle, Utf8String prefix, CancellationToken cancellationToken)
+    public override async Task<BlobLocator> WriteBundleAsync(Bundle bundle, Utf8String prefix, CancellationToken cancellationToken)
     {
         BlobLocator locator = BlobLocator.CreateUnique(prefix);
         BlobIdentifier blobIdentifier = BlobIdentifier.FromBlobLocator(locator);
@@ -99,19 +91,19 @@ public class StorageClient : IStorageClientJupiter
         return locator;
     }
 
-    public async Task AddAliasAsync(Utf8String name, BlobHandle handle, CancellationToken cancellationToken = default)
+    public override async Task AddAliasAsync(Utf8String name, BlobHandle handle, CancellationToken cancellationToken = default)
     {
         // TODO: Implement aliases
         await Task.CompletedTask;
     }
 
-    public async Task RemoveAliasAsync(Utf8String name, BlobHandle handle, CancellationToken cancellationToken = default)
+    public override async Task RemoveAliasAsync(Utf8String name, BlobHandle handle, CancellationToken cancellationToken = default)
     {
         // TODO: Implement aliases
         await Task.CompletedTask;
     }
 
-    public async IAsyncEnumerable<BlobHandle> FindNodesAsync(Utf8String name, [EnumeratorCancellation] CancellationToken cancellationToken = default)
+    public override async IAsyncEnumerable<BlobHandle> FindNodesAsync(Utf8String name, [EnumeratorCancellation] CancellationToken cancellationToken = default)
     {
         // TODO: Implement aliases
         await Task.CompletedTask;
@@ -125,16 +117,17 @@ public class StorageClient : IStorageClientJupiter
         return redirectUri;
     }
 
-    public async Task<Bundle> ReadBundleAsync(BlobLocator locator, CancellationToken cancellationToken)
+    public override async Task<Bundle> ReadBundleAsync(BlobLocator locator, CancellationToken cancellationToken)
     {
         BlobIdentifier blobIdentifier = BlobIdentifier.FromBlobLocator(locator);
         BlobContents blobContents = await _blobService.GetObject(_namespaceId, blobIdentifier);
         return await Bundle.FromStreamAsync(blobContents.Stream, cancellationToken);
     }
 
-    public async Task<ReadOnlyMemory<byte>> ReadBundleRangeAsync(BlobLocator locator, int offset, int length, CancellationToken cancellationToken)
+    public override async Task<ReadOnlyMemory<byte>> ReadBundleRangeAsync(BlobLocator locator, int offset, int length, CancellationToken cancellationToken)
     {
-        ReadOnlySequence<byte> sequence = await ReadBundleAsync(locator, cancellationToken);
+        Bundle bundle = await ReadBundleAsync(locator, cancellationToken);
+        ReadOnlySequence<byte> sequence = bundle.AsSequence();
 		sequence = sequence.Slice(offset);
 
 		if(sequence.Length > length)
@@ -145,7 +138,7 @@ public class StorageClient : IStorageClientJupiter
         return sequence.AsSingleSegment();
     }
 
-    public async Task<BlobHandle?> TryReadRefTargetAsync(RefName name, RefCacheTime cacheTime = default, CancellationToken cancellationToken = default)
+    public override async Task<BlobHandle?> TryReadRefTargetAsync(RefName name, RefCacheTime cacheTime = default, CancellationToken cancellationToken = default)
     {
         // TODO: Cache time is ignored
         try
@@ -172,14 +165,14 @@ public class StorageClient : IStorageClientJupiter
 
     public async Task<BlobHandle> WriteRefAsync(RefName name, Bundle bundle, int exportIdx, Utf8String prefix = default, RefOptions? options = null, CancellationToken cancellationToken = default)
     {
-        BlobLocator locator = await this.WriteBundleAsync(bundle, prefix, cancellationToken);
+        BlobLocator locator = await WriteBundleAsync(bundle, prefix, cancellationToken);
         BlobHandle target = new FlushedNodeHandle(_treeReader, new BundleNodeLocator(bundle.Header.Exports[exportIdx].Hash, locator, exportIdx));
         await WriteRefTargetAsync(name, target, options, cancellationToken);
 
         return target;
     }
 
-    public Task WriteRefTargetAsync(RefName refName, BlobHandle target, RefOptions? requestOptions, CancellationToken cancellationToken)
+    public override Task WriteRefTargetAsync(RefName refName, BlobHandle target, RefOptions? requestOptions, CancellationToken cancellationToken)
     {
         return WriteRefTargetAsync(refName, target.GetLocator(), requestOptions, cancellationToken);
     }
@@ -200,17 +193,10 @@ public class StorageClient : IStorageClientJupiter
         await _refStore.Put(_namespaceId, _defaultBucket, refKey, blobIdentifier, payload, isFinalized: true); 
     }
 
-    public  async Task DeleteRefAsync(RefName name, CancellationToken cancellationToken = default)
+    public override async Task DeleteRefAsync(RefName name, CancellationToken cancellationToken = default)
     {
         await _refStore.Delete(_namespaceId, _defaultBucket, IoHashKey.FromName(name.ToString()));
     }
-
-    public BundleWriter CreateWriter(RefName refName = default, BundleOptions? options = null)
-    {
-        return new BundleWriter(this, _treeReader, refName, options);
-    }
-
-    IStorageWriter IStorageClient.CreateWriter(RefName refName) => CreateWriter(refName);
 }
 
 public class RefInlinePayload
