Horde: Prevent exceptions being trapped and replaced with generic exceptions when reading from storage.

diff --git a/Engine/Source/Programs/Horde/Horde.Server/Jobs/Artifacts/ArtifactCollectionV1.cs b/Engine/Source/Programs/Horde/Horde.Server/Jobs/Artifacts/ArtifactCollectionV1.cs
--- a/Engine/Source/Programs/Horde/Horde.Server/Jobs/Artifacts/ArtifactCollectionV1.cs
+++ b/Engine/Source/Programs/Horde/Horde.Server/Jobs/Artifacts/ArtifactCollectionV1.cs
@@ -197,11 +197,7 @@ namespace Horde.Server.Jobs.Artifacts
 		/// <returns>The chunk data</returns>
 		public async Task<System.IO.Stream> OpenArtifactReadStreamAsync(IArtifactV1 artifact)
 		{
-			System.IO.Stream? stream = await _storageBackend.TryReadAsync(GetPath(artifact.JobId, artifact.StepId, artifact.Name));
-			if (stream == null)
-			{
-				throw new Exception($"Unable to get artifact {artifact.Id}");
-			}
+			System.IO.Stream stream = await _storageBackend.ReadAsync(GetPath(artifact.JobId, artifact.StepId, artifact.Name));
 			return stream;
 		}
 
diff --git a/Engine/Source/Programs/Horde/Horde.Server/Storage/Backends/AwsStorageBackend.cs b/Engine/Source/Programs/Horde/Horde.Server/Storage/Backends/AwsStorageBackend.cs
--- a/Engine/Source/Programs/Horde/Horde.Server/Storage/Backends/AwsStorageBackend.cs
+++ b/Engine/Source/Programs/Horde/Horde.Server/Storage/Backends/AwsStorageBackend.cs
@@ -257,20 +257,20 @@ namespace Horde.Server.Storage.Backends
 		string GetFullPath(string path) => _pathPrefix + path;
 
 		/// <inheritdoc/>
-		public Task<Stream?> TryReadAsync(string path, CancellationToken cancellationToken)
+		public Task<Stream> ReadAsync(string path, CancellationToken cancellationToken)
 		{
-			return TryReadAsync(path, null, cancellationToken);
+			return ReadAsync(path, null, cancellationToken);
 		}
 
 		/// <inheritdoc/>
-		public Task<Stream?> TryReadAsync(string path, int offset, int length, CancellationToken cancellationToken)
+		public Task<Stream> ReadAsync(string path, int offset, int length, CancellationToken cancellationToken)
 		{
-			return TryReadAsync(path, new ByteRange(offset, offset + length), cancellationToken);
+			return ReadAsync(path, new ByteRange(offset, offset + length), cancellationToken);
 		}
 
-		async Task<Stream?> TryReadAsync(string path, ByteRange? byteRange, CancellationToken cancellationToken)
+		async Task<Stream> ReadAsync(string path, ByteRange? byteRange, CancellationToken cancellationToken)
 		{
-			using TelemetrySpan span = OpenTelemetryTracers.Horde.StartActiveSpan($"{nameof(AwsStorageBackend)}.{nameof(TryReadAsync)}");
+			using TelemetrySpan span = OpenTelemetryTracers.Horde.StartActiveSpan($"{nameof(AwsStorageBackend)}.{nameof(ReadAsync)}");
 			span.SetAttribute("path", path);
 			
 			string fullPath = GetFullPath(path);
@@ -297,7 +297,7 @@ namespace Horde.Server.Storage.Backends
 				semaLock?.Dispose();
 				response?.Dispose();
 
-				return null;
+				throw new StorageException($"Unable to read {fullPath} from {_options.AwsBucketName}", ex);
 			}
 		}
 
diff --git a/Engine/Source/Programs/Horde/Horde.Server/Storage/Backends/FileStorageBackend.cs b/Engine/Source/Programs/Horde/Horde.Server/Storage/Backends/FileStorageBackend.cs
--- a/Engine/Source/Programs/Horde/Horde.Server/Storage/Backends/FileStorageBackend.cs
+++ b/Engine/Source/Programs/Horde/Horde.Server/Storage/Backends/FileStorageBackend.cs
@@ -63,39 +63,20 @@ namespace Horde.Server.Storage.Backends
 		}
 
 		/// <inheritdoc/>
-		public Task<Stream?> TryReadAsync(string path, CancellationToken cancellationToken)
+		public Task<Stream> ReadAsync(string path, CancellationToken cancellationToken)
 		{
-			using TelemetrySpan span = OpenTelemetryTracers.Horde.StartActiveSpan($"{nameof(FileSystemStorageBackend)}.{nameof(TryReadAsync)}");
+			using TelemetrySpan span = OpenTelemetryTracers.Horde.StartActiveSpan($"{nameof(FileSystemStorageBackend)}.{nameof(ReadAsync)}");
 			span.SetAttribute("path", path);
 			
 			FileReference location = FileReference.Combine(_baseDir, path);
-			if (!FileReference.Exists(location))
-			{
-				return Task.FromResult<Stream?>(null);
-			}
-
-			try
-			{
-				return Task.FromResult<Stream?>(FileReference.Open(location, FileMode.Open, FileAccess.Read, FileShare.Read));
-			}
-			catch (DirectoryNotFoundException)
-			{
-				return Task.FromResult<Stream?>(null);
-			}
-			catch (FileNotFoundException)
-			{
-				return Task.FromResult<Stream?>(null);
-			}
+			return Task.FromResult<Stream>(FileReference.Open(location, FileMode.Open, FileAccess.Read, FileShare.Read));
 		}
 
 		/// <inheritdoc/>
-		public async Task<Stream?> TryReadAsync(string path, int offset, int length, CancellationToken cancellationToken)
+		public async Task<Stream> ReadAsync(string path, int offset, int length, CancellationToken cancellationToken)
 		{
-			Stream? stream = await TryReadAsync(path, cancellationToken);
-			if (stream != null)
-			{
-				stream.Seek(offset, SeekOrigin.Begin);
-			}
+			Stream stream = await ReadAsync(path, cancellationToken);
+			stream.Seek(offset, SeekOrigin.Begin);
 			return stream;
 		}
 
diff --git a/Engine/Source/Programs/Horde/Horde.Server/Storage/Backends/MemoryStorageBackend.cs b/Engine/Source/Programs/Horde/Horde.Server/Storage/Backends/MemoryStorageBackend.cs
--- a/Engine/Source/Programs/Horde/Horde.Server/Storage/Backends/MemoryStorageBackend.cs
+++ b/Engine/Source/Programs/Horde/Horde.Server/Storage/Backends/MemoryStorageBackend.cs
@@ -29,27 +29,17 @@ namespace Horde.Server.Storage.Backends
 		}
 
 		/// <inheritdoc/>
-		public Task<Stream?> TryReadAsync(string path, CancellationToken cancellationToken)
+		public Task<Stream> ReadAsync(string path, CancellationToken cancellationToken)
 		{
-			byte[]? data;
-			if (_pathToData.TryGetValue(path, out data))
-			{
-				return Task.FromResult<Stream?>(new MemoryStream(data, false));
-			}
-			else
-			{
-				return Task.FromResult<Stream?>(null);
-			}
+			byte[] data = _pathToData[path];
+			return Task.FromResult<Stream>(new MemoryStream(data, false));
 		}
 
 		/// <inheritdoc/>
-		public async Task<Stream?> TryReadAsync(string path, int offset, int length, CancellationToken cancellationToken)
+		public async Task<Stream> ReadAsync(string path, int offset, int length, CancellationToken cancellationToken)
 		{
-			Stream? stream = await TryReadAsync(path, cancellationToken);
-			if (stream != null)
-			{
-				stream.Seek(offset, SeekOrigin.Begin);
-			}
+			Stream stream = await ReadAsync(path, cancellationToken);
+			stream.Seek(offset, SeekOrigin.Begin);
 			return stream;
 		}
 
diff --git a/Engine/Source/Programs/Horde/Horde.Server/Storage/IStorageBackend.cs b/Engine/Source/Programs/Horde/Horde.Server/Storage/IStorageBackend.cs
--- a/Engine/Source/Programs/Horde/Horde.Server/Storage/IStorageBackend.cs
+++ b/Engine/Source/Programs/Horde/Horde.Server/Storage/IStorageBackend.cs
@@ -6,7 +6,6 @@ using System.IO;
 using System.Threading;
 using System.Threading.Tasks;
 using EpicGames.Core;
-using EpicGames.Horde.Storage;
 
 namespace Horde.Server.Storage
 {
@@ -26,7 +25,7 @@ namespace Horde.Server.Storage
 		/// <param name="path">Relative path within the bucket</param>
 		/// <param name="cancellationToken">Cancellation token for the operation</param>
 		/// <returns></returns>
-		Task<Stream?> TryReadAsync(string path, CancellationToken cancellationToken = default);
+		Task<Stream> ReadAsync(string path, CancellationToken cancellationToken = default);
 
 		/// <summary>
 		/// Attempts to open a read stream for the given path.
@@ -36,7 +35,7 @@ namespace Horde.Server.Storage
 		/// <param name="length">Length of data to read</param>
 		/// <param name="cancellationToken">Cancellation token for the operation</param>
 		/// <returns></returns>
-		Task<Stream?> TryReadAsync(string path, int offset, int length, CancellationToken cancellationToken = default);
+		Task<Stream> ReadAsync(string path, int offset, int length, CancellationToken cancellationToken = default);
 
 		/// <summary>
 		/// Writes a stream to the given path. If the stream throws an exception during read, the write will be aborted.
@@ -120,10 +119,10 @@ namespace Horde.Server.Storage
 			public void Dispose() => _inner.Dispose();
 
 			/// <inheritdoc/>
-			public Task<Stream?> TryReadAsync(string path, CancellationToken cancellationToken) => _inner.TryReadAsync(path, cancellationToken);
+			public Task<Stream> ReadAsync(string path, CancellationToken cancellationToken) => _inner.ReadAsync(path, cancellationToken);
 
 			/// <inheritdoc/>
-			public Task<Stream?> TryReadAsync(string path, int offset, int length, CancellationToken cancellationToken) => _inner.TryReadAsync(path, offset, length, cancellationToken);
+			public Task<Stream> ReadAsync(string path, int offset, int length, CancellationToken cancellationToken) => _inner.ReadAsync(path, offset, length, cancellationToken);
 
 			/// <inheritdoc/>
 			public Task WriteAsync(string path, Stream stream, CancellationToken cancellationToken) => _inner.WriteAsync(path, stream, cancellationToken);
@@ -209,51 +208,10 @@ namespace Horde.Server.Storage
 		/// <param name="path"></param>
 		/// <param name="cancellationToken"></param>
 		/// <returns></returns>
-		public static async Task<Stream> ReadAsync(this IStorageBackend storageBackend, string path, CancellationToken cancellationToken = default)
+		public static async Task<ReadOnlyMemory<byte>> ReadBytesAsync(this IStorageBackend storageBackend, string path, CancellationToken cancellationToken = default)
 		{
-			Stream? stream = await storageBackend.TryReadAsync(path, cancellationToken);
-			if (stream == null)
+			using (Stream inputStream = await storageBackend.ReadAsync(path, cancellationToken))
 			{
-				throw new FileNotFoundException($"Unable to read from path {path}");
-			}
-			return stream;
-		}
-
-		/// <summary>
-		/// Writes a block of memory to storage
-		/// </summary>
-		/// <param name="storageBackend"></param>
-		/// <param name="path"></param>
-		/// <param name="offset">Offset of the data to read</param>
-		/// <param name="length">Length of the data to read</param>
-		/// <param name="cancellationToken"></param>
-		/// <returns></returns>
-		public static async Task<Stream> ReadAsync(this IStorageBackend storageBackend, string path, int offset, int length, CancellationToken cancellationToken = default)
-		{
-			Stream? stream = await storageBackend.TryReadAsync(path, offset, length, cancellationToken);
-			if (stream == null)
-			{
-				throw new FileNotFoundException($"Unable to read from path {path}");
-			}
-			return stream;
-		}
-
-		/// <summary>
-		/// Writes a block of memory to storage
-		/// </summary>
-		/// <param name="storageBackend"></param>
-		/// <param name="path"></param>
-		/// <param name="cancellationToken"></param>
-		/// <returns></returns>
-		public static async Task<ReadOnlyMemory<byte>?> ReadBytesAsync(this IStorageBackend storageBackend, string path, CancellationToken cancellationToken = default)
-		{
-			using (Stream? inputStream = await storageBackend.TryReadAsync(path, cancellationToken))
-			{
-				if (inputStream == null)
-				{
-					return null;
-				}
-
 				using (MemoryStream outputStream = new MemoryStream())
 				{
 					await inputStream.CopyToAsync(outputStream, cancellationToken);
diff --git a/Engine/Source/Programs/Horde/Horde.Server/Storage/StorageBackendProvider.cs b/Engine/Source/Programs/Horde/Horde.Server/Storage/StorageBackendProvider.cs
--- a/Engine/Source/Programs/Horde/Horde.Server/Storage/StorageBackendProvider.cs
+++ b/Engine/Source/Programs/Horde/Horde.Server/Storage/StorageBackendProvider.cs
@@ -77,10 +77,10 @@ namespace Horde.Server.Storage
 			public IAsyncEnumerable<string> EnumerateAsync(CancellationToken cancellationToken = default) => _backend.EnumerateAsync(cancellationToken);
 
 			/// <inheritdoc/>
-			public Task<Stream?> TryReadAsync(string path, CancellationToken cancellationToken = default) => _backend.TryReadAsync(path, cancellationToken);
+			public Task<Stream> ReadAsync(string path, CancellationToken cancellationToken = default) => _backend.ReadAsync(path, cancellationToken);
 
 			/// <inheritdoc/>
-			public Task<Stream?> TryReadAsync(string path, int offset, int length, CancellationToken cancellationToken = default) => _backend.TryReadAsync(path, offset, length, cancellationToken);
+			public Task<Stream> ReadAsync(string path, int offset, int length, CancellationToken cancellationToken = default) => _backend.ReadAsync(path, offset, length, cancellationToken);
 
 			/// <inheritdoc/>
 			public Task WriteAsync(string path, Stream stream, CancellationToken cancellationToken = default) => _backend.WriteAsync(path, stream, cancellationToken);
diff --git a/Engine/Source/Programs/Horde/Horde.Server/Storage/StorageService.cs b/Engine/Source/Programs/Horde/Horde.Server/Storage/StorageService.cs
--- a/Engine/Source/Programs/Horde/Horde.Server/Storage/StorageService.cs
+++ b/Engine/Source/Programs/Horde/Horde.Server/Storage/StorageService.cs
@@ -140,13 +140,7 @@ namespace Horde.Server.Storage
 			public override async Task<Bundle> ReadBundleAsync(BundleLocator locator, CancellationToken cancellationToken = default)
 			{
 				string path = GetBlobPath(locator);
-
-				Stream? stream = await Backend.TryReadAsync(path, cancellationToken);
-				if (stream == null)
-				{
-					throw new StorageException($"Unable to read data from {path}");
-				}
-
+				Stream stream = await Backend.ReadAsync(path, cancellationToken);
 				return await Bundle.FromStreamAsync(stream, cancellationToken);
 			}
 
@@ -157,13 +151,7 @@ namespace Horde.Server.Storage
 			public override async Task<ReadOnlyMemory<byte>> ReadBundleRangeAsync(BundleLocator locator, int offset, int length, CancellationToken cancellationToken = default)
 			{
 				string path = GetBlobPath(locator);
-
-				Stream? stream = await Backend.TryReadAsync(path, offset, length, cancellationToken);
-				if (stream == null)
-				{
-					throw new StorageException($"Unable to read data from {path}");
-				}
-
+				Stream stream = await Backend.ReadAsync(path, offset, length, cancellationToken);
 				return await stream.ReadAllBytesAsync(cancellationToken);
 			}
 
diff --git a/Engine/Source/Programs/Shared/EpicGames.Horde/Storage/Bundles/BundleReader.cs b/Engine/Source/Programs/Shared/EpicGames.Horde/Storage/Bundles/BundleReader.cs
--- a/Engine/Source/Programs/Shared/EpicGames.Horde/Storage/Bundles/BundleReader.cs
+++ b/Engine/Source/Programs/Shared/EpicGames.Horde/Storage/Bundles/BundleReader.cs
@@ -417,7 +417,7 @@ namespace EpicGames.Horde.Storage.Bundles
 					StartReadTask();
 				}
 			}
-			return await queuedHeader.CompletionSource.Task.AbandonOnCancel(cancellationToken);
+			return await queuedHeader.CompletionSource.Task.WaitAsync(cancellationToken);
 		}
 
 		/// <summary>
@@ -456,7 +456,7 @@ namespace EpicGames.Horde.Storage.Bundles
 					_decodeTasks.Add(decodedCacheKey, decodeTask);
 				}
 			}
-			return await decodeTask.AbandonOnCancel(cancellationToken);
+			return await decodeTask.WaitAsync(cancellationToken);
 		}
 
 		/// <summary>
