Horde: Fix computation of AutoSDK workspace views for conform jobs.

diff --git a/Engine/Source/Programs/Horde/Horde.Server/Agents/IAgent.cs b/Engine/Source/Programs/Horde/Horde.Server/Agents/IAgent.cs
--- a/Engine/Source/Programs/Horde/Horde.Server/Agents/IAgent.cs
+++ b/Engine/Source/Programs/Horde/Horde.Server/Agents/IAgent.cs
@@ -947,17 +947,32 @@ namespace Horde.Server.Agents
 
 				// Create the new workspace
 				workspace = new AgentWorkspace(workspaceConfig.Cluster, workspaceConfig.UserName, identifier, workspaceConfig.Stream ?? streamConfig.Name, workspaceConfig.View, workspaceConfig.Incremental, workspaceConfig.Method);
+				autoSdkConfig = GetAutoSdkConfig(workspaceConfig, streamConfig);
 
+				return true;
+			}
+
+			static AutoSdkConfig? GetAutoSdkConfig(WorkspaceConfig workspaceConfig, StreamConfig streamConfig)
+			{
+				AutoSdkConfig? autoSdkConfig = null;
 				if (workspaceConfig.UseAutoSdk)
 				{
-					autoSdkConfig = new AutoSdkConfig(Enumerable.Concat(streamConfig.AutoSdkView ?? Enumerable.Empty<string>(), workspaceConfig.AutoSdkView ?? Enumerable.Empty<string>()));
-				}
-				else
-				{
-					autoSdkConfig = null;
+					List<string> view = new List<string>();
+					if (streamConfig.AutoSdkView != null)
+					{
+						view.AddRange(streamConfig.AutoSdkView);
+					}
+					if (workspaceConfig.AutoSdkView != null)
+					{
+						view.AddRange(workspaceConfig.AutoSdkView);
+					}
+					if (view.Count == 0)
+					{
+						view.Add("...");
+					}
+					autoSdkConfig = new AutoSdkConfig(view);
 				}
-
-				return true;
+				return autoSdkConfig;
 			}
 		}
 	}
