Horde: Include autosdk config in pool responses.

diff --git a/Engine/Source/Programs/Horde/Horde.Server/Agents/Pools/PoolMessages.cs b/Engine/Source/Programs/Horde/Horde.Server/Agents/Pools/PoolMessages.cs
--- a/Engine/Source/Programs/Horde/Horde.Server/Agents/Pools/PoolMessages.cs
+++ b/Engine/Source/Programs/Horde/Horde.Server/Agents/Pools/PoolMessages.cs
@@ -508,6 +508,11 @@ namespace Horde.Server.Agents.Pools
 		/// </summary>
 		public List<GetAgentWorkspaceResponse> Workspaces { get; set; }
 
+		/// <summary>
+		/// AutoSDK configuration for the pool 
+		/// </summary>
+		public AutoSdkConfig? AutoSdkConfig { get; set; }
+
 		/// <summary>
 		/// Arbitrary properties for this pool.
 		/// </summary>
@@ -535,6 +540,7 @@ namespace Horde.Server.Agents.Pools
 			MinAgents = pool.MinAgents;
 			NumReserveAgents = pool.NumReserveAgents;
 			Workspaces = pool.Workspaces.Select(x => new GetAgentWorkspaceResponse(x)).ToList();
+			AutoSdkConfig = pool.AutoSdkConfig;
 			Properties = pool.Properties;
 		}
 	}
