Horde: Fix log output for WorkspaceMaterializer not appearing in the batch log.

diff --git a/Engine/Source/Programs/Horde/Horde.Agent/Execution/ManagedWorkspaceMaterializer.cs b/Engine/Source/Programs/Horde/Horde.Agent/Execution/ManagedWorkspaceMaterializer.cs
--- a/Engine/Source/Programs/Horde/Horde.Agent/Execution/ManagedWorkspaceMaterializer.cs
+++ b/Engine/Source/Programs/Horde/Horde.Agent/Execution/ManagedWorkspaceMaterializer.cs
@@ -20,7 +20,6 @@ public class ManagedWorkspaceMaterializer : IWorkspaceMaterializer
 	private readonly DirectoryReference _workingDir;
 	private readonly bool _useSyncMarker;
 	private readonly bool _useCacheFile;
-	private readonly ILogger<ManagedWorkspaceMaterializer> _logger;
 	private WorkspaceInfo? _workspace;
 	
 	/// <summary>
@@ -30,28 +29,25 @@ public class ManagedWorkspaceMaterializer : IWorkspaceMaterializer
 	/// <param name="workingDir">Where to put synced Perforce files and any cached data/metadata</param>
 	/// <param name="useSyncMarker">Whether to use a sync marker for identifying last synced change number</param>
 	/// <param name="useCacheFile">Whether to use a cache file during syncs</param>
-	/// <param name="logger">Logger</param>
 	public ManagedWorkspaceMaterializer(
 		AgentWorkspace agentWorkspace,
 		DirectoryReference workingDir,
 		bool useSyncMarker,
-		bool useCacheFile,
-		ILogger<ManagedWorkspaceMaterializer> logger)
+		bool useCacheFile)
 	{
 		_agentWorkspace = agentWorkspace;
 		_workingDir = workingDir;
 		_useSyncMarker = useSyncMarker;
 		_useCacheFile = useCacheFile;
-		_logger = logger;
 	}
 
 	/// <inheritdoc/>
-	public async Task<WorkspaceMaterializerSettings> InitializeAsync(CancellationToken cancellationToken)
+	public async Task<WorkspaceMaterializerSettings> InitializeAsync(ILogger logger, CancellationToken cancellationToken)
 	{
 		using IScope scope = CreateTraceSpan("ManagedWorkspaceMaterializer.InitializeAsync");
 		
 		bool useHaveTable = WorkspaceInfo.ShouldUseHaveTable(_agentWorkspace.Method);
-		_workspace = await WorkspaceInfo.SetupWorkspaceAsync(_agentWorkspace, _workingDir, useHaveTable, _logger, cancellationToken);
+		_workspace = await WorkspaceInfo.SetupWorkspaceAsync(_agentWorkspace, _workingDir, useHaveTable, logger, cancellationToken);
 		return await GetSettingsAsync(cancellationToken);
 	}
 
@@ -111,7 +107,10 @@ public class ManagedWorkspaceMaterializer : IWorkspaceMaterializer
 	/// <returns>True if already synced</returns>
 	private bool IsAlreadySynced(int changeNum)
 	{
-		if (!_useSyncMarker) return false;
+		if (!_useSyncMarker)
+		{
+			return false;
+		}
 
 		(FileReference syncFile, string syncText) = GetSyncMarker(changeNum);
 		if (!FileReference.Exists(syncFile) || FileReference.ReadAllText(syncFile) != syncText)
@@ -130,7 +129,10 @@ public class ManagedWorkspaceMaterializer : IWorkspaceMaterializer
 	/// <returns>True if already synced</returns>
 	private void MarkChangeNumSynced(int changeNum)
 	{
-		if (!_useSyncMarker) return;
+		if (!_useSyncMarker)
+		{
+			return;
+		}
 		
 		(FileReference syncFile, string syncText) = GetSyncMarker(changeNum);
 		FileReference.WriteAllText(syncFile, syncText);
diff --git a/Engine/Source/Programs/Horde/Horde.Agent/Execution/WorkspaceExecutor.cs b/Engine/Source/Programs/Horde/Horde.Agent/Execution/WorkspaceExecutor.cs
--- a/Engine/Source/Programs/Horde/Horde.Agent/Execution/WorkspaceExecutor.cs
+++ b/Engine/Source/Programs/Horde/Horde.Agent/Execution/WorkspaceExecutor.cs
@@ -44,7 +44,7 @@ namespace Horde.Agent.Execution
 				using IScope _ = GlobalTracer.Instance.BuildSpan("Workspace").WithResourceName("AutoSDK").StartActive();
 				// TODO: Set type of workspace materializer as scope tag.
 
-				autoSdkWorkspaceSettings = await _autoSdkWorkspace.InitializeAsync(cancellationToken);
+				autoSdkWorkspaceSettings = await _autoSdkWorkspace.InitializeAsync(logger, cancellationToken);
 
 				// Match change for AutoSDK and actual job change
 				int autoSdkChangeNumber = _batch.Change;
@@ -57,7 +57,7 @@ namespace Horde.Agent.Execution
 			WorkspaceMaterializerSettings workspaceSettings;
 			using (IScope scope = GlobalTracer.Instance.BuildSpan("Workspace").StartActive())
 			{
-				workspaceSettings = await _workspace.InitializeAsync(cancellationToken);
+				workspaceSettings = await _workspace.InitializeAsync(logger, cancellationToken);
 				scope.Span.SetTag(Datadog.Trace.OpenTracing.DatadogTags.ResourceName, workspaceSettings.Identifier);
 				
 				int preflightChange = (_batch.ClonedPreflightChange != 0) ? _batch.ClonedPreflightChange : _batch.PreflightChange;
diff --git a/Engine/Source/Programs/Horde/Horde.Agent/Execution/WorkspaceMaterializer.cs b/Engine/Source/Programs/Horde/Horde.Agent/Execution/WorkspaceMaterializer.cs
--- a/Engine/Source/Programs/Horde/Horde.Agent/Execution/WorkspaceMaterializer.cs
+++ b/Engine/Source/Programs/Horde/Horde.Agent/Execution/WorkspaceMaterializer.cs
@@ -92,9 +92,10 @@ public interface IWorkspaceMaterializer
 	/// <summary>
 	/// Prepare file system for syncing
 	/// </summary>
+	/// <param name="logger">Logger for output</param>
 	/// <param name="cancellationToken">Cancellation token for the call</param>
 	/// <returns>Async task</returns>
-	public Task<WorkspaceMaterializerSettings> InitializeAsync(CancellationToken cancellationToken);
+	public Task<WorkspaceMaterializerSettings> InitializeAsync(ILogger logger, CancellationToken cancellationToken);
 	
 	/// <summary>
 	/// Finalize and clean file system
@@ -167,8 +168,8 @@ class WorkspaceMaterializerFactory : IWorkspaceMaterializerFactory
 		{
 			case WorkspaceMaterializerType.ManagedWorkspace:
 				return forAutoSdk
-					? new ManagedWorkspaceMaterializer(workspaceInfo, options.Session.WorkingDir, true, true, _loggerFactory.CreateLogger<ManagedWorkspaceMaterializer>())
-					: new ManagedWorkspaceMaterializer(workspaceInfo, options.Session.WorkingDir, false, false, _loggerFactory.CreateLogger<ManagedWorkspaceMaterializer>());
+					? new ManagedWorkspaceMaterializer(workspaceInfo, options.Session.WorkingDir, true, true)
+					: new ManagedWorkspaceMaterializer(workspaceInfo, options.Session.WorkingDir, false, false);
 
 			default:
 				throw new Exception("Unhandled materializer option: " + type);
