Horde: Fix tests.

diff --git a/Engine/Source/Programs/Horde/Horde.Agent.Tests/Execution/FakeWorkspaceMaterializer.cs b/Engine/Source/Programs/Horde/Horde.Agent.Tests/Execution/FakeWorkspaceMaterializer.cs
--- a/Engine/Source/Programs/Horde/Horde.Agent.Tests/Execution/FakeWorkspaceMaterializer.cs
+++ b/Engine/Source/Programs/Horde/Horde.Agent.Tests/Execution/FakeWorkspaceMaterializer.cs
@@ -7,6 +7,7 @@ using System.Threading;
 using System.Threading.Tasks;
 using EpicGames.Core;
 using Horde.Agent.Execution;
+using Microsoft.Extensions.Logging;
 
 namespace Horde.Agent.Tests.Execution;
 
@@ -32,7 +33,7 @@ public class FakeWorkspaceMaterializer : IWorkspaceMaterializer
 	}
 	
 	/// <inheritdoc/>
-	public Task<WorkspaceMaterializerSettings> InitializeAsync(CancellationToken cancellationToken)
+	public Task<WorkspaceMaterializerSettings> InitializeAsync(ILogger logger, CancellationToken cancellationToken)
 	{
 		if (_isInitialized)
 		{
diff --git a/Engine/Source/Programs/Horde/Horde.Agent.Tests/Execution/FakeWorkspaceMaterializerTest.cs b/Engine/Source/Programs/Horde/Horde.Agent.Tests/Execution/FakeWorkspaceMaterializerTest.cs
--- a/Engine/Source/Programs/Horde/Horde.Agent.Tests/Execution/FakeWorkspaceMaterializerTest.cs
+++ b/Engine/Source/Programs/Horde/Horde.Agent.Tests/Execution/FakeWorkspaceMaterializerTest.cs
@@ -4,6 +4,7 @@ using System.IO;
 using System.Threading;
 using System.Threading.Tasks;
 using Horde.Agent.Execution;
+using Microsoft.Extensions.Logging.Abstractions;
 using Microsoft.VisualStudio.TestTools.UnitTesting;
 
 namespace Horde.Agent.Tests.Execution;
@@ -16,7 +17,7 @@ public class FakeWorkspaceMaterializerTest
 
 	public FakeWorkspaceMaterializerTest()
 	{
-		_settings = _wm.InitializeAsync(CancellationToken.None).Result;
+		_settings = _wm.InitializeAsync(NullLogger.Instance, CancellationToken.None).Result;
 	}
 
 	[TestCleanup]
