Horde: Add telemetry scopes for reading and writing through storage interface.

diff --git a/Engine/Source/Programs/Horde/Horde.Server/Storage/StorageService.cs b/Engine/Source/Programs/Horde/Horde.Server/Storage/StorageService.cs
--- a/Engine/Source/Programs/Horde/Horde.Server/Storage/StorageService.cs
+++ b/Engine/Source/Programs/Horde/Horde.Server/Storage/StorageService.cs
@@ -26,6 +26,7 @@ using Microsoft.Extensions.Options;
 using MongoDB.Bson;
 using MongoDB.Bson.Serialization.Attributes;
 using MongoDB.Driver;
+using OpenTelemetry.Trace;
 using StackExchange.Redis;
 
 namespace Horde.Server.Storage
@@ -119,11 +120,13 @@ namespace Horde.Server.Storage
 		{
 			readonly StorageService _outer;
 			readonly string _prefix;
+			readonly Tracer _tracer;
 
-			public StorageClientImpl(StorageService outer, NamespaceConfig config, IStorageBackend backend, IMemoryCache? memoryCache, ILogger logger)
+			public StorageClientImpl(StorageService outer, NamespaceConfig config, IStorageBackend backend, IMemoryCache? memoryCache, Tracer tracer, ILogger logger)
 				: base(config, backend, memoryCache, logger)
 			{
 				_outer = outer;
+				_tracer = tracer;
 
 				_prefix = config.Prefix;
 				if (_prefix.Length > 0 && !_prefix.EndsWith("/", StringComparison.Ordinal))
@@ -139,6 +142,9 @@ namespace Horde.Server.Storage
 			/// <inheritdoc/>
 			public override async Task<Bundle> ReadBundleAsync(BundleLocator locator, CancellationToken cancellationToken = default)
 			{
+				using TelemetrySpan span = _tracer.StartActiveSpan($"{nameof(StorageService)}.{nameof(StorageClientImpl)}.{nameof(ReadBundleAsync)}");
+				span.SetAttribute("locator", locator.ToString());
+
 				string path = GetBlobPath(locator);
 				Stream stream = await Backend.ReadAsync(path, cancellationToken);
 				return await Bundle.FromStreamAsync(stream, cancellationToken);
@@ -150,6 +156,11 @@ namespace Horde.Server.Storage
 			/// <inheritdoc/>
 			public override async Task<ReadOnlyMemory<byte>> ReadBundleRangeAsync(BundleLocator locator, int offset, int length, CancellationToken cancellationToken = default)
 			{
+				using TelemetrySpan span = _tracer.StartActiveSpan($"{nameof(StorageService)}.{nameof(StorageClientImpl)}.{nameof(ReadBundleRangeAsync)}");
+				span.SetAttribute("locator", locator.ToString());
+				span.SetAttribute("offset", offset);
+				span.SetAttribute("length", length);
+
 				string path = GetBlobPath(locator);
 				Stream stream = await Backend.ReadAsync(path, offset, length, cancellationToken);
 				return await stream.ReadAllBytesAsync(cancellationToken);
@@ -161,6 +172,9 @@ namespace Horde.Server.Storage
 				// Add the blob record
 				BundleLocator locator = await _outer.AddBlobAsync(NamespaceId, prefix, null, cancellationToken);
 
+				using TelemetrySpan span = _tracer.StartActiveSpan($"{nameof(StorageService)}.{nameof(StorageClientImpl)}.{nameof(WriteBundleAsync)}");
+				span.SetAttribute("locator", locator.ToString());
+
 				// Write it to the backend
 				string path = GetBlobPath(locator);
 				using ReadOnlySequenceStream memoryStream = new ReadOnlySequenceStream(bundle.AsSequence());
@@ -423,6 +437,7 @@ namespace Horde.Server.Storage
 		readonly IMemoryCache _cache;
 		readonly IStorageBackendProvider _storageBackendProvider;
 		readonly IOptionsMonitor<GlobalConfig> _globalConfig;
+		readonly Tracer _tracer;
 		readonly ILogger _logger;
 
 		readonly IMongoCollection<BlobInfo> _blobCollection;
@@ -442,7 +457,7 @@ namespace Horde.Server.Storage
 		/// <summary>
 		/// Constructor
 		/// </summary>
-		public StorageService(MongoService mongoService, RedisService redisService, IClock clock, IMemoryCache cache, IStorageBackendProvider storageBackendProvider, IOptionsMonitor<GlobalConfig> globalConfig, ILogger<StorageService> logger)
+		public StorageService(MongoService mongoService, RedisService redisService, IClock clock, IMemoryCache cache, IStorageBackendProvider storageBackendProvider, IOptionsMonitor<GlobalConfig> globalConfig, Tracer tracer, ILogger<StorageService> logger)
 		{
 			_redisService = redisService;
 			_clock = clock;
@@ -450,6 +465,7 @@ namespace Horde.Server.Storage
 			_storageBackendProvider = storageBackendProvider;
 			_cachedState = new AsyncCachedValue<State>(() => Task.FromResult(GetNextState()), TimeSpan.FromMinutes(1.0));
 			_globalConfig = globalConfig;
+			_tracer = tracer;
 			_logger = logger;
 
 			List<MongoIndex<BlobInfo>> blobIndexes = new List<MongoIndex<BlobInfo>>();
@@ -559,7 +575,7 @@ namespace Horde.Server.Storage
 					foreach (NamespaceConfig namespaceConfig in storageConfig.Namespaces)
 					{
 						IStorageBackend backend = _storageBackendProvider.CreateBackend(namespaceConfig.BackendConfig);
-						StorageClientImpl client = new StorageClientImpl(this, namespaceConfig, backend, _cache, _logger);
+						StorageClientImpl client = new StorageClientImpl(this, namespaceConfig, backend, _cache, _tracer, _logger);
 						nextState.Namespaces.Add(namespaceConfig.Id, new NamespaceInfo(namespaceConfig, client, backend));
 					}
 				}
