Horde: Allow non-pure workspace materializers that can use Perforce

Flag added for signaling if the materialized workspace is a true Perforce workspace. This is intended for rolling out ManagedWorkspace but behind the new WorkspaceExecutor. Right now, many job templates require a true P4 workspace to be available.

diff --git a/Engine/Source/Programs/Horde/Horde.Agent.Tests/Execution/FakeWorkspaceMaterializer.cs b/Engine/Source/Programs/Horde/Horde.Agent.Tests/Execution/FakeWorkspaceMaterializer.cs
--- a/Engine/Source/Programs/Horde/Horde.Agent.Tests/Execution/FakeWorkspaceMaterializer.cs
+++ b/Engine/Source/Programs/Horde/Horde.Agent.Tests/Execution/FakeWorkspaceMaterializer.cs
@@ -71,7 +71,7 @@ public class FakeWorkspaceMaterializer : IWorkspaceMaterializer
 			throw new WorkspaceMaterializationException("Cannot get settings before initialization");
 		}
 		
-		return Task.FromResult(new WorkspaceMaterializerSettings(_rootDir, "fakeWorkspaceIdentifier", "fakeWorkspaceStreamRoot"));
+		return Task.FromResult(new WorkspaceMaterializerSettings(_rootDir, "fakeWorkspaceIdentifier", "fakeWorkspaceStreamRoot", false));
 	}
 
 	/// <inheritdoc/>
diff --git a/Engine/Source/Programs/Horde/Horde.Agent/Execution/ManagedWorkspaceMaterializer.cs b/Engine/Source/Programs/Horde/Horde.Agent/Execution/ManagedWorkspaceMaterializer.cs
--- a/Engine/Source/Programs/Horde/Horde.Agent/Execution/ManagedWorkspaceMaterializer.cs
+++ b/Engine/Source/Programs/Horde/Horde.Agent/Execution/ManagedWorkspaceMaterializer.cs
@@ -67,7 +67,7 @@ public class ManagedWorkspaceMaterializer : IWorkspaceMaterializer
 	{
 		// ManagedWorkspace store synced files in a sub-directory from the top working dir.
 		DirectoryReference syncDir = DirectoryReference.Combine(_workingDir, _agentWorkspace.Identifier, "Sync");
-		return Task.FromResult(new WorkspaceMaterializerSettings(syncDir, _agentWorkspace.Identifier, _agentWorkspace.Stream));
+		return Task.FromResult(new WorkspaceMaterializerSettings(syncDir, _agentWorkspace.Identifier, _agentWorkspace.Stream, true));
 	}
 
 	/// <inheritdoc/>
diff --git a/Engine/Source/Programs/Horde/Horde.Agent/Execution/WorkspaceExecutor.cs b/Engine/Source/Programs/Horde/Horde.Agent/Execution/WorkspaceExecutor.cs
--- a/Engine/Source/Programs/Horde/Horde.Agent/Execution/WorkspaceExecutor.cs
+++ b/Engine/Source/Programs/Horde/Horde.Agent/Execution/WorkspaceExecutor.cs
@@ -103,16 +103,18 @@ namespace Horde.Agent.Execution
 		protected override async Task<bool> SetupAsync(BeginStepResponse step, ILogger logger, CancellationToken cancellationToken)
 		{
 			// Loop back to JobExecutor's SetupAsync again, but with workspace and shared storage dir set
-			DirectoryReference workspaceDir = (await _workspace.GetSettingsAsync(cancellationToken)).DirectoryPath;
-			return await SetupAsync(step, workspaceDir, _sharedStorageDir, false, logger, cancellationToken);
+			WorkspaceMaterializerSettings settings = await _workspace.GetSettingsAsync(cancellationToken);
+			DirectoryReference workspaceDir = settings.DirectoryPath;
+			return await SetupAsync(step, workspaceDir, _sharedStorageDir, settings.IsPerforceWorkspace, logger, cancellationToken);
 		}
 
 		/// <inheritdoc/>
 		protected override async Task<bool> ExecuteAsync(BeginStepResponse step, ILogger logger, CancellationToken cancellationToken)
 		{
 			// Loop back to JobExecutor's ExecuteAsync again, but with workspace and shared storage dir set
-			DirectoryReference workspaceDir = (await _workspace.GetSettingsAsync(cancellationToken)).DirectoryPath;
-			return await ExecuteAsync(step, workspaceDir, _sharedStorageDir, false, logger, cancellationToken);
+			WorkspaceMaterializerSettings settings = await _workspace.GetSettingsAsync(cancellationToken);
+			DirectoryReference workspaceDir = settings.DirectoryPath;
+			return await ExecuteAsync(step, workspaceDir, _sharedStorageDir, settings.IsPerforceWorkspace, logger, cancellationToken);
 		}
 
 		/// <inheritdoc/>
diff --git a/Engine/Source/Programs/Horde/Horde.Agent/Execution/WorkspaceMaterializer.cs b/Engine/Source/Programs/Horde/Horde.Agent/Execution/WorkspaceMaterializer.cs
--- a/Engine/Source/Programs/Horde/Horde.Agent/Execution/WorkspaceMaterializer.cs
+++ b/Engine/Source/Programs/Horde/Horde.Agent/Execution/WorkspaceMaterializer.cs
@@ -51,6 +51,13 @@ public class WorkspaceMaterializerSettings
 	/// Stream path inside Perforce
 	/// </summary>
 	public string StreamRoot { get; }
+	
+	/// <summary>
+	/// Whether the materialized workspace is a true Perforce workspace
+	/// This flag is provided as a stop-gap solution to allow replacing ManagedWorkspace with WorkspaceMaterializer.
+	/// It's *highly* recommended to set this to false for any new implementations of IWorkspaceMaterializer.
+	/// </summary>
+	public bool IsPerforceWorkspace { get; }
 
 	/// <summary>
 	/// Constructor
@@ -58,11 +65,13 @@ public class WorkspaceMaterializerSettings
 	/// <param name="directoryPath"></param>
 	/// <param name="identifier"></param>
 	/// <param name="streamRoot"></param>
-	public WorkspaceMaterializerSettings(DirectoryReference directoryPath, string identifier, string streamRoot)
+	/// <param name="isPerforceWorkspace"></param>
+	public WorkspaceMaterializerSettings(DirectoryReference directoryPath, string identifier, string streamRoot, bool isPerforceWorkspace)
 	{
 		DirectoryPath = directoryPath;
 		Identifier = identifier;
 		StreamRoot = streamRoot;
+		IsPerforceWorkspace = isPerforceWorkspace;
 	}
 }
 
