Horde: Fix streams not being disposed when reading bundles. Likely leaking semaphore required to read more data.

diff --git a/Engine/Source/Programs/Horde/Horde.Server/Storage/StorageService.cs b/Engine/Source/Programs/Horde/Horde.Server/Storage/StorageService.cs
--- a/Engine/Source/Programs/Horde/Horde.Server/Storage/StorageService.cs
+++ b/Engine/Source/Programs/Horde/Horde.Server/Storage/StorageService.cs
@@ -146,7 +146,7 @@ namespace Horde.Server.Storage
 				span.SetAttribute("locator", locator.ToString());
 
 				string path = GetBlobPath(locator);
-				Stream stream = await Backend.ReadAsync(path, cancellationToken);
+				await using Stream stream = await Backend.ReadAsync(path, cancellationToken);
 				return await Bundle.FromStreamAsync(stream, cancellationToken);
 			}
 
@@ -162,7 +162,7 @@ namespace Horde.Server.Storage
 				span.SetAttribute("length", length);
 
 				string path = GetBlobPath(locator);
-				Stream stream = await Backend.ReadAsync(path, offset, length, cancellationToken);
+				await using Stream stream = await Backend.ReadAsync(path, offset, length, cancellationToken);
 				return await stream.ReadAllBytesAsync(cancellationToken);
 			}
 
