Horde: Include the stream view and filter in the AutoSDK stream marker. Also delete the marker when unshelving any content.

diff --git a/Engine/Source/Programs/Horde/Horde.Agent/Execution/ManagedWorkspaceMaterializer.cs b/Engine/Source/Programs/Horde/Horde.Agent/Execution/ManagedWorkspaceMaterializer.cs
--- a/Engine/Source/Programs/Horde/Horde.Agent/Execution/ManagedWorkspaceMaterializer.cs
+++ b/Engine/Source/Programs/Horde/Horde.Agent/Execution/ManagedWorkspaceMaterializer.cs
@@ -1,8 +1,10 @@
 // Copyright Epic Games, Inc. All Rights Reserved.
 
+using System.Text;
 using System.Threading;
 using System.Threading.Tasks;
 using EpicGames.Core;
+using EpicGames.Perforce;
 using Horde.Agent.Utility;
 using HordeCommon.Rpc.Messages;
 using Microsoft.Extensions.Logging;
@@ -90,6 +92,16 @@ public class ManagedWorkspaceMaterializer : IWorkspaceMaterializer
 		}
 	}
 
+	private FileReference GetSyncMarkerFile()
+	{
+		if (_workspace == null)
+		{
+			throw new WorkspaceMaterializationException("Workspace not initialized");
+		}
+
+		return FileReference.Combine(_workspace.MetadataDir, "Synced.txt");
+	}
+
 	private (FileReference file, string fileContent) GetSyncMarker(int changeNum)
 	{
 		if (_workspace == null)
@@ -97,7 +109,18 @@ public class ManagedWorkspaceMaterializer : IWorkspaceMaterializer
 			throw new WorkspaceMaterializationException("Workspace not initialized");
 		}
 
-		return (FileReference.Combine(_workspace.MetadataDir, "Synced.txt"), $"Synced to CL {changeNum}");
+		StringBuilder content = new StringBuilder();
+		content.AppendLine($"Synced to CL {changeNum}");
+		foreach (PerforceViewMapEntry streamViewEntry in _workspace.StreamView.Entries)
+		{
+			content.AppendLine($"StreamView: {streamViewEntry}");
+		}
+		foreach (string viewLine in _workspace.View)
+		{
+			content.AppendLine($"View: {viewLine}");
+		}
+
+		return (GetSyncMarkerFile(), content.ToString());
 	}
 	
 	/// <summary>
@@ -150,6 +173,8 @@ public class ManagedWorkspaceMaterializer : IWorkspaceMaterializer
 		}
 		
 		await _workspace.UnshelveAsync(changeNum, cancellationToken);
+
+		FileReference.Delete(GetSyncMarkerFile());
 	}
 
 	private IScope CreateTraceSpan(string operationName)
