Horde: Allow files with a *.global.json extension to use the global config schema.

diff --git a/Engine/Source/Programs/Horde/Horde.Server/Server/GlobalConfig.cs b/Engine/Source/Programs/Horde/Horde.Server/Server/GlobalConfig.cs
--- a/Engine/Source/Programs/Horde/Horde.Server/Server/GlobalConfig.cs
+++ b/Engine/Source/Programs/Horde/Horde.Server/Server/GlobalConfig.cs
@@ -99,7 +99,7 @@ namespace Horde.Server.Server
 	/// Global configuration
 	/// </summary>
 	[JsonSchema("https://unrealengine.com/horde/global")]
-	[JsonSchemaCatalog("Horde Globals", "Horde global configuration file", "globals.json")]
+	[JsonSchemaCatalog("Horde Globals", "Horde global configuration file", new[] { "globals.json", "*.global.json" })]
 	[ConfigIncludeRoot]
 	public class GlobalConfig : IAclScope
 	{
