Horde: Add /bundles routes for being able to access bundles directly.

diff --git a/Engine/Source/Programs/Horde/Horde.Server/Artifacts/ArtifactsController.cs b/Engine/Source/Programs/Horde/Horde.Server/Artifacts/ArtifactsController.cs
--- a/Engine/Source/Programs/Horde/Horde.Server/Artifacts/ArtifactsController.cs
+++ b/Engine/Source/Programs/Horde/Horde.Server/Artifacts/ArtifactsController.cs
@@ -164,7 +164,7 @@ namespace Horde.Server.Artifacts
 		}
 
 		/// <summary>
-		/// Retrieves blobs for a particular artifact
+		/// Retrieves bundles for a particular artifact
 		/// </summary>
 		/// <param name="id">Identifier of the artifact to retrieve</param>
 		/// <param name="locator">The blob locator</param>
@@ -174,6 +174,7 @@ namespace Horde.Server.Artifacts
 		/// <returns>Information about all the artifacts</returns>
 		[HttpGet]
 		[Route("/api/v2/artifacts/{id}/blobs/{*locator}")]
+		[Route("/api/v2/artifacts/{id}/bundles/{*locator}")]
 		public async Task<ActionResult<object>> ReadArtifactBlobAsync(ArtifactId id, BundleLocator locator, [FromQuery] int? offset = null, [FromQuery] int? length = null, CancellationToken cancellationToken = default)
 		{
 			IArtifact? artifact = await _artifactCollection.GetAsync(id, cancellationToken);
diff --git a/Engine/Source/Programs/Horde/Horde.Server/Storage/StorageController.cs b/Engine/Source/Programs/Horde/Horde.Server/Storage/StorageController.cs
--- a/Engine/Source/Programs/Horde/Horde.Server/Storage/StorageController.cs
+++ b/Engine/Source/Programs/Horde/Horde.Server/Storage/StorageController.cs
@@ -188,6 +188,7 @@ namespace Horde.Server.Storage
 		/// <param name="cancellationToken">Cancellation token for the operation</param>
 		[HttpPost]
 		[Route("/api/v1/storage/{namespaceId}/blobs")]
+		[Route("/api/v1/storage/{namespaceId}/bundles")]
 		public async Task<ActionResult<WriteBlobResponse>> WriteBlobAsync(NamespaceId namespaceId, IFormFile? file, [FromForm] string? prefix = default, CancellationToken cancellationToken = default)
 		{
 			NamespaceConfig? namespaceConfig;
@@ -253,6 +254,7 @@ namespace Horde.Server.Storage
 		/// <param name="cancellationToken">Cancellation token for the operation</param>
 		[HttpGet]
 		[Route("/api/v1/storage/{namespaceId}/blobs/{*locator}")]
+		[Route("/api/v1/storage/{namespaceId}/bundles/{*locator}")]
 		public async Task<ActionResult> ReadBlobAsync(NamespaceId namespaceId, BundleLocator locator, [FromQuery] int? offset = null, [FromQuery] int? length = null, CancellationToken cancellationToken = default)
 		{
 			NamespaceConfig? namespaceConfig;
@@ -476,7 +478,7 @@ namespace Horde.Server.Storage
 		/// <param name="cancellationToken">Cancellation token for the operation</param>
 		/// <returns></returns>
 		[HttpGet]
-		[Route("/api/v1/storage/{namespaceId}/bundles/{*locator}")]
+		[Route("/api/v1/storage/{namespaceId}/bundleinfo/{*locator}")]
 		public async Task<ActionResult<object>> GetBundleAsync(NamespaceId namespaceId, BundleLocator locator, [FromQuery(Name = "imports")] bool includeImports = false, [FromQuery(Name = "exports")] bool includeExports = true, [FromQuery(Name = "packets")] bool includePackets = false, CancellationToken cancellationToken = default)
 		{
 			NamespaceConfig? namespaceConfig;
