Horde: Set additional env vars supplied by materializer in WorkspaceExecutor

This allows ManagedWorkspaceMaterializer to set its Perforce-specific env vars as required for UAT/BuildGraph.

diff --git a/Engine/Source/Programs/Horde/Horde.Agent.Tests/Execution/FakeWorkspaceMaterializer.cs b/Engine/Source/Programs/Horde/Horde.Agent.Tests/Execution/FakeWorkspaceMaterializer.cs
--- a/Engine/Source/Programs/Horde/Horde.Agent.Tests/Execution/FakeWorkspaceMaterializer.cs
+++ b/Engine/Source/Programs/Horde/Horde.Agent.Tests/Execution/FakeWorkspaceMaterializer.cs
@@ -71,7 +71,7 @@ public class FakeWorkspaceMaterializer : IWorkspaceMaterializer
 			throw new WorkspaceMaterializationException("Cannot get settings before initialization");
 		}
 		
-		return Task.FromResult(new WorkspaceMaterializerSettings(_rootDir, "fakeWorkspaceIdentifier", "fakeWorkspaceStreamRoot", false));
+		return Task.FromResult(new WorkspaceMaterializerSettings(_rootDir, "fakeWorkspaceIdentifier", "fakeWorkspaceStreamRoot", new Dictionary<string, string>(), false));
 	}
 
 	/// <inheritdoc/>
diff --git a/Engine/Source/Programs/Horde/Horde.Agent/Execution/ManagedWorkspaceMaterializer.cs b/Engine/Source/Programs/Horde/Horde.Agent/Execution/ManagedWorkspaceMaterializer.cs
--- a/Engine/Source/Programs/Horde/Horde.Agent/Execution/ManagedWorkspaceMaterializer.cs
+++ b/Engine/Source/Programs/Horde/Horde.Agent/Execution/ManagedWorkspaceMaterializer.cs
@@ -1,5 +1,6 @@
 // Copyright Epic Games, Inc. All Rights Reserved.
 
+using System.Collections.Generic;
 using System.Text;
 using System.Threading;
 using System.Threading.Tasks;
@@ -67,9 +68,28 @@ public class ManagedWorkspaceMaterializer : IWorkspaceMaterializer
 	/// <inheritdoc/>
 	public Task<WorkspaceMaterializerSettings> GetSettingsAsync(CancellationToken cancellationToken)
 	{
+		if (_workspace == null)
+		{
+			throw new WorkspaceMaterializationException("Workspace not initialized");
+		}
+		
 		// ManagedWorkspace store synced files in a sub-directory from the top working dir.
 		DirectoryReference syncDir = DirectoryReference.Combine(_workingDir, _agentWorkspace.Identifier, "Sync");
-		return Task.FromResult(new WorkspaceMaterializerSettings(syncDir, _agentWorkspace.Identifier, _agentWorkspace.Stream, true));
+
+		// Variables expected to be set for UAT/BuildGraph when Perforce is enabled (-P4 flag is set) 
+		Dictionary<string, string> envVars = new()
+			{
+				["uebp_PORT"] = _workspace.ServerAndPort,
+				["uebp_USER"] = _workspace.UserName,
+				["uebp_CLIENT"] = _workspace.ClientName,
+				["uebp_CLIENT_ROOT"] = $"//{_workspace.ClientName}"
+			};
+
+		// Perforce-specific variables
+		envVars["P4USER"] = _workspace.UserName;
+		envVars["P4CLIENT"] = _workspace.ClientName;
+
+		return Task.FromResult(new WorkspaceMaterializerSettings(syncDir, _agentWorkspace.Identifier, _agentWorkspace.Stream, envVars, true));
 	}
 
 	/// <inheritdoc/>
diff --git a/Engine/Source/Programs/Horde/Horde.Agent/Execution/WorkspaceExecutor.cs b/Engine/Source/Programs/Horde/Horde.Agent/Execution/WorkspaceExecutor.cs
--- a/Engine/Source/Programs/Horde/Horde.Agent/Execution/WorkspaceExecutor.cs
+++ b/Engine/Source/Programs/Horde/Horde.Agent/Execution/WorkspaceExecutor.cs
@@ -93,6 +93,12 @@ namespace Horde.Agent.Execution
 			_envVars["uebp_CL"] = _batch.Change.ToString();
 			_envVars["uebp_CodeCL"] = _batch.CodeChange.ToString();
 
+			WorkspaceMaterializerSettings settings = await _workspace.GetSettingsAsync(cancellationToken);
+			foreach ((string key, string value) in settings.EnvironmentVariables)
+			{
+				_envVars[key] = value;
+			}
+
 			if (autoSdkWorkspaceSettings != null)
 			{
 				_envVars["UE_SDKS_ROOT"] = autoSdkWorkspaceSettings.DirectoryPath.FullName;
diff --git a/Engine/Source/Programs/Horde/Horde.Agent/Execution/WorkspaceMaterializer.cs b/Engine/Source/Programs/Horde/Horde.Agent/Execution/WorkspaceMaterializer.cs
--- a/Engine/Source/Programs/Horde/Horde.Agent/Execution/WorkspaceMaterializer.cs
+++ b/Engine/Source/Programs/Horde/Horde.Agent/Execution/WorkspaceMaterializer.cs
@@ -1,6 +1,7 @@
 // Copyright Epic Games, Inc. All Rights Reserved.
 
 using System;
+using System.Collections.Generic;
 using System.Threading;
 using System.Threading.Tasks;
 using EpicGames.Core;
@@ -52,6 +53,12 @@ public class WorkspaceMaterializerSettings
 	/// </summary>
 	public string StreamRoot { get; }
 	
+	/// <summary>
+	/// Environment variables expected to be set for applications executing inside the workspace
+	/// Mostly intended for Perforce-specific variables when <see cref="IsPerforceWorkspace" /> is set to true
+	/// </summary>
+	public IReadOnlyDictionary<string, string> EnvironmentVariables { get; }
+	
 	/// <summary>
 	/// Whether the materialized workspace is a true Perforce workspace
 	/// This flag is provided as a stop-gap solution to allow replacing ManagedWorkspace with WorkspaceMaterializer.
@@ -65,12 +72,14 @@ public class WorkspaceMaterializerSettings
 	/// <param name="directoryPath"></param>
 	/// <param name="identifier"></param>
 	/// <param name="streamRoot"></param>
+	/// <param name="envVars"></param>
 	/// <param name="isPerforceWorkspace"></param>
-	public WorkspaceMaterializerSettings(DirectoryReference directoryPath, string identifier, string streamRoot, bool isPerforceWorkspace)
+	public WorkspaceMaterializerSettings(DirectoryReference directoryPath, string identifier, string streamRoot, IReadOnlyDictionary<string, string> envVars, bool isPerforceWorkspace)
 	{
 		DirectoryPath = directoryPath;
 		Identifier = identifier;
 		StreamRoot = streamRoot;
+		EnvironmentVariables = envVars;
 		IsPerforceWorkspace = isPerforceWorkspace;
 	}
 }
