Horde: Add confirmation for deleting and creating new pools, make it possible to sort agents in edit modal by name or by assigned agent count
#rnx
#skipci

diff --git a/Engine/Source/Programs/Horde/HordeDashboard/src/components/AgentView.tsx b/Engine/Source/Programs/Horde/HordeDashboard/src/components/AgentView.tsx
--- a/Engine/Source/Programs/Horde/HordeDashboard/src/components/AgentView.tsx
+++ b/Engine/Source/Programs/Horde/HordeDashboard/src/components/AgentView.tsx
@@ -1,4 +1,4 @@
-// Copyright Epic Games, Inc. All Rights Reserved.  
+// Copyright Epic Games, Inc. All Rights Reserved.
 import { Checkbox, CommandButton, ConstrainMode, ContextualMenu, DefaultButton, DetailsHeader, DetailsList, DetailsListLayoutMode, Dialog, DialogType, DirectionalHint, Dropdown, FontSizes, FontWeights, getTheme, IBasePickerProps, IColumn, Icon, IconButton, IContextualMenuItem, IContextualMenuProps, IDetailsHeaderProps, IDetailsHeaderStyles, IDetailsListProps, ITag, ITagItemStyles, ITooltipHostStyles, Link as ReactLink, mergeStyles, mergeStyleSets, PrimaryButton, ProgressIndicator, ScrollablePane, ScrollbarVisibility, SearchBox, Selection, SelectionMode, Slider, Spinner, SpinnerSize, Stack, Sticky, StickyPositionType, TagItem, TagPicker, Text, TextField } from '@fluentui/react';
 import { action, makeObservable, observable } from 'mobx';
 import { observer } from 'mobx-react-lite';
@@ -752,6 +752,8 @@ class EditPoolsModalState {
    @observable modifiedPools: PoolEditorItem[] = [];
    @observable selectedColor = "";
 
+   @observable isConfirmationOpen = false;
+
    // last selected color of the color modal
    @observable lastSelectedPool: PoolEditorItem | undefined = undefined;
    isDirectEdit = false;
@@ -786,6 +788,11 @@ class EditPoolsModalState {
       };
    }
 
+   @action
+   showConfirmation(show: boolean) {
+      this.isConfirmationOpen = show;
+   }
+
    // sets the pool editor dialog open
    @action
    setOpen(isEditorOpen?: boolean, selectedItem?: PoolData) {
@@ -884,6 +891,25 @@ class EditPoolsModalState {
       this.isPoolValueValid = value;
    }
 
+   getPoolModifications(): { deletedPools: string[], newPools: string[] } {
+
+
+      const value = { deletedPools: [] as string[], newPools: [] as string[] };
+
+      this.modifiedPools.forEach(item => {
+         if (item.deleted) {
+            value.deletedPools.push(item.pool.name);
+         }
+
+         if (item.key.indexOf(this.newPoolInputId) !== -1) {
+            value.newPools.push(item.pool.name);
+         }
+
+      })
+
+      return value;
+   }
+
    @action
    async saveChanges() {
       // save off pool changes
@@ -1360,9 +1386,9 @@ export const AgentMenuBar: React.FC<{ agentView?: boolean }> = observer(({ agent
             </Stack>
 
          </Stack.Item>
-         {!!agentView && 
+         {!!agentView &&
             <Stack horizontal tokens={{ childrenGap: 12 }} grow>
-               <Stack grow/>
+               <Stack grow />
                <PrimaryButton styles={{ root: { fontFamily: "Horde Open Sans SemiBold !important" } }} text="Download Agent" onClick={() => { backend.downloadAgentZip() }} />
                <CommandButton
                   onClick={() => { editPoolsModalState.setOpen(); }}
@@ -1379,7 +1405,62 @@ export const AgentMenuBar: React.FC<{ agentView?: boolean }> = observer(({ agent
    );
 });
 
+const PoolEditorConfirmation: React.FC = observer(() => {
+
+
+   if (!editPoolsModalState.isConfirmationOpen) {
+      return null;
+   }
+
+   const mods = editPoolsModalState.getPoolModifications();
+
+   if (mods.deletedPools.length === 0 && mods.newPools.length === 0) {
+      return null;
+   }
+
+   let title = "";
+   let subText = "";
+
+   if (mods.deletedPools.length && !mods.newPools.length) {
+      title = "Delete Pool";
+      if (mods.deletedPools.length > 1) {
+         title += "s";
+      }
+   } else if (!mods.deletedPools.length && mods.newPools.length) {
+      title = "Create Pool";
+      if (mods.newPools.length > 1) {
+         title += "s";
+      }
+   } else {
+      title = "Create and Delete Pools";
+   }
+
+   title += "?";
+
+   if (mods.deletedPools.length) {
+      subText += "Delete: " + mods.deletedPools.join(", ") + " ";
+   }
+
+   if (mods.newPools.length) {
+      subText += "Create: " + mods.newPools.join(", ");
+   }
+
+   return <ConfirmationDialog
+      title={title}
+      subText={subText}
+      isOpen={true}
+      confirmText="Save"
+      cancelText="Cancel"
+      onConfirm={() => { editPoolsModalState.showConfirmation(false); editPoolsModalState.saveChanges(); }}
+      onCancel={() => editPoolsModalState.showConfirmation(false)}
+   />
+
+});
+
 export const PoolEditorModal: React.FC = observer(() => {
+
+   const [state, setState] = useState<{ sortBy: "Agents" | "Name" }>({ sortBy: "Name" });
+
    let selectedColor = "#ffffff";
    let selectedColorValue = "0";
    if (editPoolsModalState.lastSelectedPool) {
@@ -1428,7 +1509,18 @@ export const PoolEditorModal: React.FC = observer(() => {
       if (props) {
          return (
             <Sticky stickyPosition={StickyPositionType.Header} isScrollSynced={true}>
-               <DetailsHeader {...props} styles={customStyles} onRenderColumnHeaderTooltip={onRenderColumnHeaderTooltip} />
+               <DetailsHeader {...props} styles={customStyles} onColumnClick={(ev: React.MouseEvent<HTMLElement>, column: IColumn) => {
+                  if (column.name === "Agents") {
+                     if (state.sortBy !== 'Agents') {
+                        setState({ ...state, sortBy: "Agents" });
+                     }
+
+                  } else if (column.name === "Name") {
+                     if (state.sortBy !== 'Name') {
+                        setState({ ...state, sortBy: "Name" });
+                     }
+                  }
+               }} onRenderColumnHeaderTooltip={onRenderColumnHeaderTooltip} />
             </Sticky>
          );
       }
@@ -1452,7 +1544,12 @@ export const PoolEditorModal: React.FC = observer(() => {
       return null;
    };
 
-   const pools = editPoolsModalState.modifiedPools.filter(item => item.deleted === false).sort((a, b) => { return a.sortKey.localeCompare(b.sortKey); });
+   const pools = editPoolsModalState.modifiedPools.filter(item => item.deleted === false).sort((a, b) => {
+      if (state.sortBy === "Name")
+         return a.sortKey.localeCompare(b.sortKey);
+      else
+         return b.numAgentsAssigned - a.numAgentsAssigned;
+   });
 
    // 75vh for max
    const viewportHeight = document.documentElement.clientHeight * .75;
@@ -1460,6 +1557,7 @@ export const PoolEditorModal: React.FC = observer(() => {
 
    return (
       <Stack>
+         <PoolEditorConfirmation />
          <Dialog
             modalProps={{
                isBlocking: false,
@@ -1504,7 +1602,7 @@ export const PoolEditorModal: React.FC = observer(() => {
                      </Stack>
                      <Stack grow />
                      <Stack>
-                        <PrimaryButton onClick={() => editPoolsModalState.saveChanges()} styles={{ root: { marginRight: "10px" } }}>Save</PrimaryButton>
+                        <PrimaryButton disabled={editPoolsModalState.getPoolModifications().deletedPools.length === 0 && editPoolsModalState.getPoolModifications().newPools.length === 0} onClick={() => editPoolsModalState.showConfirmation(true)} styles={{ root: { marginRight: "10px" } }}>Save</PrimaryButton>
                      </Stack>
                      <Stack>
                         <DefaultButton onClick={() => editPoolsModalState.setClose()} styles={{ root: { marginRight: "10px" } }}>Cancel</DefaultButton>
@@ -1794,7 +1892,7 @@ export const AgentViewInner: React.FC<{ agentId?: string, poolId?: string, searc
           // no other way to filter children being centered other than to drill into private members??
           if (props.children._owner.child?.child.child.child.elementType === "span") {
               const data = props.children._owner.child.child.child.child.child.child.stateNode.data;
-              // if this cluster happens to be true, reset back to the default, because doing this the other way around 
+              // if this cluster happens to be true, reset back to the default, because doing this the other way around
               // takes too long when the columns update on details switch
               // ugh :(
               if (data === "Name" || data === "Pools") {
@@ -2140,7 +2238,7 @@ export const AgentViewInner: React.FC<{ agentId?: string, poolId?: string, searc
                onItemClick={() => localState.setAgentContextMenuOpen(false)}
                onDismiss={() => localState.setAgentContextMenuOpen(false)}
                isBeakVisible={true}
-               target={{x: localState.mouseX, y: localState.mouseY }}
+               target={{ x: localState.mouseX, y: localState.mouseY }}
                hidden={!localState.agentContextMenuOpen}
                directionalHint={DirectionalHint.bottomLeftEdge}
                directionalHintFixed={true}
