Horde: Select template instead of redundant category in stream all jobs view
#rnx
#skipci

diff --git a/Engine/Source/Programs/Horde/HordeDashboard/src/components/JobFilterBar.tsx b/Engine/Source/Programs/Horde/HordeDashboard/src/components/JobFilterBar.tsx
--- a/Engine/Source/Programs/Horde/HordeDashboard/src/components/JobFilterBar.tsx
+++ b/Engine/Source/Programs/Horde/HordeDashboard/src/components/JobFilterBar.tsx
@@ -1,10 +1,10 @@
 // Copyright Epic Games, Inc. All Rights Reserved.
 
 import { action, makeObservable, observable } from "mobx";
-import { Dropdown, IDropdownOption, Stack } from "@fluentui/react";
+import { DefaultButton, Dropdown, IContextualMenuItem, IContextualMenuProps, IDropdownOption, Label, Stack } from "@fluentui/react";
 import React, { useEffect, useState } from "react";
 import { useBackend } from "../backend";
-import { GetJobsTabResponse, TabType, GetTemplateRefResponse } from "../backend/Api";
+import { GetJobsTabResponse, TabType, GetTemplateRefResponse, StreamData } from "../backend/Api";
 import { FilterStatus } from "../backend/JobHandler";
 import templateCache from '../backend/TemplateCache';
 import { hordeClasses, modeColors } from "../styles/Styles";
@@ -14,25 +14,25 @@ import { hordeClasses, modeColors } from "../styles/Styles";
 
 const dropDownStyle: any = () => {
 
-    return {
-        dropdown: {},
-        callout: {
-            selectors: {
-                ".ms-Callout-main": {
-                    padding: "4px 4px 12px 12px",
-                    overflow: "hidden"
-                }
+   return {
+      dropdown: {},
+      callout: {
+         selectors: {
+            ".ms-Callout-main": {
+               padding: "4px 4px 12px 12px",
+               overflow: "hidden"
             }
-        },
-        dropdownItemHeader: { fontSize: 12, color: modeColors.text },
-        dropdownOptionText: { fontSize: 12 },
-        dropdownItem: {
-            minHeight: 28, lineHeight: 28
-        },
-        dropdownItemSelected: {
-            minHeight: 28, lineHeight: 28, backgroundColor: "inherit"
-        }
-    }
+         }
+      },
+      dropdownItemHeader: { fontSize: 12, color: modeColors.text },
+      dropdownOptionText: { fontSize: 12 },
+      dropdownItem: {
+         minHeight: 28, lineHeight: 28
+      },
+      dropdownItemSelected: {
+         minHeight: 28, lineHeight: 28, backgroundColor: "inherit"
+      }
+   }
 }
 
 
@@ -42,218 +42,203 @@ export class JobFilter {
       makeObservable(this);
    }
 
-    @action
-    setUpdated() {
-        this.updated++;
-    }
+   @action
+   setUpdated() {
+      this.updated++;
+   }
 
-    set(templates?: GetTemplateRefResponse[], status?: FilterStatus[]) {
+   set(templates?: GetTemplateRefResponse[], status?: FilterStatus[]) {
 
-        let templatesDirty = true;
+      let templatesDirty = true;
 
-        if (!templates?.length && !this.templates?.length) {
-            templatesDirty = false;
-        }
+      if (!templates?.length && !this.templates?.length) {
+         templatesDirty = false;
+      }
 
-        if (templates && this.templates) {
+      if (templates && this.templates) {
 
-            if (templates.length === this.templates.length && templates.every((val, index) => val === this.templates![index])) {
-                templatesDirty = false;
-            }
+         if (templates.length === this.templates.length && templates.every((val, index) => val === this.templates![index])) {
+            templatesDirty = false;
+         }
 
-        }
+      }
 
-        if (templatesDirty) {
-            this.templates = templates;
-        }
+      if (templatesDirty) {
+         this.templates = templates;
+      }
 
-        let statusDirty = true;
+      let statusDirty = true;
 
-        if (!status?.length && !this.status?.length) {
-            statusDirty = false;
-        }
+      if (!status?.length && !this.status?.length) {
+         statusDirty = false;
+      }
 
-        if (status && this.status) {
+      if (status && this.status) {
 
-            if (status.length === this.status.length && status.every((val, index) => val === this.status![index])) {
-                statusDirty = false;
-            }
+         if (status.length === this.status.length && status.every((val, index) => val === this.status![index])) {
+            statusDirty = false;
+         }
 
-        }
+      }
 
-        if (statusDirty) {
-            this.status = status;
-        }
+      if (statusDirty) {
+         this.status = status;
+      }
 
-        // check any dirty
-        if (templatesDirty || statusDirty) {
-            this.setUpdated();
-        }
-    }
+      // check any dirty
+      if (templatesDirty || statusDirty) {
+         this.setUpdated();
+      }
+   }
 
-    @observable
-    updated: number = 0;
+   @observable
+   updated: number = 0;
 
-    templates?: GetTemplateRefResponse[];
-    status?: FilterStatus[];
+   templates?: GetTemplateRefResponse[];
+   status?: FilterStatus[];
 }
 
 export const jobFilter = new JobFilter();
 
-interface CategoryItem extends IDropdownOption {
+const TemplateSelector: React.FC<{ stream: StreamData, templates: GetTemplateRefResponse[], selected: GetTemplateRefResponse | undefined, setSelectedTemplate: (template: GetTemplateRefResponse | undefined) => void }> = ({ stream, templates, selected, setSelectedTemplate }) => {
+
+   let templateOptions: IContextualMenuItem[] = [];
+
+   const sorted = new Map<string, GetTemplateRefResponse[]>();
+
+   templates.forEach(t => {
+
+      stream.tabs.forEach(tab => {
+         if (tab.type !== TabType.Jobs) {
+            return;
+         }
+
+         const jtab = tab as GetJobsTabResponse;
+         if (!jtab.templates?.find(template => template === t.id)) {
+            return;
+         }
+
+         if (!sorted.has(jtab.title)) {
+            sorted.set(jtab.title, []);
+         }
+
+         sorted.get(jtab.title)!.push(t);
+
+      })
+   })
+
+   Array.from(sorted.keys()).sort((a, b) => a < b ? -1 : 1).forEach(cat => {
+
+      const templates = sorted.get(cat);
+      if (!templates?.length) {
+         return;
+      }
+
+      const subItems = templates.sort((a, b) => a.name < b.name ? -1 : 1).map(t => {
+         return { key: t.id, text: t.name, onClick: () => { setSelectedTemplate(t) } };
+      })
+
+      templateOptions.push({ key: `${cat}_category`, text: cat, subMenuProps: { items: subItems } });
+
+   })
+
+   templateOptions.push({ key: `show_all_templates`, text: "All Templates", onClick: (ev) => { ev?.stopPropagation(); ev?.preventDefault(); setSelectedTemplate(undefined) } });
+
+   const templateMenuProps: IContextualMenuProps = {
+      shouldFocusOnMount: true,
+      subMenuHoverDelay: 0,
+      items: templateOptions,
+   };
+
+   return <Stack style={{ paddingTop: 4 }}>
+      <Label>Template</Label>
+      <DefaultButton style={{ width: 352, textAlign: "left", height: 30 }} menuProps={templateMenuProps} text={selected?.name ?? "All"} />
+   </Stack>
 
-    templates: GetTemplateRefResponse[];
 }
 
 // Job filter bar for "all" jobs view
 export const JobFilterBar: React.FC<{ streamId: string }> = ({ streamId }) => {
 
-    const { projectStore } = useBackend();
+   const { projectStore } = useBackend();
 
-    const [state, setState] = useState<{ streamId?: string; templates?: GetTemplateRefResponse[], categories?: string[], status?: FilterStatus | "All" | undefined }>({});
+   const [state, setState] = useState<{ streamId?: string; templates?: GetTemplateRefResponse[], selectedTemplate?: GetTemplateRefResponse, status?: FilterStatus | "All" | undefined }>({});
 
-    const stream = projectStore.streamById(streamId);
+   const stream = projectStore.streamById(streamId);
 
-    useEffect(() => {
+   useEffect(() => {
 
-        return () => {
-            jobFilter.set(undefined, undefined);
-        };
+      return () => {
+         jobFilter.set(undefined, undefined);
+      };
 
-    }, []);
+   }, []);
 
 
-    if (!stream) {
-        console.error("unable to get stream");
-        return <div>unable to get stream</div>;
-    }
+   if (!stream) {
+      console.error("unable to get stream");
+      return <div>unable to get stream</div>;
+   }
 
-    if (!state.templates || streamId !== state.streamId) {
-        templateCache.getStreamTemplates(stream).then(data => {
-            setState({ streamId: streamId, templates: data, categories: ["All"], status: "All" });
-            jobFilter.set(data);
-        });
-        return null;
-    } 
+   if (!state.templates || streamId !== state.streamId) {
+      templateCache.getStreamTemplates(stream).then(data => {
+         setState({ streamId: streamId, templates: data, status: "All" });
+         jobFilter.set(data);
+      });
+      return null;
+   }
 
-    let templates: GetTemplateRefResponse[] = state.templates.map(t => t);
+   let templates: GetTemplateRefResponse[] = state.templates.map(t => t);
 
-    if (!templates.length) {
-        return null;
-    }
+   if (!templates.length) {
+      return null;
+   }
 
-    const catMap = new Map<string, GetTemplateRefResponse[]>();
+   const statusItems = ["Running", "Complete", "Succeeded", "Failed", "Waiting", "All"].map(status => {
+      return {
+         key: status,
+         text: status,
+         status: status
+      };
+   });
+
+   let filterTemplates: GetTemplateRefResponse[] = [];
+   if (state.selectedTemplate) {
+      filterTemplates = [state.selectedTemplate];
+   }
 
-    stream.tabs?.forEach(t => {
+   if (!jobFilter?.status?.find(s => s === state.status)) {
+      if ((!state.status || state.status === "All") && jobFilter.status?.length) {
+         jobFilter.set(jobFilter.templates, undefined);
+         return null;
+      }
 
-        if (t.type !== TabType.Jobs) {
-            return;
-        }
-
-        let tab = t as GetJobsTabResponse;
+      if (state.status && state.status !== "All") {
+         jobFilter.set(jobFilter.templates, [state.status]);
+         return null;
+      }
+   }
 
-        const ctemps = tab.templates?.map(tid => templates.find(t => t.id === tid)).filter(t => !!t) as GetTemplateRefResponse[];
+   return <Stack horizontal tokens={{ childrenGap: 24 }} className={hordeClasses.modal}>
+      <TemplateSelector stream={stream} templates={templates} selected={state.selectedTemplate} setSelectedTemplate={(t) => {
+         jobFilter.set(t ? [t] : [...templates], (state.status === "All" || !state.status) ? undefined : [state.status]);
+         setState({ ...state, selectedTemplate: t });         
+      }} />
+      <Dropdown
+         style={{ width: 120 }}
+         styles={dropDownStyle}
+         label="Status"
+         selectedKey={state.status}
+         options={statusItems}
+         onDismiss={() => {
+            jobFilter.set(filterTemplates.length ? filterTemplates : undefined, (state.status === "All" || !state.status) ? undefined : [state.status]);
+         }}
+         onChange={(event, option, index) => {
+            if (option) {
+               setState({ streamId: streamId, templates: state.templates, selectedTemplate: state.selectedTemplate, status: option.key as FilterStatus });
+            }
+         }}
+      />
 
-        if (!ctemps?.length) {
-            return;
-        }
-
-        catMap.set(t.title, ctemps)
-
-    })
-
-    const categoryItems: CategoryItem[] = Array.from(catMap.keys()).sort((a, b) => a < b ? -1 : 1).map(name => {
-        return {
-            key: name,
-            text: name,
-            templates: catMap.get(name)!
-        }
-    });
-
-    catMap.set("All", templates);
-    categoryItems.push({
-        key: "All",
-        text: "All",
-        templates: catMap.get("All")!
-    });
-
-    const statusItems = ["Running", "Complete", "Succeeded", "Failed", "Waiting", "All"].map(status => {
-        return {
-            key: status,
-            text: status,
-            status: status
-        };
-    });
-
-    // update filter state
-    const categories = state.categories ?? [];
-    const filterTemplates = categories.map(c => catMap.get(c)!).flat() ?? [];
-
-    if (!jobFilter?.status?.find( s => s === state.status))
-    {
-        if ((!state.status || state.status === "All") && jobFilter.status?.length) {
-            jobFilter.set(jobFilter.templates, undefined);
-            return null;        
-        }
-
-        if (state.status && state.status !== "All") {
-            jobFilter.set(jobFilter.templates, [state.status]);
-            return null;        
-        }
-    }
-
-    
-    return <Stack horizontal tokens={{ childrenGap: 24 }} className={hordeClasses.modal}>
-        <Dropdown                        
-            style={{ width: 200 }}
-            styles={dropDownStyle}
-            label="Category"
-            selectedKeys={categories}
-            multiSelect
-            options={categoryItems}
-            onDismiss={()=> {
-                jobFilter.set(filterTemplates.length ? filterTemplates : undefined, (state.status === "All" || !state.status) ? undefined : [state.status]);
-            }}
-            onChange={(event, option, index) => {
-
-                if (option) {
-
-                    let cats = [...categories];
-                    if (!option.selected) {
-                        cats = cats.filter(k => k !== option.key);
-                    } else {
-                        if (cats.indexOf(option.key as string) === -1) {
-                            cats.push(option.key as string);
-                        }
-                    }
-
-                    if (!cats.length || (option.selected && option.key === "All")) {
-                        cats = ["All"];
-                    }
-
-                    if (cats.find(k => k === "All") && cats.length > 1) {
-                        cats = cats.filter(k => k !== "All");
-                    }                    
-
-                    setState({ streamId: streamId, templates: state.templates, categories: cats, status: state.status });
-                }
-            }}
-        />
-        <Dropdown
-            style={{ width: 120 }}
-            styles={dropDownStyle}
-            label="Status"
-            selectedKey={state.status}
-            options={statusItems}
-            onDismiss={() => {
-                jobFilter.set(filterTemplates.length ? filterTemplates : undefined, (state.status === "All" || !state.status) ? undefined : [state.status]);
-            }}
-            onChange={(event, option, index) => {
-                if (option) {
-                    setState({ streamId: streamId, templates: state.templates, categories: state.categories, status: option.key as FilterStatus });                    
-                }
-            }}
-        />
-
-    </Stack>;
+   </Stack>;
 };
\ No newline at end of file
