Horde: Fix issue with modified pool saving
#rnx
#skipci

diff --git a/Engine/Source/Programs/Horde/HordeDashboard/src/components/AgentView.tsx b/Engine/Source/Programs/Horde/HordeDashboard/src/components/AgentView.tsx
--- a/Engine/Source/Programs/Horde/HordeDashboard/src/components/AgentView.tsx
+++ b/Engine/Source/Programs/Horde/HordeDashboard/src/components/AgentView.tsx
@@ -891,18 +891,24 @@ class EditPoolsModalState {
       this.isPoolValueValid = value;
    }
 
-   getPoolModifications(): { deletedPools: string[], newPools: string[] } {
+   getPoolModifications(): { deletedPools: string[], newPools: string[], modifiedPools: string[] } {
 
 
-      const value = { deletedPools: [] as string[], newPools: [] as string[] };
+      const value = { deletedPools: [] as string[], newPools: [] as string[], modifiedPools: [] as string[] };
 
       this.modifiedPools.forEach(item => {
          if (item.deleted) {
             value.deletedPools.push(item.pool.name);
          }
-
-         if (item.key.indexOf(this.newPoolInputId) !== -1) {
+         else if (item.key.indexOf(this.newPoolInputId) !== -1) {
             value.newPools.push(item.pool.name);
+         } else {
+            const pool = agentStore.pools.find(pool => { return pool.id === item.pool.id; })!;
+            if (pool) {
+               if (pool.name !== item.pool.name || pool.properties!["Color"] !== item.pool.properties!["Color"]) {
+                  value.modifiedPools.push(item.pool.name);
+               }   
+            }
          }
 
       })
@@ -1414,11 +1420,11 @@ const PoolEditorConfirmation: React.FC = observer(() => {
 
    const mods = editPoolsModalState.getPoolModifications();
 
-   if (mods.deletedPools.length === 0 && mods.newPools.length === 0) {
+   if (mods.deletedPools.length === 0 && mods.newPools.length === 0 && mods.modifiedPools.length === 0) {
       return null;
    }
 
-   let title = "";
+   let title = "Modify Pool" + (mods.modifiedPools.length > 1 ? "s" : "");
    let subText = "";
 
    if (mods.deletedPools.length && !mods.newPools.length) {
@@ -1431,9 +1437,7 @@ const PoolEditorConfirmation: React.FC = observer(() => {
       if (mods.newPools.length > 1) {
          title += "s";
       }
-   } else {
-      title = "Create and Delete Pools";
-   }
+   } 
 
    title += "?";
 
@@ -1442,7 +1446,11 @@ const PoolEditorConfirmation: React.FC = observer(() => {
    }
 
    if (mods.newPools.length) {
-      subText += "Create: " + mods.newPools.join(", ");
+      subText += "Create: " + mods.newPools.join(", ") + " ";
+   }
+
+   if (mods.modifiedPools.length) {
+      subText += "Modified: " + mods.modifiedPools.join(", ");
    }
 
    return <ConfirmationDialog
@@ -1602,7 +1610,7 @@ export const PoolEditorModal: React.FC = observer(() => {
                      </Stack>
                      <Stack grow />
                      <Stack>
-                        <PrimaryButton disabled={editPoolsModalState.getPoolModifications().deletedPools.length === 0 && editPoolsModalState.getPoolModifications().newPools.length === 0} onClick={() => editPoolsModalState.showConfirmation(true)} styles={{ root: { marginRight: "10px" } }}>Save</PrimaryButton>
+                        <PrimaryButton disabled={editPoolsModalState.getPoolModifications().deletedPools.length === 0 && editPoolsModalState.getPoolModifications().newPools.length === 0 && editPoolsModalState.getPoolModifications().modifiedPools.length === 0} onClick={() => editPoolsModalState.showConfirmation(true)} styles={{ root: { marginRight: "10px" } }}>Save</PrimaryButton>
                      </Stack>
                      <Stack>
                         <DefaultButton onClick={() => editPoolsModalState.setClose()} styles={{ root: { marginRight: "10px" } }}>Cancel</DefaultButton>
