Horde: Add support for child leases, which are terminated when their parent completes.

diff --git a/Engine/Source/Programs/Horde/Horde.Server.Tests/Fleet/AwsAutoScalingLifecycleTest.cs b/Engine/Source/Programs/Horde/Horde.Server.Tests/Fleet/AwsAutoScalingLifecycleTest.cs
--- a/Engine/Source/Programs/Horde/Horde.Server.Tests/Fleet/AwsAutoScalingLifecycleTest.cs
+++ b/Engine/Source/Programs/Horde/Horde.Server.Tests/Fleet/AwsAutoScalingLifecycleTest.cs
@@ -170,7 +170,7 @@ public class AwsAutoScalingLifecycleServiceTest : TestSetup
 	public async Task GetInstancesAvailableForTermination_AgentRunningJob_ReturnsNoInstanceId()
 	{
 		// Arrange
-		AgentLease lease = new(new LeaseId(ObjectId.GenerateNewId()), "test-lease", null, null, null, LeaseState.Active, null, false, null);
+		AgentLease lease = new(new LeaseId(ObjectId.GenerateNewId()), null, "test-lease", null, null, null, LeaseState.Active, null, false, null);
 		IAgent agent = await CreateAgentAsync(new PoolId("pool1"), awsInstanceId: "i-1000", lease: lease);
 		TerminationPolicyEvent e = CreateTerminationPolicyEvent("i-1000");
 
diff --git a/Engine/Source/Programs/Horde/Horde.Server.Tests/Fleet/LeaseUtilizationStrategyTest.cs b/Engine/Source/Programs/Horde/Horde.Server.Tests/Fleet/LeaseUtilizationStrategyTest.cs
--- a/Engine/Source/Programs/Horde/Horde.Server.Tests/Fleet/LeaseUtilizationStrategyTest.cs
+++ b/Engine/Source/Programs/Horde/Horde.Server.Tests/Fleet/LeaseUtilizationStrategyTest.cs
@@ -108,7 +108,7 @@ namespace Horde.Server.Tests.Fleet
 			placeholderJobTask.JobName = "placeholderJobName";
 			byte[] payload = Any.Pack(placeholderJobTask).ToByteArray();
 
-			ILease lease = await LeaseCollection.AddAsync(LeaseId.GenerateNewId(), "placeholderLease", agent.Id, agent.SessionId!.Value, new StreamId("placeholderStream"), pool.Id, null, startTime, payload);
+			ILease lease = await LeaseCollection.AddAsync(LeaseId.GenerateNewId(), null, "placeholderLease", agent.Id, agent.SessionId!.Value, new StreamId("placeholderStream"), pool.Id, null, startTime, payload);
 			bool wasModified = await LeaseCollection.TrySetOutcomeAsync(lease.Id, startTime + duration, LeaseOutcome.Success, null);
 			Assert.IsTrue(wasModified);
 			return lease;
diff --git a/Engine/Source/Programs/Horde/Horde.Server.Tests/LeasesControllerDbTest.cs b/Engine/Source/Programs/Horde/Horde.Server.Tests/LeasesControllerDbTest.cs
--- a/Engine/Source/Programs/Horde/Horde.Server.Tests/LeasesControllerDbTest.cs
+++ b/Engine/Source/Programs/Horde/Horde.Server.Tests/LeasesControllerDbTest.cs
@@ -39,7 +39,7 @@ namespace Horde.Server.Tests
 		{
 			LeaseId id = LeaseId.GenerateNewId();
 			SessionId sessionId = SessionId.GenerateNewId();
-			ILease lease = await LeaseCollection.AddAsync(id, "myLease", new AgentId("agent-1"), sessionId, null, null, null, startTime, Array.Empty<byte>());
+			ILease lease = await LeaseCollection.AddAsync(id, null, "myLease", new AgentId("agent-1"), sessionId, null, null, null, startTime, Array.Empty<byte>());
 			await LeaseCollection.TrySetOutcomeAsync(id, startTime + duration, LeaseOutcome.Success, null);
 			return lease;
 		}
diff --git a/Engine/Source/Programs/Horde/Horde.Server/Acls/AclService.cs b/Engine/Source/Programs/Horde/Horde.Server/Acls/AclService.cs
--- a/Engine/Source/Programs/Horde/Horde.Server/Acls/AclService.cs
+++ b/Engine/Source/Programs/Horde/Horde.Server/Acls/AclService.cs
@@ -118,6 +118,19 @@ namespace Horde.Server.Acls
 			return user.HasClaim(HordeClaimTypes.AgentSessionId, sessionId.ToString());
 		}
 
+		public static LeaseId? GetLeaseClaim(this ClaimsPrincipal user)
+		{
+			Claim? claim = user.FindFirst(HordeClaimTypes.Lease);
+			if (claim == null || !LeaseId.TryParse(claim.Value, out LeaseId leaseIdValue))
+			{
+				return null;
+			}
+			else
+			{
+				return leaseIdValue;
+			}
+		}
+
 		public static SessionId? GetSessionClaim(this ClaimsPrincipal user)
 		{
 			Claim? claim = user.FindFirst(HordeClaimTypes.AgentSessionId);
diff --git a/Engine/Source/Programs/Horde/Horde.Server/Agents/AgentService.cs b/Engine/Source/Programs/Horde/Horde.Server/Agents/AgentService.cs
--- a/Engine/Source/Programs/Horde/Horde.Server/Agents/AgentService.cs
+++ b/Engine/Source/Programs/Horde/Horde.Server/Agents/AgentService.cs
@@ -800,7 +800,7 @@ namespace Horde.Server.Agents
 		{
 			try
 			{
-				return _leases.AddAsync(agentLease.Id, agentLease.Name, agent.Id, agent.SessionId!.Value, agentLease.StreamId, agentLease.PoolId, agentLease.LogId, agentLease.StartTime, agentLease.Payload!);
+				return _leases.AddAsync(agentLease.Id, agentLease.ParentId, agentLease.Name, agent.Id, agent.SessionId!.Value, agentLease.StreamId, agentLease.PoolId, agentLease.LogId, agentLease.StartTime, agentLease.Payload!);
 			}
 			catch (Exception ex)
 			{
@@ -828,7 +828,7 @@ namespace Horde.Server.Agents
 			span.SetAttribute("FinishTime", finishTime?.ToString());
 			span.SetAttribute("Index", index);
 			span.SetAttribute("Count", count);
-			return _leases.FindLeasesAsync(agentId, sessionId, startTime, finishTime, index, count);
+			return _leases.FindLeasesAsync(null, agentId, sessionId, startTime, finishTime, index, count);
 		}
 
 		/// <summary>
@@ -898,6 +898,25 @@ namespace Horde.Server.Agents
 
 			// Update the lease
 			await _leases.TrySetOutcomeAsync(lease.Id, finishTime, outcome, output);
+
+			// Terminate any child leases
+			List<ILease> childLeases = await _leases.FindLeasesAsync(parentId: lease.Id);
+			foreach (ILease childLease in childLeases)
+			{
+				if (childLease.Outcome == LeaseOutcome.Unspecified)
+				{
+					IAgent? otherAgent = await GetAgentAsync(childLease.AgentId);
+					if (otherAgent != null)
+					{
+						AgentLease? otherLease = otherAgent.Leases.FirstOrDefault(x => x.Id == childLease.Id);
+						if (otherLease != null)
+						{
+							_logger.LogInformation("Terminating child lease {LeaseId} with parent {ParentLeaseId}", otherLease.Id, lease.Id);
+							await RemoveLeaseAsync(otherAgent, otherLease, utcNow, LeaseOutcome.Failed, null);
+						}
+					}
+				}
+			}
 		}
 
 		/// <summary>
diff --git a/Engine/Source/Programs/Horde/Horde.Server/Agents/IAgent.cs b/Engine/Source/Programs/Horde/Horde.Server/Agents/IAgent.cs
--- a/Engine/Source/Programs/Horde/Horde.Server/Agents/IAgent.cs
+++ b/Engine/Source/Programs/Horde/Horde.Server/Agents/IAgent.cs
@@ -272,6 +272,11 @@ namespace Horde.Server.Agents
 		[BsonRequired]
 		public LeaseId Id { get; set; }
 
+		/// <summary>
+		/// The parent lease id
+		/// </summary>
+		public LeaseId? ParentId { get; set; }
+
 		/// <summary>
 		/// Name of this lease
 		/// </summary>
@@ -341,6 +346,7 @@ namespace Horde.Server.Agents
 		/// Constructor
 		/// </summary>
 		/// <param name="id">Identifier for the lease</param>
+		/// <param name="parentId">The parent lease id</param>
 		/// <param name="name">Name of this lease</param>
 		/// <param name="streamId"></param>
 		/// <param name="poolId"></param>
@@ -349,9 +355,10 @@ namespace Horde.Server.Agents
 		/// <param name="resources">Resources required for this lease</param>
 		/// <param name="exclusive">Whether to reserve the entire device</param>
 		/// <param name="payload">Encoded "any" protobuf describing the contents of the payload</param>
-		public AgentLease(LeaseId id, string name, StreamId? streamId, PoolId? poolId, LogId? logId, LeaseState state, IReadOnlyDictionary<string, int>? resources, bool exclusive, byte[]? payload)
+		public AgentLease(LeaseId id, LeaseId? parentId, string name, StreamId? streamId, PoolId? poolId, LogId? logId, LeaseState state, IReadOnlyDictionary<string, int>? resources, bool exclusive, byte[]? payload)
 		{
 			Id = id;
+			ParentId = parentId;
 			Name = name;
 			StreamId = streamId;
 			PoolId = poolId;
diff --git a/Engine/Source/Programs/Horde/Horde.Server/Agents/Leases/ILease.cs b/Engine/Source/Programs/Horde/Horde.Server/Agents/Leases/ILease.cs
--- a/Engine/Source/Programs/Horde/Horde.Server/Agents/Leases/ILease.cs
+++ b/Engine/Source/Programs/Horde/Horde.Server/Agents/Leases/ILease.cs
@@ -7,7 +7,6 @@ using Horde.Server.Agents.Pools;
 using Horde.Server.Agents.Sessions;
 using Horde.Server.Logs;
 using Horde.Server.Streams;
-using Horde.Server.Utilities;
 using HordeCommon;
 
 namespace Horde.Server.Agents.Leases
@@ -22,6 +21,11 @@ namespace Horde.Server.Agents.Leases
 		/// </summary>
 		public LeaseId Id { get; }
 
+		/// <summary>
+		/// Identifier for the parent lease. Used to terminate hierarchies of leases.
+		/// </summary>
+		public LeaseId? ParentId { get; }
+
 		/// <summary>
 		/// Name of this lease
 		/// </summary>
diff --git a/Engine/Source/Programs/Horde/Horde.Server/Agents/Leases/ILeaseCollection.cs b/Engine/Source/Programs/Horde/Horde.Server/Agents/Leases/ILeaseCollection.cs
--- a/Engine/Source/Programs/Horde/Horde.Server/Agents/Leases/ILeaseCollection.cs
+++ b/Engine/Source/Programs/Horde/Horde.Server/Agents/Leases/ILeaseCollection.cs
@@ -21,6 +21,7 @@ namespace Horde.Server.Agents.Leases
 		/// Adds a lease to the collection
 		/// </summary>
 		/// <param name="id">The lease id</param>
+		/// <param name="parentId">The parent lease id</param>
 		/// <param name="name">Name of the lease</param>
 		/// <param name="agentId">The agent id</param>
 		/// <param name="sessionId">The agent session handling the lease</param>
@@ -30,7 +31,7 @@ namespace Horde.Server.Agents.Leases
 		/// <param name="startTime">Start time of the lease</param>
 		/// <param name="payload">Payload for the lease</param>
 		/// <returns>Async task</returns>
-		Task<ILease> AddAsync(LeaseId id, string name, AgentId agentId, SessionId sessionId, StreamId? streamId, PoolId? poolId, LogId? logId, DateTime startTime, byte[] payload);
+		Task<ILease> AddAsync(LeaseId id, LeaseId? parentId, string name, AgentId agentId, SessionId sessionId, StreamId? streamId, PoolId? poolId, LogId? logId, DateTime startTime, byte[] payload);
 
 		/// <summary>
 		/// Deletes a lease from the collection
@@ -49,6 +50,7 @@ namespace Horde.Server.Agents.Leases
 		/// <summary>
 		/// Finds all leases matching a set of criteria
 		/// </summary>
+		/// <param name="parentId">The parent lease id</param>
 		/// <param name="agentId">Unqiue id of the agent executing this lease</param>
 		/// <param name="sessionId">Unique id of the agent session</param>
 		/// <param name="minTime">Start of the window to include leases</param>
@@ -58,7 +60,7 @@ namespace Horde.Server.Agents.Leases
 		/// <param name="indexHint">Name of index to be specified as a hint to the database query planner</param>
 		/// <param name="consistentRead">If the database read should be made to the replica server</param>
 		/// <returns>List of leases matching the given criteria</returns>
-		Task<List<ILease>> FindLeasesAsync(AgentId? agentId = null, SessionId? sessionId = null, DateTime? minTime = null, DateTime? maxTime = null, int? index = null, int? count = null, string? indexHint = null, bool consistentRead = true);
+		Task<List<ILease>> FindLeasesAsync(LeaseId? parentId = null, AgentId? agentId = null, SessionId? sessionId = null, DateTime? minTime = null, DateTime? maxTime = null, int? index = null, int? count = null, string? indexHint = null, bool consistentRead = true);
 
 		/// <summary>
 		/// Finds all leases by finish time
diff --git a/Engine/Source/Programs/Horde/Horde.Server/Agents/Leases/LeaseCollection.cs b/Engine/Source/Programs/Horde/Horde.Server/Agents/Leases/LeaseCollection.cs
--- a/Engine/Source/Programs/Horde/Horde.Server/Agents/Leases/LeaseCollection.cs
+++ b/Engine/Source/Programs/Horde/Horde.Server/Agents/Leases/LeaseCollection.cs
@@ -27,6 +27,10 @@ namespace Horde.Server.Agents.Leases
 		class LeaseDocument : ILease
 		{
 			public LeaseId Id { get; set; }
+
+			[BsonIgnoreIfNull]
+			public LeaseId? ParentId { get; set; }
+
 			public string Name { get; set; }
 			public AgentId AgentId { get; set; }
 			public SessionId SessionId { get; set; }
@@ -51,9 +55,10 @@ namespace Horde.Server.Agents.Leases
 				Payload = null!;
 			}
 
-			public LeaseDocument(LeaseId id, string name, AgentId agentId, SessionId sessionId, StreamId? streamId, PoolId? poolId, LogId? logId, DateTime startTime, byte[] payload)
+			public LeaseDocument(LeaseId id, LeaseId? parentId, string name, AgentId agentId, SessionId sessionId, StreamId? streamId, PoolId? poolId, LogId? logId, DateTime startTime, byte[] payload)
 			{
 				Id = id;
+				ParentId = parentId;
 				Name = name;
 				AgentId = agentId;
 				SessionId = sessionId;
@@ -75,6 +80,7 @@ namespace Horde.Server.Agents.Leases
 		public LeaseCollection(MongoService mongoService)
 		{
 			List<MongoIndex<LeaseDocument>> indexes = new List<MongoIndex<LeaseDocument>>();
+			indexes.Add(keys => keys.Ascending(x => x.ParentId), sparse: true);
 			indexes.Add(keys => keys.Ascending(x => x.AgentId));
 			indexes.Add(keys => keys.Ascending(x => x.SessionId));
 			indexes.Add(keys => keys.Ascending(x => x.StartTime));
@@ -85,9 +91,9 @@ namespace Horde.Server.Agents.Leases
 		}
 
 		/// <inheritdoc/>
-		public async Task<ILease> AddAsync(LeaseId id, string name, AgentId agentId, SessionId sessionId, StreamId? streamId, PoolId? poolId, LogId? logId, DateTime startTime, byte[] payload)
+		public async Task<ILease> AddAsync(LeaseId id, LeaseId? parentId, string name, AgentId agentId, SessionId sessionId, StreamId? streamId, PoolId? poolId, LogId? logId, DateTime startTime, byte[] payload)
 		{
-			LeaseDocument lease = new LeaseDocument(id, name, agentId, sessionId, streamId, poolId, logId, startTime, payload);
+			LeaseDocument lease = new LeaseDocument(id, parentId, name, agentId, sessionId, streamId, poolId, logId, startTime, payload);
 			await _leases.ReplaceOneAsync(x => x.Id == id, lease, new ReplaceOptions { IsUpsert = true });
 			return lease;
 		}
@@ -105,12 +111,16 @@ namespace Horde.Server.Agents.Leases
 		}
 
 		/// <inheritdoc/>
-		public async Task<List<ILease>> FindLeasesAsync(AgentId? agentId, SessionId? sessionId, DateTime? minTime, DateTime? maxTime, int? index, int? count, string? indexHint = null, bool consistentRead = true)
+		public async Task<List<ILease>> FindLeasesAsync(LeaseId? parentId, AgentId? agentId, SessionId? sessionId, DateTime? minTime, DateTime? maxTime, int? index, int? count, string? indexHint = null, bool consistentRead = true)
 		{
 			IMongoCollection<LeaseDocument> collection = consistentRead ? _leases : _leases.WithReadPreference(ReadPreference.SecondaryPreferred);
 			
 			FilterDefinitionBuilder<LeaseDocument> filterBuilder = Builders<LeaseDocument>.Filter;
 			FilterDefinition<LeaseDocument> filter = FilterDefinition<LeaseDocument>.Empty;
+			if (parentId != null)
+			{
+				filter &= filterBuilder.Eq(x => x.ParentId, parentId.Value);
+			}
 			if (agentId != null)
 			{
 				filter &= filterBuilder.Eq(x => x.AgentId, agentId.Value);
@@ -160,7 +170,7 @@ namespace Horde.Server.Agents.Leases
 		/// <inheritdoc/>
 		public async Task<List<ILease>> FindLeasesAsync(DateTime? minTime, DateTime? maxTime)
 		{
-			return await FindLeasesAsync(null, null, minTime, maxTime, null, null, _finishTimeStartTimeCompoundIndex.Name, false);
+			return await FindLeasesAsync(null, null, null, minTime, maxTime, null, null, _finishTimeStartTimeCompoundIndex.Name, false);
 		}
 
 		/// <inheritdoc/>
diff --git a/Engine/Source/Programs/Horde/Horde.Server/Compute/ComputeController.cs b/Engine/Source/Programs/Horde/Horde.Server/Compute/ComputeController.cs
--- a/Engine/Source/Programs/Horde/Horde.Server/Compute/ComputeController.cs
+++ b/Engine/Source/Programs/Horde/Horde.Server/Compute/ComputeController.cs
@@ -7,6 +7,7 @@ using System.Threading.Tasks;
 using EpicGames.Core;
 using EpicGames.Horde.Compute;
 using Horde.Server.Acls;
+using Horde.Server.Agents.Leases;
 using Horde.Server.Server;
 using Horde.Server.Utilities;
 using Microsoft.AspNetCore.Authorization;
@@ -103,9 +104,11 @@ namespace Horde.Server.Compute
 				return Forbid(ComputeAclAction.AddComputeTasks, clusterId);
 			}
 
+			LeaseId? parentLeaseId = User.GetLeaseClaim();
+
 			Requirements requirements = request.Requirements ?? new Requirements();
 
-			ComputeResource? computeResource = await _computeService.TryAllocateResource(requirements);
+			ComputeResource? computeResource = await _computeService.TryAllocateResource(requirements, parentLeaseId);
 			if (computeResource == null)
 			{
 				return StatusCode((int)HttpStatusCode.ServiceUnavailable);
diff --git a/Engine/Source/Programs/Horde/Horde.Server/Compute/ComputeService.cs b/Engine/Source/Programs/Horde/Horde.Server/Compute/ComputeService.cs
--- a/Engine/Source/Programs/Horde/Horde.Server/Compute/ComputeService.cs
+++ b/Engine/Source/Programs/Horde/Horde.Server/Compute/ComputeService.cs
@@ -38,8 +38,7 @@ namespace Horde.Server.Compute
 		/// <summary>
 		/// Allocates a compute resource
 		/// </summary>
-		/// <returns></returns>
-		public async Task<ComputeResource?> TryAllocateResource(Requirements requirements)
+		public async Task<ComputeResource?> TryAllocateResource(Requirements requirements, LeaseId? parentLeaseId)
 		{
 			List<IAgent> agents = await _agentCollection.FindAsync();
 			foreach (IAgent agent in agents)
@@ -50,7 +49,7 @@ namespace Horde.Server.Compute
 					ComputeTask computeTask = CreateComputeTask(assignedResources);
 
 					byte[] payload = Any.Pack(computeTask).ToByteArray();
-					AgentLease lease = new AgentLease(LeaseId.GenerateNewId(), "Compute task", null, null, null, LeaseState.Pending, assignedResources, requirements.Exclusive, payload);
+					AgentLease lease = new AgentLease(LeaseId.GenerateNewId(), parentLeaseId, "Compute task", null, null, null, LeaseState.Pending, assignedResources, requirements.Exclusive, payload);
 
 					ComputeResource? resource = TryAssign(agent, computeTask);
 					if (resource != null)
diff --git a/Engine/Source/Programs/Horde/Horde.Server/Jobs/JobTaskSource.cs b/Engine/Source/Programs/Horde/Horde.Server/Jobs/JobTaskSource.cs
--- a/Engine/Source/Programs/Horde/Horde.Server/Jobs/JobTaskSource.cs
+++ b/Engine/Source/Programs/Horde/Horde.Server/Jobs/JobTaskSource.cs
@@ -766,7 +766,7 @@ namespace Horde.Server.Jobs
 					byte[] payload = Any.Pack(task).ToByteArray();
 
 					// Create the lease and try to set it on the waiter. If this fails, the waiter has already moved on, and the lease can be cancelled.
-					AgentLease lease = new AgentLease(leaseId, leaseName.ToString(), job.StreamId, item._poolId, logId, LeaseState.Pending, null, true, payload);
+					AgentLease lease = new AgentLease(leaseId, null, leaseName.ToString(), job.StreamId, item._poolId, logId, LeaseState.Pending, null, true, payload);
 					if (waiter.LeaseSource.TrySetResult(lease))
 					{
 						_logger.LogDebug("Assigned lease {LeaseId} to agent {AgentId}", leaseId, agent.Id);
diff --git a/Engine/Source/Programs/Horde/Horde.Server/Tasks/ConformTaskSource.cs b/Engine/Source/Programs/Horde/Horde.Server/Tasks/ConformTaskSource.cs
--- a/Engine/Source/Programs/Horde/Horde.Server/Tasks/ConformTaskSource.cs
+++ b/Engine/Source/Programs/Horde/Horde.Server/Tasks/ConformTaskSource.cs
@@ -182,7 +182,7 @@ namespace Horde.Server.Tasks
 
 						byte[] payload = Any.Pack(task).ToByteArray();
 
-						return Lease(new AgentLease(leaseId, "Updating workspaces", null, null, log.Id, LeaseState.Pending, null, true, payload));
+						return Lease(new AgentLease(leaseId, null, "Updating workspaces", null, null, log.Id, LeaseState.Pending, null, true, payload));
 					}
 				}
 			}
diff --git a/Engine/Source/Programs/Horde/Horde.Server/Tasks/RestartTaskSource.cs b/Engine/Source/Programs/Horde/Horde.Server/Tasks/RestartTaskSource.cs
--- a/Engine/Source/Programs/Horde/Horde.Server/Tasks/RestartTaskSource.cs
+++ b/Engine/Source/Programs/Horde/Horde.Server/Tasks/RestartTaskSource.cs
@@ -48,7 +48,7 @@ namespace Horde.Server.Tasks
 
 			byte[] payload = Any.Pack(task).ToByteArray();
 
-			return Lease(new AgentLease(leaseId, "Restart", null, null, log.Id, LeaseState.Pending, null, true, payload));
+			return Lease(new AgentLease(leaseId, null, "Restart", null, null, log.Id, LeaseState.Pending, null, true, payload));
 		}
 	}
 }
diff --git a/Engine/Source/Programs/Horde/Horde.Server/Tasks/ShutdownTaskSource.cs b/Engine/Source/Programs/Horde/Horde.Server/Tasks/ShutdownTaskSource.cs
--- a/Engine/Source/Programs/Horde/Horde.Server/Tasks/ShutdownTaskSource.cs
+++ b/Engine/Source/Programs/Horde/Horde.Server/Tasks/ShutdownTaskSource.cs
@@ -47,7 +47,7 @@ namespace Horde.Server.Tasks
 
 			byte[] payload = Any.Pack(task).ToByteArray();
 
-			return Lease(new AgentLease(leaseId, "Shutdown", null, null, log.Id, LeaseState.Pending, null, true, payload));
+			return Lease(new AgentLease(leaseId, null, "Shutdown", null, null, log.Id, LeaseState.Pending, null, true, payload));
 		}
 	}
 }
diff --git a/Engine/Source/Programs/Horde/Horde.Server/Tasks/UpgradeTaskSource.cs b/Engine/Source/Programs/Horde/Horde.Server/Tasks/UpgradeTaskSource.cs
--- a/Engine/Source/Programs/Horde/Horde.Server/Tasks/UpgradeTaskSource.cs
+++ b/Engine/Source/Programs/Horde/Horde.Server/Tasks/UpgradeTaskSource.cs
@@ -75,7 +75,7 @@ namespace Horde.Server.Tasks
 			task.LogId = logFile.Id.ToString();
 
 			byte[] payload = Any.Pack(task).ToByteArray();
-			return Lease(new AgentLease(leaseId, $"Upgrade to {tool.Id} {deployment.Version}", null, null, logFile.Id, LeaseState.Pending, null, true, payload));
+			return Lease(new AgentLease(leaseId, null, $"Upgrade to {tool.Id} {deployment.Version}", null, null, logFile.Id, LeaseState.Pending, null, true, payload));
 		}
 
 		/// <summary>
