Horde: Sync the autosdk workspace corresponding to all pools an agent is part of. Current implementation is thrashing autosdk files in and out of the cache.

diff --git a/Engine/Source/Programs/Horde/Horde.Server/Agents/IAgent.cs b/Engine/Source/Programs/Horde/Horde.Server/Agents/IAgent.cs
--- a/Engine/Source/Programs/Horde/Horde.Server/Agents/IAgent.cs
+++ b/Engine/Source/Programs/Horde/Horde.Server/Agents/IAgent.cs
@@ -788,29 +788,25 @@ namespace Horde.Server.Agents
 		}
 
 		/// <summary>
-		/// Gets all the autosdk workspaces required for an agent
+		/// Get the AutoSDK workspace required for an agent
 		/// </summary>
 		/// <param name="agent"></param>
-		/// <param name="globalConfig"></param>
-		/// <param name="autoSdkConfig">Config for autosdk workspaces</param>
-		/// <param name="workspaces"></param>
+		/// <param name="cluster">The perforce cluster to get a workspace for</param>
+		/// <param name="pools">Pools that the agent belongs to</param>
 		/// <returns></returns>
-		public static HashSet<AgentWorkspace> GetAutoSdkWorkspaces(this IAgent agent, GlobalConfig globalConfig, AutoSdkConfig autoSdkConfig, List<AgentWorkspace> workspaces)
+		public static AgentWorkspace? GetAutoSdkWorkspace(this IAgent agent, PerforceCluster cluster, IEnumerable<IPool> pools)
 		{
-			HashSet<AgentWorkspace> autoSdkWorkspaces = new HashSet<AgentWorkspace>();
-			foreach (string? clusterName in workspaces.Select(x => x.Cluster).Distinct())
+			AutoSdkConfig? autoSdkConfig = null;
+			foreach(IPool pool in pools)
 			{
-				PerforceCluster? cluster = globalConfig.FindPerforceCluster(clusterName);
-				if (cluster != null)
-				{
-					AgentWorkspace? autoSdkWorkspace = GetAutoSdkWorkspace(agent, cluster, autoSdkConfig);
-					if (autoSdkWorkspace != null)
-					{
-						autoSdkWorkspaces.Add(autoSdkWorkspace);
-					}
-				}
+				autoSdkConfig = AutoSdkConfig.Merge(autoSdkConfig, pool.AutoSdkConfig);
+			}
+			if (autoSdkConfig == null)
+			{
+				return null;
 			}
-			return autoSdkWorkspaces;
+
+			return GetAutoSdkWorkspace(agent, cluster, autoSdkConfig);
 		}
 
 		/// <summary>
diff --git a/Engine/Source/Programs/Horde/Horde.Server/Agents/Pools/PoolService.cs b/Engine/Source/Programs/Horde/Horde.Server/Agents/Pools/PoolService.cs
--- a/Engine/Source/Programs/Horde/Horde.Server/Agents/Pools/PoolService.cs
+++ b/Engine/Source/Programs/Horde/Horde.Server/Agents/Pools/PoolService.cs
@@ -6,10 +6,7 @@ using System.Linq;
 using System.Threading;
 using System.Threading.Tasks;
 using EpicGames.Horde;
-using EpicGames.Horde.Common;
-using Horde.Server.Agents.Fleet;
 using Horde.Server.Server;
-using Horde.Server.Utilities;
 using HordeCommon;
 
 namespace Horde.Server.Agents.Pools
@@ -117,7 +114,7 @@ namespace Horde.Server.Agents.Pools
 		/// <param name="poolId"></param>
 		/// <param name="validAtTime"></param>
 		/// <returns></returns>
-		public async Task<IPool?> GetCachedPoolAsync(PoolId poolId, DateTime validAtTime)
+		public async Task<IPool?> GetPoolAsync(PoolId poolId, DateTime validAtTime)
 		{
 			Dictionary<PoolId, IPool> poolMapping = await GetPoolLookupAsync(validAtTime);
 			poolMapping.TryGetValue(poolId, out IPool? pool);
@@ -130,7 +127,7 @@ namespace Horde.Server.Agents.Pools
 		/// <param name="agent"></param>
 		/// <param name="validAtTime"></param>
 		/// <returns></returns>
-		public async Task<List<IPool>> GetCachedPoolsAsync(IAgent agent, DateTime validAtTime)
+		public async Task<List<IPool>> GetPoolsAsync(IAgent agent, DateTime validAtTime)
 		{
 			Dictionary<PoolId, IPool> poolMapping = await GetPoolLookupAsync(validAtTime);
 
@@ -155,8 +152,75 @@ namespace Horde.Server.Agents.Pools
 		/// <returns>List of workspaces</returns>
 		public async Task<HashSet<AgentWorkspace>> GetWorkspacesAsync(IAgent agent, DateTime validAtTime, GlobalConfig globalConfig)
 		{
-			AutoSdkConfig? autoSdkConfig = null;
+			List<IPool> pools = await GetPoolsAsync(agent, validAtTime);
+
 			HashSet<AgentWorkspace> workspaces = new HashSet<AgentWorkspace>();
+			foreach (IPool pool in pools)
+			{
+				workspaces.UnionWith(pool.Workspaces);
+			}
+
+			AutoSdkConfig? autoSdkConfig = GetAutoSdkConfig(pools);
+			if (autoSdkConfig != null)
+			{
+				foreach (string? clusterName in workspaces.Select(x => x.Cluster).Distinct())
+				{
+					PerforceCluster? cluster = globalConfig.FindPerforceCluster(clusterName);
+					if (cluster != null)
+					{
+						AgentWorkspace? autoSdkWorkspace = agent.GetAutoSdkWorkspace(cluster, autoSdkConfig);
+						if (autoSdkWorkspace != null)
+						{
+							workspaces.Add(autoSdkWorkspace);
+						}
+					}
+				}
+			}
+
+			return workspaces;
+		}
+
+		static AutoSdkConfig? GetAutoSdkConfig(IEnumerable<IPool> pools)
+		{
+			AutoSdkConfig? autoSdkConfig = null;
+			foreach(IPool pool in pools)
+			{
+				autoSdkConfig = AutoSdkConfig.Merge(autoSdkConfig, pool.AutoSdkConfig);
+			}
+			return autoSdkConfig;
+		}
+
+		/// <summary>
+		/// Gets all the autosdk workspaces required for an agent
+		/// </summary>
+		/// <param name="agent"></param>
+		/// <param name="cluster"></param>
+		/// <param name="globalConfig"></param>
+		/// <returns></returns>
+		public async Task<AgentWorkspace?> GetAutoSdkWorkspaceAsync(IAgent agent, PerforceCluster cluster, GlobalConfig globalConfig)
+		{
+			List<IPool> pools = await GetPoolsAsync(agent, DateTime.UtcNow - TimeSpan.FromSeconds(10.0));
+
+			AutoSdkConfig? autoSdkConfig = GetAutoSdkConfig(pools);
+			if (autoSdkConfig == null)
+			{
+				return null;
+			}
+
+			return agent.GetAutoSdkWorkspace(cluster, autoSdkConfig);
+		}
+
+		/// <summary>
+		/// Get a list of workspaces for the given agent
+		/// </summary>
+		/// <param name="agent">The agent to return workspaces for</param>
+		/// <param name="perforceCluster">The P4 cluster to find a workspace for</param>
+		/// <param name="validAtTime">Absolute time at which we expect the results to be valid. Values may be cached as long as they are after this time.</param>
+		/// <param name="globalConfig">Current configuration</param>
+		/// <returns>List of workspaces</returns>
+		public async Task<AgentWorkspace?> GetAutoSdkWorkspaceAsync(IAgent agent, PerforceCluster perforceCluster, DateTime validAtTime, GlobalConfig globalConfig)
+		{
+			AutoSdkConfig? autoSdkConfig = null;
 
 			Dictionary<PoolId, IPool> poolMapping = await GetPoolLookupAsync(validAtTime);
 			foreach (PoolId poolId in agent.GetPools())
@@ -164,17 +228,18 @@ namespace Horde.Server.Agents.Pools
 				IPool? pool;
 				if (poolMapping.TryGetValue(poolId, out pool))
 				{
-					workspaces.UnionWith(pool.Workspaces);
 					autoSdkConfig = AutoSdkConfig.Merge(autoSdkConfig, pool.AutoSdkConfig);
 				}
 			}
 
-			if (autoSdkConfig != null)
+			if (autoSdkConfig == null)
 			{
-				workspaces.UnionWith(agent.GetAutoSdkWorkspaces(globalConfig, autoSdkConfig, workspaces.ToList()));
+				return null;
+			}
+			else
+			{
+				return agent.GetAutoSdkWorkspace(perforceCluster, autoSdkConfig);
 			}
-
-			return workspaces;
 		}
 
 		/// <summary>
diff --git a/Engine/Source/Programs/Horde/Horde.Server/Jobs/JobTaskSource.cs b/Engine/Source/Programs/Horde/Horde.Server/Jobs/JobTaskSource.cs
--- a/Engine/Source/Programs/Horde/Horde.Server/Jobs/JobTaskSource.cs
+++ b/Engine/Source/Programs/Horde/Horde.Server/Jobs/JobTaskSource.cs
@@ -75,7 +75,7 @@ namespace Horde.Server.Jobs
 			/// <summary>
 			/// Whether or not to use the AutoSDK.
 			/// </summary>
-			public AutoSdkConfig? _autoSdkConfig;
+			public bool _useAutoSdk;
 
 			/// <summary>
 			/// Task for creating a lease and assigning to a waiter
@@ -99,14 +99,14 @@ namespace Horde.Server.Jobs
 			/// <param name="batchIdx">The batch index to execute</param>
 			/// <param name="poolId">Unique id of the pool of machines to allocate from</param>
 			/// <param name="workspace">The workspace that this job should run in</param>
-			/// <param name="autoSdkConfig">Whether or not to use the AutoSDK</param>
-			public QueueItem(IJob job, int batchIdx, PoolId poolId, AgentWorkspace workspace, AutoSdkConfig? autoSdkConfig)
+			/// <param name="useAutoSdk">Whether or not to use the AutoSDK</param>
+			public QueueItem(IJob job, int batchIdx, PoolId poolId, AgentWorkspace workspace, bool useAutoSdk)
 			{
 				_job = job;
 				_batchIdx = batchIdx;
 				_poolId = poolId;
 				_workspace = workspace;
-				_autoSdkConfig = autoSdkConfig;
+				_useAutoSdk = useAutoSdk;
 			}
 		}
 
@@ -183,6 +183,7 @@ namespace Horde.Server.Jobs
 		readonly IJobStepRefCollection _jobStepRefs;
 		readonly IGraphCollection _graphs;
 		readonly IPoolCollection _poolCollection;
+		readonly PoolService _poolService;
 		readonly IUgsMetadataCollection _ugsMetadataCollection;
 		readonly PerforceLoadBalancer _perforceLoadBalancer;
 		readonly IClock _clock;
@@ -222,7 +223,7 @@ namespace Horde.Server.Jobs
 		/// <summary>
 		/// Constructor
 		/// </summary>
-		public JobTaskSource(AclService aclService, IAgentCollection agents, IJobCollection jobs, IJobStepRefCollection jobStepRefs, IGraphCollection graphs, IPoolCollection pools, IUgsMetadataCollection ugsMetadataCollection, IStreamCollection streamCollection, ILogFileService logFileService, PerforceLoadBalancer perforceLoadBalancer, IClock clock, IOptionsMonitor<ServerSettings> settings, IOptionsMonitor<GlobalConfig> globalConfig, ILogger<JobTaskSource> logger)
+		public JobTaskSource(AclService aclService, IAgentCollection agents, IJobCollection jobs, IJobStepRefCollection jobStepRefs, IGraphCollection graphs, IPoolCollection pools, PoolService poolService, IUgsMetadataCollection ugsMetadataCollection, IStreamCollection streamCollection, ILogFileService logFileService, PerforceLoadBalancer perforceLoadBalancer, IClock clock, IOptionsMonitor<ServerSettings> settings, IOptionsMonitor<GlobalConfig> globalConfig, ILogger<JobTaskSource> logger)
 		{
 			_aclService = aclService;
 			_agentsCollection = agents;
@@ -230,6 +231,7 @@ namespace Horde.Server.Jobs
 			_jobStepRefs = jobStepRefs;
 			_graphs = graphs;
 			_poolCollection = pools;
+			_poolService = poolService;
 			_ugsMetadataCollection = ugsMetadataCollection;
 			_streamCollection = streamCollection;
 			_logFileService = logFileService;
@@ -454,7 +456,7 @@ namespace Horde.Server.Jobs
 
 						if (newJob != null)
 						{
-							QueueItem newQueueItem = new QueueItem(newJob, batchIdx, agentType.Pool, workspace, autoSdkConfig);
+							QueueItem newQueueItem = new QueueItem(newJob, batchIdx, agentType.Pool, workspace, autoSdkConfig != null);
 							newQueue.Add(newQueueItem);
 							newBatchIdToQueueItem[(newJob.Id, batch.Id)] = newQueueItem;
 
@@ -628,7 +630,7 @@ namespace Horde.Server.Jobs
 								else
 								{
 									RemoveQueueItem(existingItem);
-									InsertQueueItem(job, batchIdx, existingItem._poolId, existingItem._workspace, existingItem._autoSdkConfig);
+									InsertQueueItem(job, batchIdx, existingItem._poolId, existingItem._workspace, existingItem._useAutoSdk);
 								}
 							}
 							continue;
@@ -643,7 +645,7 @@ namespace Horde.Server.Jobs
 						{
 							if (streamConfig.TryGetAgentWorkspace(agentType, out AgentWorkspace? agentWorkspace, out AutoSdkConfig? autoSdkConfig))
 							{
-								InsertQueueItem(job, batchIdx, agentType.Pool, agentWorkspace, autoSdkConfig);
+								InsertQueueItem(job, batchIdx, agentType.Pool, agentWorkspace, autoSdkConfig != null);
 							}
 						}
 					}
@@ -759,8 +761,16 @@ namespace Horde.Server.Jobs
 				}
 				leaseName.Append(CultureInfo.InvariantCulture, $" - {job.Name}");
 
+				// Get the autosdk workspace
+				AgentWorkspace? autoSdkWorkspace = null;
+				if (item._useAutoSdk)
+				{
+					PerforceCluster cluster = _globalConfig.CurrentValue.FindPerforceCluster(streamConfig.ClusterName)!;
+					autoSdkWorkspace = await _poolService.GetAutoSdkWorkspaceAsync(agent, cluster, DateTime.UtcNow - TimeSpan.FromSeconds(10.0), _globalConfig.CurrentValue);
+				}
+
 				// Encode the payload
-				ExecuteJobTask? task = await CreateExecuteJobTaskAsync(leaseId, streamConfig, job, batch, agent, item._workspace, item._autoSdkConfig, logId);
+				ExecuteJobTask? task = await CreateExecuteJobTaskAsync(leaseId, streamConfig, job, batch, agent, item._workspace, autoSdkWorkspace, logId);
 				if (task != null)
 				{
 					byte[] payload = Any.Pack(task).ToByteArray();
@@ -810,7 +820,7 @@ namespace Horde.Server.Jobs
 			return null;
 		}
 
-		async Task<ExecuteJobTask?> CreateExecuteJobTaskAsync(LeaseId leaseId, StreamConfig streamConfig, IJob job, IJobStepBatch batch, IAgent agent, AgentWorkspace workspace, AutoSdkConfig? autoSdkConfig, LogId logId)
+		async Task<ExecuteJobTask?> CreateExecuteJobTaskAsync(LeaseId leaseId, StreamConfig streamConfig, IJob job, IJobStepBatch batch, IAgent agent, AgentWorkspace workspace, AgentWorkspace? autoSdkWorkspace, LogId logId)
 		{
 			// Get the lease name
 			StringBuilder leaseName = new StringBuilder($"{streamConfig.Name} - ");
@@ -857,7 +867,6 @@ namespace Horde.Server.Jobs
 				return null;
 			}
 
-			AgentWorkspace? autoSdkWorkspace = (autoSdkConfig != null)? agent.GetAutoSdkWorkspace(cluster, autoSdkConfig) : null;
 			if (autoSdkWorkspace != null)
 			{
 				autoSdkWorkspace.Method = workspace.Method;
@@ -1095,13 +1104,13 @@ namespace Horde.Server.Jobs
 		/// <param name="batchIdx"></param>
 		/// <param name="poolId">The pool to use</param>
 		/// <param name="workspace">The workspace for this item to run in</param>
-		/// <param name="autoSdkConfig">Whether or not to use the AutoSDK</param>
+		/// <param name="useAutoSdk"></param>
 		/// <returns></returns>
-		void InsertQueueItem(IJob job, int batchIdx, PoolId poolId, AgentWorkspace workspace, AutoSdkConfig? autoSdkConfig)
+		void InsertQueueItem(IJob job, int batchIdx, PoolId poolId, AgentWorkspace workspace, bool useAutoSdk)
 		{
-			_logger.LogDebug("Adding queued job {JobId}, batch {BatchId} [Pool: {Pool}, Workspace: {Workspace}]", job.Id, job.Batches[batchIdx].Id, poolId, workspace.Identifier);
+			_logger.LogDebug("Adding queued job {JobId}, batch {BatchId} [Pool: {Pool}, Workspace: {Workspace}, AutoSdk: {AutoSdk}]", job.Id, job.Batches[batchIdx].Id, poolId, workspace.Identifier, useAutoSdk);
 
-			QueueItem newItem = new QueueItem(job, batchIdx, poolId, workspace, autoSdkConfig);
+			QueueItem newItem = new QueueItem(job, batchIdx, poolId, workspace, useAutoSdk);
 			_batchIdToQueueItem[newItem.Id] = newItem;
 			_queue.Add(newItem);
 
diff --git a/Engine/Source/Programs/Horde/Horde.Server/Tasks/ConformTaskSource.cs b/Engine/Source/Programs/Horde/Horde.Server/Tasks/ConformTaskSource.cs
--- a/Engine/Source/Programs/Horde/Horde.Server/Tasks/ConformTaskSource.cs
+++ b/Engine/Source/Programs/Horde/Horde.Server/Tasks/ConformTaskSource.cs
@@ -341,7 +341,7 @@ namespace Horde.Server.Tasks
 				}
 
 				// Get the current pools for the agent
-				List<IPool> pools = await _poolService.GetCachedPoolsAsync(agent, DateTime.UtcNow - TimeSpan.FromMinutes(2.0));
+				List<IPool> pools = await _poolService.GetPoolsAsync(agent, DateTime.UtcNow - TimeSpan.FromMinutes(2.0));
 
 				TimeSpan? conformInterval = null;
 				foreach(IPool pool in pools)
