Horde: Fix collection modification during iteration

diff --git a/Engine/Source/Programs/Horde/Horde.Server/Agents/Pools/PoolService.cs b/Engine/Source/Programs/Horde/Horde.Server/Agents/Pools/PoolService.cs
--- a/Engine/Source/Programs/Horde/Horde.Server/Agents/Pools/PoolService.cs
+++ b/Engine/Source/Programs/Horde/Horde.Server/Agents/Pools/PoolService.cs
@@ -163,7 +163,7 @@ namespace Horde.Server.Agents.Pools
 			AutoSdkConfig? autoSdkConfig = GetAutoSdkConfig(pools);
 			if (autoSdkConfig != null)
 			{
-				foreach (string? clusterName in workspaces.Select(x => x.Cluster).Distinct())
+				foreach (string? clusterName in workspaces.Select(x => x.Cluster).Distinct().ToList())
 				{
 					PerforceCluster? cluster = globalConfig.FindPerforceCluster(clusterName);
 					if (cluster != null)
