Horde: Add a telemetry span showing how long semaphores are kept when reading from S3.

diff --git a/Engine/Source/Programs/Horde/Horde.Server/Storage/Backends/AwsStorageBackend.cs b/Engine/Source/Programs/Horde/Horde.Server/Storage/Backends/AwsStorageBackend.cs
--- a/Engine/Source/Programs/Horde/Horde.Server/Storage/Backends/AwsStorageBackend.cs
+++ b/Engine/Source/Programs/Horde/Horde.Server/Storage/Backends/AwsStorageBackend.cs
@@ -209,12 +209,14 @@ namespace Horde.Server.Storage.Backends
 		class WrappedResponseStream : Stream
 		{
 			readonly IDisposable _semaphore;
+			readonly TelemetrySpan _semaphoreSpan;
 			readonly GetObjectResponse _response;
 			readonly Stream _responseStream;
 
-			public WrappedResponseStream(IDisposable semaphore, GetObjectResponse response)
+			public WrappedResponseStream(IDisposable semaphore, TelemetrySpan semaphoreSpan, GetObjectResponse response)
 			{
 				_semaphore = semaphore;
+				_semaphoreSpan = semaphoreSpan;
 				_response = response;
 				_responseStream = response.ResponseStream;
 			}
@@ -239,9 +241,13 @@ namespace Horde.Server.Storage.Backends
 			{
 				base.Dispose(disposing);
 
-				_semaphore.Dispose();
-				_response.Dispose();
-				_responseStream.Dispose();
+				if (disposing)
+				{
+					_semaphore.Dispose();
+					_semaphoreSpan.Dispose();
+					_response.Dispose();
+					_responseStream.Dispose();
+				}
 			}
 
 			public override async ValueTask DisposeAsync()
@@ -276,11 +282,15 @@ namespace Horde.Server.Storage.Backends
 			string fullPath = GetFullPath(path);
 
 			IDisposable? semaLock = null;
+			TelemetrySpan? semaphoreSpan = null;
 			GetObjectResponse? response = null;
 			try
 			{
 				semaLock = await _semaphore.UseWaitAsync(cancellationToken);
 
+				semaphoreSpan = OpenTelemetryTracers.Horde.StartActiveSpan($"{nameof(AwsStorageBackend)}.{nameof(ReadAsync)}.Semaphore");
+				semaphoreSpan.SetAttribute("path", path);
+
 				GetObjectRequest newGetRequest = new GetObjectRequest();
 				newGetRequest.BucketName = _options.AwsBucketName;
 				newGetRequest.Key = fullPath;
@@ -288,7 +298,7 @@ namespace Horde.Server.Storage.Backends
 
 				response = await _client.GetObjectAsync(newGetRequest, cancellationToken);
 
-				return new WrappedResponseStream(semaLock, response);
+				return new WrappedResponseStream(semaLock, semaphoreSpan, response);
 			}
 			catch (Exception ex)
 			{
@@ -296,6 +306,7 @@ namespace Horde.Server.Storage.Backends
 
 				semaLock?.Dispose();
 				response?.Dispose();
+				semaphoreSpan?.Dispose();
 
 				throw new StorageException($"Unable to read {fullPath} from {_options.AwsBucketName}", ex);
 			}
