Horde: Rethrow any exception from UpdateSession to ensure gRPC stream is closed

diff --git a/Engine/Source/Programs/Horde/Horde.Server/Server/RpcService.cs b/Engine/Source/Programs/Horde/Horde.Server/Server/RpcService.cs
--- a/Engine/Source/Programs/Horde/Horde.Server/Server/RpcService.cs
+++ b/Engine/Source/Programs/Horde/Horde.Server/Server/RpcService.cs
@@ -393,6 +393,7 @@ namespace Horde.Server.Server
 					catch (Exception ex)
 					{
 						_logger.LogError(ex, "Swallowed exception while updating session for {AgentId}.", request.AgentId);
+						throw new StructuredRpcException(StatusCode.Internal, "Failed updating session. Reason: {Reason}", ex.Message);
 					}
 				}
 
