Horde: Add disable reason to agent modal, also update the agent so the modal displays the current agent data once updated
#rnx
#skipci

diff --git a/Engine/Source/Programs/Horde/HordeDashboard/src/components/HistoryModal.tsx b/Engine/Source/Programs/Horde/HordeDashboard/src/components/HistoryModal.tsx
--- a/Engine/Source/Programs/Horde/HordeDashboard/src/components/HistoryModal.tsx
+++ b/Engine/Source/Programs/Horde/HordeDashboard/src/components/HistoryModal.tsx
@@ -1,6 +1,6 @@
 // Copyright Epic Games, Inc. All Rights Reserved.
 
-import { ConstrainMode, DefaultButton, DetailsHeader, DetailsList, DetailsListLayoutMode, DetailsRow, Dialog, DialogFooter, DialogType, GroupedList, GroupHeader, IColumn, IContextualMenuItem, IContextualMenuProps, IDetailsHeaderProps, IDetailsHeaderStyles, IDetailsListProps, IGroup, ITooltipHostStyles, mergeStyleSets, Modal, Pivot, PivotItem, PrimaryButton, ScrollablePane, ScrollbarVisibility, Selection, SelectionMode, Spinner, SpinnerSize, Stack, Sticky, StickyPositionType, Text } from "@fluentui/react";
+import { ConstrainMode, DefaultButton, DetailsHeader, DetailsList, DetailsListLayoutMode, DetailsRow, Dialog, DialogFooter, DialogType, GroupedList, GroupHeader, IColumn, IContextualMenuItem, IContextualMenuProps, IDetailsHeaderProps, IDetailsHeaderStyles, IDetailsListProps, IGroup, ITextField, ITooltipHostStyles, mergeStyleSets, Modal, Pivot, PivotItem, PrimaryButton, ScrollablePane, ScrollbarVisibility, Selection, SelectionMode, Spinner, SpinnerSize, Stack, Sticky, StickyPositionType, Text, TextField } from "@fluentui/react";
 import { action, makeObservable, observable } from "mobx";
 import { observer } from "mobx-react-lite";
 import React, { useState } from "react";
@@ -251,57 +251,6 @@ class HistoryModalState {
       });
    }
 
-   private _sortData(a: any, b: any) {
-      if (this.mode === "leases") {
-         let left = a as LeaseData, right = b as LeaseData;
-         if (this.sortedLeaseColumnDescending) {
-            left = b;
-            right = a;
-         }
-         switch (this.sortedLeaseColumn) {
-            case 'type':
-               return left.type.localeCompare(right.type);
-            case 'name':
-               if (left.name && right.name) {
-                  left.name.localeCompare(right.name);
-               }
-               break;
-            case 'executing':
-               return Number(left.executing) - Number(right.executing);
-            case 'startTime':
-               return (left.startTime as Date).getTime() - (right.startTime as Date).getTime();
-            case 'endTime':
-               if (left.finishTime && right.finishTime) {
-                  return (left.finishTime as Date).getTime() - (right.finishTime as Date).getTime();
-               }
-               break;
-            default:
-               return 0;
-         }
-      }
-      else if (this.mode === "sessions") {
-         let left = a as SessionData, right = b as SessionData;
-         if (this.sortedSessionColumnDescending) {
-            left = b;
-            right = a;
-         }
-         switch (this.sortedSessionColumn) {
-            case 'id':
-               return left.id.localeCompare(right.id);
-            case 'startTime':
-               return (left.startTime as Date).getTime() - (right.startTime as Date).getTime();
-            case 'endTime':
-               if (left.finishTime && right.finishTime) {
-                  return (left.finishTime as Date).getTime() - (right.finishTime as Date).getTime();
-               }
-               break;
-            default:
-               return 0;
-         }
-      }
-      return 0;
-   }
-
    @action
    appendData(newData: any[]) {
       // if there's any data, there might be more data next time, so add another callback.
@@ -330,7 +279,8 @@ const state = new HistoryModalState();
 export const HistoryModal: React.FC<{ agentId: string | undefined, onDismiss: (...args: any[]) => any; }> = observer(({ agentId, onDismiss }) => {
 
    const [selectedAgent, setSelectedAgent] = useState<string | undefined>(undefined);
-   const [actionState, setActionState] = useState<{ action?: string, confirmed?: boolean }>({});
+   const [actionState, setActionState] = useState<{ action?: string, confirmed?: boolean, comment?: string }>({});
+   const actionTextInputRef = React.useRef<ITextField>(null);
 
    //  subscribe to updates
    if (state.selectedAgent) { }
@@ -371,10 +321,9 @@ export const HistoryModal: React.FC<{ agentId: string | undefined, onDismiss: (.
    type AgentAction = {
       name: string;
       confirmText: string;
-      update?: (request: UpdateAgentRequest) => void;
+      textInput?: boolean;
+      update?: (request: UpdateAgentRequest, comment?: string) => void;
    }
-
-
    const actions: AgentAction[] = [
       {
          name: 'Enable',
@@ -384,7 +333,8 @@ export const HistoryModal: React.FC<{ agentId: string | undefined, onDismiss: (.
       {
          name: 'Disable',
          confirmText: "Are you sure you would like to disable this agent?",
-         update: (request) => { request.enabled = false }
+         textInput: true,
+         update: (request, comment) => { request.enabled = false;  request.comment = comment }
       },
       {
          name: 'Cancel Leases',
@@ -451,9 +401,12 @@ export const HistoryModal: React.FC<{ agentId: string | undefined, onDismiss: (.
    if (currentAction && actionState.confirmed) {
       if (currentAction.update) {
          const request: UpdateAgentRequest = {};
-         currentAction.update(request);
-         backend.updateAgent(agentId, request).then(() => {
-            setActionState({});
+         currentAction.update(request, actionState.comment);
+         backend.updateAgent(agentId, request).then(() => {            
+            agentStore.update(agentStore.pools?.length ? true : false).then(() => {
+               state.setSelectedAgent(agentStore.agents.find(agent => agent.id === agentId));
+               setActionState({});
+            });                  
          }).catch((reason) => {
             console.error(reason);
          });
@@ -493,25 +446,6 @@ export const HistoryModal: React.FC<{ agentId: string | undefined, onDismiss: (.
       }
    })
 
-   /*
-   const DayPickerStrings: IDatePickerStrings = {
-        months: ['January', 'February', 'March', 'April', 'May', 'June', 'July', 'August', 'September', 'October', 'November', 'December'],
-
-        shortMonths: ['Jan', 'Feb', 'Mar', 'Apr', 'May', 'Jun', 'Jul', 'Aug', 'Sep', 'Oct', 'Nov', 'Dec'],
-
-        days: ['Sunday', 'Monday', 'Tuesday', 'Wednesday', 'Thursday', 'Friday', 'Saturday'],
-
-        shortDays: ['S', 'M', 'T', 'W', 'T', 'F', 'S'],
-
-        goToToday: 'Go to today',
-        prevMonthAriaLabel: 'Go to previous month',
-        nextMonthAriaLabel: 'Go to next month',
-        prevYearAriaLabel: 'Go to previous year',
-        nextYearAriaLabel: 'Go to next year',
-        closeButtonAriaLabel: 'Close date picker'
-   };
-   */
-
    function onColumnClick(ev: React.MouseEvent<HTMLElement>, column: IColumn) {
       //historyModalState.setSorted(column.key);
    }
@@ -835,8 +769,9 @@ export const HistoryModal: React.FC<{ agentId: string | undefined, onDismiss: (.
                <Stack style={{ paddingBottom: 18, paddingLeft: 4 }}>
                   <Text>{currentAction.confirmText}</Text>
                </Stack>
+               {!!currentAction.textInput && <TextField componentRef={actionTextInputRef} label={"Disable Reason"} />}
                <DialogFooter>
-                  <PrimaryButton disabled={actionState.confirmed} onClick={() => { setActionState({ ...actionState, confirmed: true }) }} text={currentAction.name} />
+                  <PrimaryButton disabled={actionState.confirmed} onClick={() => { setActionState({ ...actionState, confirmed: true, comment: actionTextInputRef.current?.value }) }} text={currentAction.name} />
                   <DefaultButton disabled={actionState.confirmed} onClick={() => { setActionState({}) }} text="Cancel" />
                </DialogFooter>
             </Dialog>
