Horde: Cancel running bisection jobs when canceling a bisection
#rnx

diff --git a/Engine/Source/Programs/Horde/Horde.Server/Jobs/Bisect/BisectController.cs b/Engine/Source/Programs/Horde/Horde.Server/Jobs/Bisect/BisectController.cs
--- a/Engine/Source/Programs/Horde/Horde.Server/Jobs/Bisect/BisectController.cs
+++ b/Engine/Source/Programs/Horde/Horde.Server/Jobs/Bisect/BisectController.cs
@@ -160,6 +160,7 @@ namespace Horde.Server.Jobs.Bisect
 	{
 		readonly IBisectTaskCollection _bisectTaskCollection;
 		readonly IJobCollection _jobCollection;
+		readonly JobService _jobService;
 		readonly IJobStepRefCollection _jobStepRefs;
 		readonly IGraphCollection _graphCollection;
 		readonly IUserCollection _userCollection;
@@ -168,9 +169,10 @@ namespace Horde.Server.Jobs.Bisect
 		/// <summary>
 		/// Constructor
 		/// </summary>
-		public BisectTasksController(IBisectTaskCollection bisectTaskCollection, IJobCollection jobCollection, IJobStepRefCollection jobStepRefs, IGraphCollection graphCollection, IUserCollection userCollection, IOptionsSnapshot<GlobalConfig> globalConfig)
+		public BisectTasksController(IBisectTaskCollection bisectTaskCollection, JobService jobService, IJobCollection jobCollection, IJobStepRefCollection jobStepRefs, IGraphCollection graphCollection, IUserCollection userCollection, IOptionsSnapshot<GlobalConfig> globalConfig)
 		{
 			_bisectTaskCollection = bisectTaskCollection;
+			_jobService = jobService;
 			_jobCollection = jobCollection;
 			_jobStepRefs = jobStepRefs;
 			_graphCollection = graphCollection;
@@ -316,6 +318,15 @@ namespace Horde.Server.Jobs.Bisect
 				IBisectTask? updatedTask = await _bisectTaskCollection.TryUpdateAsync(bisectTask, options, cancellationToken);
 				if (updatedTask != null)
 				{
+					// cancel any running bisection jobs
+					if (bisectTask.State != BisectTaskState.Cancelled && updatedTask.State == BisectTaskState.Cancelled)
+					{
+						IJob? existingJob = await _jobCollection.FindBisectTaskJobsAsync(bisectTask.Id, true, cancellationToken).FirstOrDefaultAsync(cancellationToken);
+						if (existingJob != null && existingJob.AbortedByUserId == null)
+						{
+							await _jobService.UpdateJobAsync(existingJob, null, null, null, KnownUsers.System, null, null, null);
+						}						
+					}
 					return Ok();
 				}
 			}
