Horde: Add job name to the job filter keywords
#rnx
#skipci

diff --git a/Engine/Source/Programs/Horde/HordeDashboard/src/base/utilities/filter.ts b/Engine/Source/Programs/Horde/HordeDashboard/src/base/utilities/filter.ts
--- a/Engine/Source/Programs/Horde/HordeDashboard/src/base/utilities/filter.ts
+++ b/Engine/Source/Programs/Horde/HordeDashboard/src/base/utilities/filter.ts
@@ -12,6 +12,10 @@ const getJobKeywords = (job: GetJobResponse): string[] => {
 
     let keywords: string[] = [];
 
+    if (job.name) {
+        keywords.push(job.name);
+    }
+
     if (job.change) {
         keywords.push(job.change.toString());
     }
