Horde: Fix server tests build error
#rnx

diff --git a/Engine/Source/Programs/Horde/Horde.Server.Tests/TestSetup.cs b/Engine/Source/Programs/Horde/Horde.Server.Tests/TestSetup.cs
--- a/Engine/Source/Programs/Horde/Horde.Server.Tests/TestSetup.cs
+++ b/Engine/Source/Programs/Horde/Horde.Server.Tests/TestSetup.cs
@@ -317,7 +317,7 @@ namespace Horde.Server.Tests
 
 		private BisectTasksController GetBisectTasksController()
 		{
-			BisectTasksController bisectCtrl = new BisectTasksController(BisectTaskCollection, JobCollection, JobStepRefCollection, GraphCollection, UserCollection, GlobalConfigSnapshot);
+			BisectTasksController bisectCtrl = new BisectTasksController(BisectTaskCollection, JobService, JobCollection, JobStepRefCollection, GraphCollection, UserCollection, GlobalConfigSnapshot);
 			bisectCtrl.ControllerContext = GetControllerContext();
 			return bisectCtrl;
 		}
