Horde: Set user who canceled a bisection as the one who canceled any running bisection jobs
#rnx

diff --git a/Engine/Source/Programs/Horde/Horde.Server/Jobs/Bisect/BisectController.cs b/Engine/Source/Programs/Horde/Horde.Server/Jobs/Bisect/BisectController.cs
--- a/Engine/Source/Programs/Horde/Horde.Server/Jobs/Bisect/BisectController.cs
+++ b/Engine/Source/Programs/Horde/Horde.Server/Jobs/Bisect/BisectController.cs
@@ -324,7 +324,7 @@ namespace Horde.Server.Jobs.Bisect
 						IJob? existingJob = await _jobCollection.FindBisectTaskJobsAsync(bisectTask.Id, true, cancellationToken).FirstOrDefaultAsync(cancellationToken);
 						if (existingJob != null && existingJob.AbortedByUserId == null)
 						{
-							await _jobService.UpdateJobAsync(existingJob, null, null, null, KnownUsers.System, null, null, null);
+							await _jobService.UpdateJobAsync(existingJob, null, null, null,  User.GetUserId() ?? KnownUsers.System, null, null, null);
 						}						
 					}
 					return Ok();
