Horde: Add a more direct error message when an agent is already executing a new session.

diff --git a/Engine/Source/Programs/Horde/Horde.Server/Server/RpcService.cs b/Engine/Source/Programs/Horde/Horde.Server/Server/RpcService.cs
--- a/Engine/Source/Programs/Horde/Horde.Server/Server/RpcService.cs
+++ b/Engine/Source/Programs/Horde/Horde.Server/Server/RpcService.cs
@@ -371,7 +371,13 @@ namespace Horde.Server.Server
 				IAgent? agent = await _agentService.GetAgentAsync(new AgentId(request.AgentId));
 				if (agent != null)
 				{
+					SessionId sessionId = SessionId.Parse(request.SessionId);
+
 					// Check we're authorized to update it
+					if (agent.SessionId != sessionId)
+					{
+						throw new StructuredRpcException(StatusCode.PermissionDenied, "Agent {AgentId} has completed session {SessionId}; now executing session {NewSessionId}. Cannot update state.", request.AgentId, sessionId, agent.SessionId?.ToString() ?? "(None)");
+					}
 					if (!_agentService.AuthorizeSession(agent, context.GetHttpContext().User))
 					{
 						throw new StructuredRpcException(StatusCode.PermissionDenied, "Not authenticated for {AgentId}", request.AgentId);
@@ -388,7 +394,7 @@ namespace Horde.Server.Server
 					// Update the session
 					try
 					{
-						agent = await _agentService.UpdateSessionWithWaitAsync(agent, SessionId.Parse(request.SessionId), request.Status, properties, resources, request.Leases, cancellationSource.Token);
+						agent = await _agentService.UpdateSessionWithWaitAsync(agent, sessionId, request.Status, properties, resources, request.Leases, cancellationSource.Token);
 					}
 					catch (Exception ex)
 					{
