Horde: Remove legacy trends and all of the hoisted fluent charting dependency
#rnx
#skipci

diff --git a/Engine/Source/Programs/Horde/HordeDashboard/src/base/components/HordeFluentCharting/HCartesianChart/HCartesianChart.base.tsx b/Engine/Source/Programs/Horde/HordeDashboard/src/base/components/HordeFluentCharting/HCartesianChart/HCartesianChart.base.tsx
deleted file mode 100644
diff --git a/Engine/Source/Programs/Horde/HordeDashboard/src/base/components/HordeFluentCharting/HCartesianChart/HCartesianChart.tsx b/Engine/Source/Programs/Horde/HordeDashboard/src/base/components/HordeFluentCharting/HCartesianChart/HCartesianChart.tsx
deleted file mode 100644
diff --git a/Engine/Source/Programs/Horde/HordeDashboard/src/base/components/HordeFluentCharting/HCartesianChart/HCartesianChart.types.ts b/Engine/Source/Programs/Horde/HordeDashboard/src/base/components/HordeFluentCharting/HCartesianChart/HCartesianChart.types.ts
deleted file mode 100644
diff --git a/Engine/Source/Programs/Horde/HordeDashboard/src/base/components/HordeFluentCharting/HHoverCard/HHoverCard.base.tsx b/Engine/Source/Programs/Horde/HordeDashboard/src/base/components/HordeFluentCharting/HHoverCard/HHoverCard.base.tsx
deleted file mode 100644
diff --git a/Engine/Source/Programs/Horde/HordeDashboard/src/base/components/HordeFluentCharting/HHoverCard/HHoverCard.styles.ts b/Engine/Source/Programs/Horde/HordeDashboard/src/base/components/HordeFluentCharting/HHoverCard/HHoverCard.styles.ts
deleted file mode 100644
diff --git a/Engine/Source/Programs/Horde/HordeDashboard/src/base/components/HordeFluentCharting/HHoverCard/HHoverCard.tsx b/Engine/Source/Programs/Horde/HordeDashboard/src/base/components/HordeFluentCharting/HHoverCard/HHoverCard.tsx
deleted file mode 100644
diff --git a/Engine/Source/Programs/Horde/HordeDashboard/src/base/components/HordeFluentCharting/HHoverCard/HHoverCard.types.ts b/Engine/Source/Programs/Horde/HordeDashboard/src/base/components/HordeFluentCharting/HHoverCard/HHoverCard.types.ts
deleted file mode 100644
diff --git a/Engine/Source/Programs/Horde/HordeDashboard/src/base/components/HordeFluentCharting/HLineChart/HLineChart.base.tsx b/Engine/Source/Programs/Horde/HordeDashboard/src/base/components/HordeFluentCharting/HLineChart/HLineChart.base.tsx
deleted file mode 100644
diff --git a/Engine/Source/Programs/Horde/HordeDashboard/src/base/components/HordeFluentCharting/HLineChart/HLineChart.tsx b/Engine/Source/Programs/Horde/HordeDashboard/src/base/components/HordeFluentCharting/HLineChart/HLineChart.tsx
deleted file mode 100644
diff --git a/Engine/Source/Programs/Horde/HordeDashboard/src/base/components/HordeFluentCharting/HLineChart/HLineChart.types.ts b/Engine/Source/Programs/Horde/HordeDashboard/src/base/components/HordeFluentCharting/HLineChart/HLineChart.types.ts
deleted file mode 100644
diff --git a/Engine/Source/Programs/Horde/HordeDashboard/src/base/components/HordeFluentCharting/Hutilities.ts b/Engine/Source/Programs/Horde/HordeDashboard/src/base/components/HordeFluentCharting/Hutilities.ts
deleted file mode 100644
diff --git a/Engine/Source/Programs/Horde/HordeDashboard/src/components/jobDetailsV2/JobDetailStepTrends.tsx b/Engine/Source/Programs/Horde/HordeDashboard/src/components/jobDetailsV2/JobDetailStepTrends.tsx
deleted file mode 100644
diff --git a/Engine/Source/Programs/Horde/HordeDashboard/src/components/jobDetailsV2/JobDetailViewV2.tsx b/Engine/Source/Programs/Horde/HordeDashboard/src/components/jobDetailsV2/JobDetailViewV2.tsx
--- a/Engine/Source/Programs/Horde/HordeDashboard/src/components/jobDetailsV2/JobDetailViewV2.tsx
+++ b/Engine/Source/Programs/Horde/HordeDashboard/src/components/jobDetailsV2/JobDetailViewV2.tsx
@@ -6,11 +6,9 @@ import { GetJobsTabResponse } from "../../backend/Api";
 import { useWindowSize } from "../../base/utilities/hooks";
 import { hordeClasses, modeColors } from "../../styles/Styles";
 import { BreadcrumbItem, Breadcrumbs } from "../Breadcrumbs";
-import ErrorHandler from "../ErrorHandler";
 import { useQuery } from "../JobDetailCommon";
 import { TopNav } from "../TopNav";
 import { HealthPanel } from "./JobDetailHealthV2";
-import { StepTrendsPanel } from "./JobDetailStepTrends";
 import { JobDataView, JobDetailsV2 } from "./JobDetailsViewCommon";
 import { TimelinePanel } from "./JobDetailTimeline";
 import { StepsPanelV2 } from "./JobDetailViewSteps";
@@ -305,8 +303,6 @@ const DetailsViewOverview: React.FC<{ jobDetails: JobDetailsV2 }> = ({ jobDetail
       <StepsPanelV2 jobDetails={details} />
       <HealthPanel jobDetails={details} />
       <TimelinePanel jobDetails={details} />
-      { false &&<StepTrendsPanel jobDetails={details} /> }
-
    </Stack>
 };
 
diff --git a/Engine/Source/Programs/Horde/HordeDashboard/src/components/jobDetailsV2/StepDetailView.tsx b/Engine/Source/Programs/Horde/HordeDashboard/src/components/jobDetailsV2/StepDetailView.tsx
--- a/Engine/Source/Programs/Horde/HordeDashboard/src/components/jobDetailsV2/StepDetailView.tsx
+++ b/Engine/Source/Programs/Horde/HordeDashboard/src/components/jobDetailsV2/StepDetailView.tsx
@@ -10,7 +10,6 @@ import { useQuery } from "../JobDetailCommon";
 import { JobDetailArtifactsV2 } from "./JobDetailArtifactsV2";
 import { HealthPanel } from "./JobDetailHealthV2";
 import { StepHistoryPanel } from "./JobDetailStepHistory";
-import { StepTrendsPanel } from "./JobDetailStepTrends";
 import { JobDataView, JobDetailsV2 } from "./JobDetailsViewCommon";
 import { TimelinePanel } from "./JobDetailTimeline";
 import { StepsPanelV2 } from "./JobDetailViewSteps";
