Horde: Fix initial device report service tick
#rnx

diff --git a/Engine/Source/Programs/Horde/Horde.Server/Devices/DeviceReportService.cs b/Engine/Source/Programs/Horde/Horde.Server/Devices/DeviceReportService.cs
--- a/Engine/Source/Programs/Horde/Horde.Server/Devices/DeviceReportService.cs
+++ b/Engine/Source/Programs/Horde/Horde.Server/Devices/DeviceReportService.cs
@@ -189,7 +189,7 @@ namespace Horde.Server.Devices
 
 			DeviceReportState state = await _state.GetAsync();
 			DateTime currentTime = _clock.UtcNow;
-			DateTime lastReportTime = state.ReportTime == DateTime.MinValue ? DateTime.Now : state.ReportTime;
+			DateTime lastReportTime = state.ReportTime == DateTime.MinValue ? DateTime.Now.Subtract(TimeSpan.FromMinutes(_reportIntervalMinutes + 1)) : state.ReportTime;
 
 			if ((currentTime - lastReportTime).TotalMinutes <= _reportIntervalMinutes)
 			{
