Horde: Disable child lease termination; suspect may be causing gRPC deadline warnings.

diff --git a/Engine/Source/Programs/Horde/Horde.Server/Agents/AgentService.cs b/Engine/Source/Programs/Horde/Horde.Server/Agents/AgentService.cs
--- a/Engine/Source/Programs/Horde/Horde.Server/Agents/AgentService.cs
+++ b/Engine/Source/Programs/Horde/Horde.Server/Agents/AgentService.cs
@@ -899,6 +899,8 @@ namespace Horde.Server.Agents
 			// Update the lease
 			await _leases.TrySetOutcomeAsync(lease.Id, finishTime, outcome, output);
 
+			// Temporarily disabling due to gRPC timeouts.
+#if false
 			// Terminate any child leases
 			List<ILease> childLeases = await _leases.FindLeasesAsync(parentId: lease.Id);
 			foreach (ILease childLease in childLeases)
@@ -917,6 +919,7 @@ namespace Horde.Server.Agents
 					}
 				}
 			}
+#endif
 		}
 
 		/// <summary>
