Horde: Fix issue with notice and breadcrumb bar pointer events
#rnx
#skipci

diff --git a/Engine/Source/Programs/Horde/HordeDashboard/src/components/Breadcrumbs.tsx b/Engine/Source/Programs/Horde/HordeDashboard/src/components/Breadcrumbs.tsx
--- a/Engine/Source/Programs/Horde/HordeDashboard/src/components/Breadcrumbs.tsx
+++ b/Engine/Source/Programs/Horde/HordeDashboard/src/components/Breadcrumbs.tsx
@@ -117,13 +117,13 @@ export const Breadcrumbs: React.FC<{ items: BreadcrumbItem[], title?: string, su
             </Stack>
          </Stack>}
 
-         {!!alert && dashboard.alertSquelch !== alert && <Stack horizontalAlign="center" disableShrink={true} style={{ position: "absolute", width: "100%" }} >
+         {!!alert && dashboard.alertSquelch !== alert && <Stack horizontalAlign="center" disableShrink={true} style={{ position: "absolute", width: "100%", pointerEvents:"none" }} >
             <Stack horizontal>
                <Stack grow />
-               <Stack>
+               <Stack style={{pointerEvents: "auto"}}>
                   <MessageBar onDismiss={() => { dashboard.alertSquelch = alert; setHideAlert(true) }}
                      messageBarType={MessageBarType.severeWarning} isMultiline={false} >
-                     <Text variant={"small"} style={{ fontFamily: "Horde Open Sans Bold" }}>{alert}</Text>
+                     <Text variant={"small"} style={{ fontFamily: "Horde Open Sans Bold"}}>{alert}</Text>
                   </MessageBar>
                </Stack>
                <Stack grow />
