Horde: Ignore null values for badge names in UGS controller.

diff --git a/Engine/Source/Programs/Horde/Horde.Server/Ugs/UgsController.cs b/Engine/Source/Programs/Horde/Horde.Server/Ugs/UgsController.cs
--- a/Engine/Source/Programs/Horde/Horde.Server/Ugs/UgsController.cs
+++ b/Engine/Source/Programs/Horde/Horde.Server/Ugs/UgsController.cs
@@ -99,7 +99,10 @@ namespace Horde.Server.Ugs
 			{
 				foreach (AddUgsBadgeRequest badge in request.Badges)
 				{
-					metadata = await _ugsMetadataCollection.UpdateBadgeAsync(metadata, badge.Name, badge.Url, badge.State);
+					if (!String.IsNullOrEmpty(badge.Name))
+					{
+						metadata = await _ugsMetadataCollection.UpdateBadgeAsync(metadata, badge.Name, badge.Url, badge.State);
+					}
 				}
 			}
 			return Ok();
