Horde: Improve logging around index creation errors, and remove partial index over lease parents.

diff --git a/Engine/Source/Programs/Horde/Horde.Server/Agents/Leases/LeaseCollection.cs b/Engine/Source/Programs/Horde/Horde.Server/Agents/Leases/LeaseCollection.cs
--- a/Engine/Source/Programs/Horde/Horde.Server/Agents/Leases/LeaseCollection.cs
+++ b/Engine/Source/Programs/Horde/Horde.Server/Agents/Leases/LeaseCollection.cs
@@ -80,7 +80,7 @@ namespace Horde.Server.Agents.Leases
 		public LeaseCollection(MongoService mongoService)
 		{
 			List<MongoIndex<LeaseDocument>> indexes = new List<MongoIndex<LeaseDocument>>();
-			indexes.Add(keys => keys.Ascending(x => x.ParentId), sparse: true);
+//			indexes.Add(keys => keys.Ascending(x => x.ParentId), sparse: true);
 			indexes.Add(keys => keys.Ascending(x => x.AgentId));
 			indexes.Add(keys => keys.Ascending(x => x.SessionId));
 			indexes.Add(keys => keys.Ascending(x => x.StartTime));
diff --git a/Engine/Source/Programs/Horde/Horde.Server/Server/MongoService.cs b/Engine/Source/Programs/Horde/Horde.Server/Server/MongoService.cs
--- a/Engine/Source/Programs/Horde/Horde.Server/Server/MongoService.cs
+++ b/Engine/Source/Programs/Horde/Horde.Server/Server/MongoService.cs
@@ -742,8 +742,16 @@ namespace Horde.Server.Server
 
 						CreateIndexModel<T> model = new CreateIndexModel<T>(createIndex.Keys, options);
 
-						string result = await collection.Indexes.CreateOneAsync(model, cancellationToken: cancellationToken);
-						_logger.LogInformation("Created index {IndexName}", result);
+						try
+						{
+							string result = await collection.Indexes.CreateOneAsync(model, cancellationToken: cancellationToken);
+							_logger.LogInformation("Created index {IndexName}", result);
+						}
+						catch (Exception ex)
+						{
+							_logger.LogError(ex, "Unable to create index {IndexName}: {Message}", createIndex.Name, ex.Message);
+							throw;
+						}
 					}
 				}
 
@@ -981,7 +989,7 @@ namespace Horde.Server.Server
 				}
 				catch (Exception ex)
 				{
-					_logger.LogWarning(ex, "Exception updating indexes: {Message}", ex.Message);
+					_logger.LogError(ex, "Exception updating indexes: {Message}", ex.Message);
 					await Task.Delay(TimeSpan.FromSeconds(30.0), cancellationToken);
 				}
 			}
@@ -998,11 +1006,11 @@ namespace Horde.Server.Server
 			}
 			catch (MongoCommandException ex)
 			{
-				_logger.LogWarning(ex, "Command exception while attempting to update indexes ({Code})", ex.Code);
+				_logger.LogError(ex, "Command exception while attempting to update indexes ({Code}): {Message}", ex.Code, ex.Message);
 			}
 			catch (Exception ex)
 			{
-				_logger.LogWarning(ex, "Exception while attempting to update indexes");
+				_logger.LogError(ex, "Exception while attempting to update indexes: {Message}", ex.Message);
 			}
 		}
 
@@ -1017,7 +1025,7 @@ namespace Horde.Server.Server
 					SemVer currentVersion = SemVer.Parse(currentSchema.Version);
 					if (schemaVersion < currentVersion)
 					{
-						_logger.LogDebug("Ignoring upgrade command; server is older than current schema version ({ProgramVer} < {CurrentVer})", Program.Version, currentVersion);
+						_logger.LogInformation("Ignoring upgrade command; server is older than current schema version ({ProgramVer} < {CurrentVer})", Program.Version, currentVersion);
 						return false;
 					}
 					if (schemaVersion == currentVersion)
@@ -1026,7 +1034,7 @@ namespace Horde.Server.Server
 					}
 				}
 
-				_logger.LogDebug("Upgrading schema version {OldVersion} -> {NewVersion}", currentSchema.Version, schemaVersion.ToString());
+				_logger.LogInformation("Upgrading schema version {OldVersion} -> {NewVersion}", currentSchema.Version, schemaVersion.ToString());
 				currentSchema.Version = schemaVersion.ToString();
 
 				if (_mongoService.ReadOnlyMode)
