Horde: Ensure that all index updates have completed before performing operations on MongoDB collections.

The presence of an index can change whether updates succeed or fail (eg. by introducing uniqueness constraints).

diff --git a/Engine/Source/Programs/Horde/Horde.Server/Server/MongoService.cs b/Engine/Source/Programs/Horde/Horde.Server/Server/MongoService.cs
--- a/Engine/Source/Programs/Horde/Horde.Server/Server/MongoService.cs
+++ b/Engine/Source/Programs/Horde/Horde.Server/Server/MongoService.cs
@@ -186,6 +186,11 @@ namespace Horde.Server.Server
 		}
 	}
 
+	/// <summary>
+	/// Task used to update indexes
+	/// </summary>
+	record class MongoUpgradeTask(Func<CancellationToken, Task> UpgradeAsync, TaskCompletionSource CompletionSource);
+
 	/// <summary>
 	/// Singleton for accessing the database
 	/// </summary>
@@ -252,7 +257,7 @@ namespace Horde.Server.Server
 		const int DefaultMongoPort = 27017;
 
 		readonly HashSet<string> _collectionNames = new HashSet<string>(StringComparer.Ordinal);
-		readonly Channel<Func<CancellationToken, Task>> _updateIndexesChannel = Channel.CreateUnbounded<Func<CancellationToken, Task>>();
+		readonly Channel<MongoUpgradeTask> _updateIndexesChannel = Channel.CreateUnbounded<MongoUpgradeTask>();
 
 		/// <summary>
 		/// Constructor
@@ -617,7 +622,7 @@ namespace Horde.Server.Server
 		/// <returns></returns>
 		public IMongoCollection<T> GetCollection<T>(string name)
 		{
-			return new MongoTracingCollection<T>(Database.GetCollection<T>(name), _tracer);
+			return GetCollection<T>(name, Enumerable.Empty<MongoIndex<T>>());
 		}
 
 		/// <summary>
@@ -644,20 +649,29 @@ namespace Horde.Server.Server
 		/// <returns></returns>
 		public IMongoCollection<T> GetCollection<T>(string name, IEnumerable<MongoIndex<T>> indexes)
 		{
-			IMongoCollection<T> collection = GetCollection<T>(name);
-			lock (_collectionNames)
-			{
-				_logger.LogDebug("Queuing update for collection {Name}", name);
+			IMongoCollection<T> collection = Database.GetCollection<T>(name);
 
-				if (!_collectionNames.Add(name))
+			Task upgradeTask = Task.CompletedTask;
+			if (indexes.Any())
+			{
+				lock (_collectionNames)
 				{
-					throw new NotImplementedException();
-				}
+					_logger.LogDebug("Queuing update for collection {Name}", name);
 
-				MongoIndex<T>[] indexesCopy = indexes.ToArray();
-				_updateIndexesChannel.Writer.TryWrite(ctx => UpdateIndexesAsync(name, collection, indexesCopy, ctx));
+					if (!_collectionNames.Add(name))
+					{
+						throw new NotImplementedException();
+					}
+
+					MongoIndex<T>[] indexesCopy = indexes.ToArray();
+
+					TaskCompletionSource tcs = new TaskCompletionSource();
+					_updateIndexesChannel.Writer.TryWrite(new MongoUpgradeTask(ctx => UpdateIndexesAsync(name, collection, indexesCopy, ctx), tcs));
+					upgradeTask = tcs.Task;
+				}
 			}
-			return collection;
+
+			return new MongoTracingCollection<T>(Database.GetCollection<T>(name), upgradeTask, _tracer);
 		}
 
 		/// <summary>
@@ -665,7 +679,7 @@ namespace Horde.Server.Server
 		/// </summary>
 		/// <param name="cancellationToken"></param>
 		/// <returns></returns>
-		public ValueTask<Func<CancellationToken, Task>> ReadNextUpgradeTask(CancellationToken cancellationToken)
+		internal ValueTask<MongoUpgradeTask> ReadNextUpgradeTask(CancellationToken cancellationToken)
 		{
 			return _updateIndexesChannel.Reader.ReadAsync(cancellationToken);
 		}
@@ -965,14 +979,14 @@ namespace Horde.Server.Server
 			{
 				try
 				{
-					Func<CancellationToken, Task> updateIndexTask = await _mongoService.ReadNextUpgradeTask(cancellationToken);
+					MongoUpgradeTask upgradeTask = await _mongoService.ReadNextUpgradeTask(cancellationToken);
 					for (; ; )
 					{
 						using (RedisLock schemaLock = new(_redisService.GetDatabase(), s_schemaLockKey))
 						{
 							if (await schemaLock.AcquireAsync(TimeSpan.FromMinutes(5.0)))
 							{
-								await UpdateOneAsync(updateIndexTask, cancellationToken);
+								await UpdateOneAsync(upgradeTask, cancellationToken);
 								break;
 							}
 							else
@@ -995,22 +1009,25 @@ namespace Horde.Server.Server
 			}
 		}
 
-		async ValueTask UpdateOneAsync(Func<CancellationToken, Task> taskAsync, CancellationToken cancellationToken)
+		async ValueTask UpdateOneAsync(MongoUpgradeTask upgradeTask, CancellationToken cancellationToken)
 		{
 			try
 			{
 				if (await SetSchemaVersion(Program.Version))
 				{
-					await taskAsync(cancellationToken);
+					await upgradeTask.UpgradeAsync(cancellationToken);
 				}
+				upgradeTask.CompletionSource.SetResult();
 			}
 			catch (MongoCommandException ex)
 			{
 				_logger.LogError(ex, "Command exception while attempting to update indexes ({Code}): {Message}", ex.Code, ex.Message);
+				upgradeTask.CompletionSource.TrySetException(ex);
 			}
 			catch (Exception ex)
 			{
 				_logger.LogError(ex, "Exception while attempting to update indexes: {Message}", ex.Message);
+				upgradeTask.CompletionSource.TrySetException(ex);
 			}
 		}
 
diff --git a/Engine/Source/Programs/Horde/Horde.Server/Utilities/MongoTracingCollection.cs b/Engine/Source/Programs/Horde/Horde.Server/Utilities/MongoTracingCollection.cs
--- a/Engine/Source/Programs/Horde/Horde.Server/Utilities/MongoTracingCollection.cs
+++ b/Engine/Source/Programs/Horde/Horde.Server/Utilities/MongoTracingCollection.cs
@@ -3,6 +3,7 @@
 using System.Collections.Generic;
 using System.Threading;
 using System.Threading.Tasks;
+using Microsoft.Extensions.Caching.Memory;
 using MongoDB.Bson;
 using MongoDB.Bson.Serialization;
 using MongoDB.Driver;
@@ -20,16 +21,19 @@ namespace Horde.Server.Utilities
 	public class MongoTracingCollection<T> : IMongoCollection<T>
 	{
 		private readonly IMongoCollection<T> _collection;
+		private readonly Task _upgradeTask;
 		private readonly Tracer _tracer;
 
 		/// <summary>
 		/// Constructor
 		/// </summary>
 		/// <param name="collection">Collection to wrap with tracing</param>
+		/// <param name="indexCreationTask">Task which creates indexes</param>
 		/// <param name="tracer">Tracer</param>
-		public MongoTracingCollection(IMongoCollection<T> collection, Tracer tracer)
+		public MongoTracingCollection(IMongoCollection<T> collection, Task indexCreationTask, Tracer tracer)
 		{
 			_collection = collection;
+			_upgradeTask = indexCreationTask;
 			_tracer = tracer;
 		}
 
@@ -37,18 +41,21 @@ namespace Horde.Server.Utilities
 		/// <inheritdoc />
 		public IAsyncCursor<TResult> Aggregate<TResult>(PipelineDefinition<T, TResult> pipeline, AggregateOptions options = null!, CancellationToken cancellationToken = default)
 		{
+			_upgradeTask.Wait(cancellationToken);
 			return _collection.Aggregate(pipeline, options,  cancellationToken);
 		}
 
 		/// <inheritdoc />
 		public IAsyncCursor<TResult> Aggregate<TResult>(IClientSessionHandle session, PipelineDefinition<T, TResult> pipeline, AggregateOptions options = null!, CancellationToken cancellationToken = default)
 		{
+			_upgradeTask.Wait(cancellationToken);
 			return _collection.Aggregate(session, pipeline, options,  cancellationToken);
 		}
 
 		/// <inheritdoc />
-		public async Task<IAsyncCursor<TResult>> AggregateAsync<TResult>(PipelineDefinition<T, TResult> pipeline, AggregateOptions options = null!, CancellationToken  cancellationToken = default)
+		public async Task<IAsyncCursor<TResult>> AggregateAsync<TResult>(PipelineDefinition<T, TResult> pipeline, AggregateOptions options = null!, CancellationToken cancellationToken = default)
 		{
+			await _upgradeTask.WaitAsync(cancellationToken);
 			using TelemetrySpan span = _tracer.StartMongoDbSpan(nameof(AggregateAsync), _collection);
 			return await _collection.AggregateAsync(pipeline, options,  cancellationToken);
 		}
@@ -57,6 +64,7 @@ namespace Horde.Server.Utilities
 		public async Task<IAsyncCursor<TResult>> AggregateAsync<TResult>(IClientSessionHandle session, PipelineDefinition<T, TResult> pipeline, AggregateOptions options = null!,
 			CancellationToken cancellationToken = default)
 		{
+			await _upgradeTask.WaitAsync(cancellationToken);
 			using TelemetrySpan span = _tracer.StartMongoDbSpan(nameof(AggregateAsync), _collection);
 			return await _collection.AggregateAsync(session, pipeline, options, cancellationToken);
 		}
@@ -64,18 +72,21 @@ namespace Horde.Server.Utilities
 		/// <inheritdoc />
 		public void AggregateToCollection<TResult>(PipelineDefinition<T, TResult> pipeline, AggregateOptions? options = null, CancellationToken cancellationToken = default)
 		{
+			_upgradeTask.Wait(cancellationToken);
 			_collection.AggregateToCollection(pipeline, options, cancellationToken);
 		}
 
 		/// <inheritdoc />
 		public void AggregateToCollection<TResult>(IClientSessionHandle session, PipelineDefinition<T, TResult> pipeline, AggregateOptions? options = null, CancellationToken cancellationToken = default)
 		{
+			_upgradeTask.Wait(cancellationToken);
 			_collection.AggregateToCollection(session, pipeline, options, cancellationToken);
 		}
 
 		/// <inheritdoc />
 		public async Task AggregateToCollectionAsync<TResult>(PipelineDefinition<T, TResult> pipeline, AggregateOptions? options = null, CancellationToken cancellationToken = default)
 		{
+			await _upgradeTask.WaitAsync(cancellationToken);
 			using TelemetrySpan span = _tracer.StartMongoDbSpan(nameof(AggregateToCollectionAsync), _collection);
 			await _collection.AggregateToCollectionAsync(pipeline, options, cancellationToken);
 		}
@@ -83,6 +94,7 @@ namespace Horde.Server.Utilities
 		/// <inheritdoc />
 		public async Task AggregateToCollectionAsync<TResult>(IClientSessionHandle session, PipelineDefinition<T, TResult> pipeline, AggregateOptions? options = null, CancellationToken cancellationToken = default)
 		{
+			await _upgradeTask.WaitAsync(cancellationToken);
 			using TelemetrySpan span = _tracer.StartMongoDbSpan(nameof(AggregateToCollectionAsync), _collection);
 			await _collection.AggregateToCollectionAsync(session, pipeline, options, cancellationToken);
 		}
@@ -90,18 +102,21 @@ namespace Horde.Server.Utilities
 		/// <inheritdoc />
 		public BulkWriteResult<T> BulkWrite(IEnumerable<WriteModel<T>> requests, BulkWriteOptions options = null!, CancellationToken cancellationToken = default)
 		{
+			_upgradeTask.Wait(cancellationToken);
 			return _collection.BulkWrite(requests, options,  cancellationToken);
 		}
 
 		/// <inheritdoc />
 		public BulkWriteResult<T> BulkWrite(IClientSessionHandle session, IEnumerable<WriteModel<T>> requests, BulkWriteOptions options = null!, CancellationToken cancellationToken = default)
 		{
+			_upgradeTask.Wait(cancellationToken);
 			return _collection.BulkWrite(session, requests, options,  cancellationToken);
 		}
 
 		/// <inheritdoc />
 		public async Task<BulkWriteResult<T>> BulkWriteAsync(IEnumerable<WriteModel<T>> requests, BulkWriteOptions options = null!, CancellationToken cancellationToken = default)
 		{
+			await _upgradeTask.WaitAsync(cancellationToken);
 			using TelemetrySpan span = _tracer.StartMongoDbSpan(nameof(BulkWriteAsync), _collection);
 			return await _collection.BulkWriteAsync(requests, options,  cancellationToken);
 		}
@@ -109,6 +124,7 @@ namespace Horde.Server.Utilities
 		/// <inheritdoc />
 		public async Task<BulkWriteResult<T>> BulkWriteAsync(IClientSessionHandle session, IEnumerable<WriteModel<T>> requests, BulkWriteOptions options = null!, CancellationToken cancellationToken = default)
 		{
+			await _upgradeTask.WaitAsync(cancellationToken);
 			using TelemetrySpan span = _tracer.StartMongoDbSpan(nameof(BulkWriteAsync), _collection);
 			return await _collection.BulkWriteAsync(session, requests, options,  cancellationToken);
 		}
@@ -116,18 +132,21 @@ namespace Horde.Server.Utilities
 		/// <inheritdoc />
 		public long Count(FilterDefinition<T> filter, CountOptions options = null!, CancellationToken cancellationToken = default)
 		{
+			_upgradeTask.Wait(cancellationToken);
 			return _collection.Count(filter, options,  cancellationToken);
 		}
 
 		/// <inheritdoc />
 		public long Count(IClientSessionHandle session, FilterDefinition<T> filter, CountOptions options = null!, CancellationToken cancellationToken = default)
 		{
+			_upgradeTask.Wait(cancellationToken);
 			return _collection.Count(session, filter, options,  cancellationToken);
 		}
 
 		/// <inheritdoc />
 		public async Task<long> CountAsync(FilterDefinition<T> filter, CountOptions options = null!, CancellationToken cancellationToken = default)
 		{
+			await _upgradeTask.WaitAsync(cancellationToken);
 			using TelemetrySpan span = _tracer.StartMongoDbSpan(nameof(CountAsync), _collection);
 			return await _collection.CountAsync(filter, options,  cancellationToken);
 		}
@@ -135,6 +154,7 @@ namespace Horde.Server.Utilities
 		/// <inheritdoc />
 		public async Task<long> CountAsync(IClientSessionHandle session, FilterDefinition<T> filter, CountOptions options = null!, CancellationToken cancellationToken = default)
 		{
+			await _upgradeTask.WaitAsync(cancellationToken);
 			using TelemetrySpan span = _tracer.StartMongoDbSpan(nameof(CountAsync), _collection);
 			return await _collection.CountAsync(session, filter, options,  cancellationToken);
 		}
@@ -142,18 +162,21 @@ namespace Horde.Server.Utilities
 		/// <inheritdoc />
 		public long CountDocuments(FilterDefinition<T> filter, CountOptions options = null!, CancellationToken cancellationToken = default)
 		{
+			_upgradeTask.Wait(cancellationToken);
 			return _collection.CountDocuments(filter, options,  cancellationToken);
 		}
 
 		/// <inheritdoc />
 		public long CountDocuments(IClientSessionHandle session, FilterDefinition<T> filter, CountOptions options = null!, CancellationToken cancellationToken = default)
 		{
+			_upgradeTask.Wait(cancellationToken);
 			return _collection.CountDocuments(session, filter, options,  cancellationToken);
 		}
 
 		/// <inheritdoc />
 		public async Task<long> CountDocumentsAsync(FilterDefinition<T> filter, CountOptions options = null!, CancellationToken cancellationToken = default)
 		{
+			await _upgradeTask;
 			using TelemetrySpan span = _tracer.StartMongoDbSpan(nameof(CountDocumentsAsync), _collection);
 			return await _collection.CountDocumentsAsync(filter, options,  cancellationToken);
 		}
@@ -161,6 +184,7 @@ namespace Horde.Server.Utilities
 		/// <inheritdoc />
 		public async Task<long> CountDocumentsAsync(IClientSessionHandle session, FilterDefinition<T> filter, CountOptions options = null!, CancellationToken cancellationToken = default)
 		{
+			await _upgradeTask.WaitAsync(cancellationToken);
 			using TelemetrySpan span = _tracer.StartMongoDbSpan(nameof(CountDocumentsAsync), _collection);
 			return await _collection.CountDocumentsAsync(session, filter, options,  cancellationToken);
 		}
@@ -168,24 +192,28 @@ namespace Horde.Server.Utilities
 		/// <inheritdoc />
 		public DeleteResult DeleteMany(FilterDefinition<T> filter, CancellationToken cancellationToken = default)
 		{
+			_upgradeTask.Wait(cancellationToken);
 			return _collection.DeleteMany(filter,  cancellationToken);
 		}
 
 		/// <inheritdoc />
 		public DeleteResult DeleteMany(FilterDefinition<T> filter, DeleteOptions options, CancellationToken cancellationToken = default)
 		{
+			_upgradeTask.Wait(cancellationToken);
 			return _collection.DeleteMany(filter, options,  cancellationToken);
 		}
 
 		/// <inheritdoc />
 		public DeleteResult DeleteMany(IClientSessionHandle session, FilterDefinition<T> filter, DeleteOptions options = null!, CancellationToken cancellationToken = default)
 		{
+			_upgradeTask.Wait(cancellationToken);
 			return _collection.DeleteMany(session, filter, options,  cancellationToken);
 		}
 
 		/// <inheritdoc />
 		public async Task<DeleteResult> DeleteManyAsync(FilterDefinition<T> filter, CancellationToken cancellationToken = default)
 		{
+			await _upgradeTask.WaitAsync(cancellationToken);
 			using TelemetrySpan span = _tracer.StartMongoDbSpan(nameof(DeleteManyAsync), _collection);
 			return await _collection.DeleteManyAsync(filter,  cancellationToken);
 		}
@@ -193,6 +221,7 @@ namespace Horde.Server.Utilities
 		/// <inheritdoc />
 		public async Task<DeleteResult> DeleteManyAsync(FilterDefinition<T> filter, DeleteOptions options, CancellationToken cancellationToken = default)
 		{
+			await _upgradeTask.WaitAsync(cancellationToken);
 			using TelemetrySpan span = _tracer.StartMongoDbSpan(nameof(DeleteManyAsync), _collection);
 			return await _collection.DeleteManyAsync(filter, options,  cancellationToken);
 		}
@@ -200,6 +229,7 @@ namespace Horde.Server.Utilities
 		/// <inheritdoc />
 		public async Task<DeleteResult> DeleteManyAsync(IClientSessionHandle session, FilterDefinition<T> filter, DeleteOptions options = null!, CancellationToken cancellationToken = default)
 		{
+			await _upgradeTask;
 			using TelemetrySpan span = _tracer.StartMongoDbSpan(nameof(DeleteManyAsync), _collection);
 			return await _collection.DeleteManyAsync(session, filter, options,  cancellationToken);
 		}
@@ -207,24 +237,28 @@ namespace Horde.Server.Utilities
 		/// <inheritdoc />
 		public DeleteResult DeleteOne(FilterDefinition<T> filter, CancellationToken cancellationToken = default)
 		{
+			_upgradeTask.Wait(cancellationToken);
 			return _collection.DeleteOne(filter,  cancellationToken);
 		}
 
 		/// <inheritdoc />
 		public DeleteResult DeleteOne(FilterDefinition<T> filter, DeleteOptions options, CancellationToken cancellationToken = default)
 		{
+			_upgradeTask.Wait(cancellationToken);
 			return _collection.DeleteOne(filter, options,  cancellationToken);
 		}
 
 		/// <inheritdoc />
 		public DeleteResult DeleteOne(IClientSessionHandle session, FilterDefinition<T> filter, DeleteOptions options = null!, CancellationToken cancellationToken = default)
 		{
+			_upgradeTask.Wait(cancellationToken);
 			return _collection.DeleteOne(session, filter, options,  cancellationToken);
 		}
 
 		/// <inheritdoc />
 		public async Task<DeleteResult> DeleteOneAsync(FilterDefinition<T> filter, CancellationToken cancellationToken = default)
 		{
+			await _upgradeTask.WaitAsync(cancellationToken);
 			using TelemetrySpan span = _tracer.StartMongoDbSpan(nameof(DeleteOneAsync), _collection);
 			return await _collection.DeleteOneAsync(filter,  cancellationToken);
 		}
@@ -232,32 +266,37 @@ namespace Horde.Server.Utilities
 		/// <inheritdoc />
 		public async Task<DeleteResult> DeleteOneAsync(FilterDefinition<T> filter, DeleteOptions options, CancellationToken cancellationToken = default)
 		{
+			await _upgradeTask.WaitAsync(cancellationToken);
 			using TelemetrySpan span = _tracer.StartMongoDbSpan(nameof(DeleteOneAsync), _collection);
 			return await _collection.DeleteOneAsync(filter, options,  cancellationToken);
 		}
 
 		/// <inheritdoc />
-		public Task<DeleteResult> DeleteOneAsync(IClientSessionHandle session, FilterDefinition<T> filter, DeleteOptions options = null!, CancellationToken cancellationToken = default)
+		public async Task<DeleteResult> DeleteOneAsync(IClientSessionHandle session, FilterDefinition<T> filter, DeleteOptions options = null!, CancellationToken cancellationToken = default)
 		{
+			await _upgradeTask.WaitAsync(cancellationToken);
 			using TelemetrySpan span = _tracer.StartMongoDbSpan(nameof(DeleteOneAsync), _collection);
-			return _collection.DeleteOneAsync(session, filter, options,  cancellationToken);
+			return await _collection.DeleteOneAsync(session, filter, options,  cancellationToken);
 		}
 
 		/// <inheritdoc />
 		public IAsyncCursor<TField> Distinct<TField>(FieldDefinition<T, TField> field, FilterDefinition<T> filter, DistinctOptions options = null!, CancellationToken cancellationToken = default)
 		{
+			_upgradeTask.Wait(cancellationToken);
 			return _collection.Distinct(field, filter, options,  cancellationToken);
 		}
 
 		/// <inheritdoc />
 		public IAsyncCursor<TField> Distinct<TField>(IClientSessionHandle session, FieldDefinition<T, TField> field, FilterDefinition<T> filter, DistinctOptions options = null!, CancellationToken cancellationToken = default)
 		{
+			_upgradeTask.Wait(cancellationToken);
 			return _collection.Distinct(session, field, filter, options,  cancellationToken);
 		}
 
 		/// <inheritdoc />
 		public async Task<IAsyncCursor<TField>> DistinctAsync<TField>(FieldDefinition<T, TField> field, FilterDefinition<T> filter, DistinctOptions options = null!, CancellationToken cancellationToken = default)
 		{
+			await _upgradeTask.WaitAsync(cancellationToken);
 			using TelemetrySpan span = _tracer.StartMongoDbSpan(nameof(DistinctAsync), _collection);
 			return await _collection.DistinctAsync(field, filter, options,  cancellationToken);
 		}
@@ -265,6 +304,7 @@ namespace Horde.Server.Utilities
 		/// <inheritdoc />
 		public async Task<IAsyncCursor<TField>> DistinctAsync<TField>(IClientSessionHandle session, FieldDefinition<T, TField> field, FilterDefinition<T> filter, DistinctOptions options = null!, CancellationToken cancellationToken = default)
 		{
+			await _upgradeTask.WaitAsync(cancellationToken);
 			using TelemetrySpan span = _tracer.StartMongoDbSpan(nameof(DistinctAsync), _collection);
 			return await _collection.DistinctAsync(session, field, filter, options,  cancellationToken);
 		}
@@ -272,12 +312,14 @@ namespace Horde.Server.Utilities
 		/// <inheritdoc />
 		public long EstimatedDocumentCount(EstimatedDocumentCountOptions options = null!, CancellationToken cancellationToken = default)
 		{
+			_upgradeTask.Wait(cancellationToken);
 			return _collection.EstimatedDocumentCount(options,  cancellationToken);
 		}
 
 		/// <inheritdoc />
 		public async Task<long> EstimatedDocumentCountAsync(EstimatedDocumentCountOptions options = null!, CancellationToken cancellationToken = default)
 		{
+			await _upgradeTask.WaitAsync(cancellationToken);
 			using TelemetrySpan span = _tracer.StartMongoDbSpan(nameof(EstimatedDocumentCountAsync), _collection);
 			return await _collection.EstimatedDocumentCountAsync(options,  cancellationToken);
 		}
@@ -285,18 +327,21 @@ namespace Horde.Server.Utilities
 		/// <inheritdoc />
 		public IAsyncCursor<TProjection> FindSync<TProjection>(FilterDefinition<T> filter, FindOptions<T, TProjection> options = null!, CancellationToken cancellationToken = default)
 		{
+			_upgradeTask.Wait(cancellationToken);
 			return _collection.FindSync(filter, options,  cancellationToken);
 		}
 
 		/// <inheritdoc />
 		public IAsyncCursor<TProjection> FindSync<TProjection>(IClientSessionHandle session, FilterDefinition<T> filter, FindOptions<T, TProjection> options = null!, CancellationToken cancellationToken = default)
 		{
+			_upgradeTask.Wait(cancellationToken);
 			return _collection.FindSync(session, filter, options,  cancellationToken);
 		}
 
 		/// <inheritdoc />
 		public async Task<IAsyncCursor<TProjection>> FindAsync<TProjection>(FilterDefinition<T> filter, FindOptions<T, TProjection> options = null!, CancellationToken cancellationToken = default)
 		{
+			await _upgradeTask.WaitAsync(cancellationToken);
 			using TelemetrySpan span = _tracer.StartMongoDbSpan(nameof(FindAsync), _collection);
 			return await _collection.FindAsync(filter, options,  cancellationToken);
 		}
@@ -304,6 +349,7 @@ namespace Horde.Server.Utilities
 		/// <inheritdoc />
 		public async Task<IAsyncCursor<TProjection>> FindAsync<TProjection>(IClientSessionHandle session, FilterDefinition<T> filter, FindOptions<T, TProjection> options = null!, CancellationToken cancellationToken = default)
 		{
+			await _upgradeTask.WaitAsync(cancellationToken);
 			using TelemetrySpan span = _tracer.StartMongoDbSpan(nameof(FindAsync), _collection);
 			return await _collection.FindAsync(session, filter, options,  cancellationToken);
 		}
@@ -311,18 +357,21 @@ namespace Horde.Server.Utilities
 		/// <inheritdoc />
 		public TProjection FindOneAndDelete<TProjection>(FilterDefinition<T> filter, FindOneAndDeleteOptions<T, TProjection> options = null!, CancellationToken cancellationToken = default)
 		{
+			_upgradeTask.Wait(cancellationToken);
 			return _collection.FindOneAndDelete(filter, options,  cancellationToken);
 		}
 
 		/// <inheritdoc />
 		public TProjection FindOneAndDelete<TProjection>(IClientSessionHandle session, FilterDefinition<T> filter, FindOneAndDeleteOptions<T, TProjection> options = null!, CancellationToken cancellationToken = default)
 		{
+			_upgradeTask.Wait(cancellationToken);
 			return _collection.FindOneAndDelete(session, filter, options,  cancellationToken);
 		}
 
 		/// <inheritdoc />
 		public async Task<TProjection> FindOneAndDeleteAsync<TProjection>(FilterDefinition<T> filter, FindOneAndDeleteOptions<T, TProjection> options = null!, CancellationToken cancellationToken = default)
 		{
+			await _upgradeTask.WaitAsync(cancellationToken);
 			using TelemetrySpan span = _tracer.StartMongoDbSpan(nameof(FindOneAndDeleteAsync), _collection);
 			return await _collection.FindOneAndDeleteAsync(filter, options,  cancellationToken);
 		}
@@ -330,6 +379,7 @@ namespace Horde.Server.Utilities
 		/// <inheritdoc />
 		public async Task<TProjection> FindOneAndDeleteAsync<TProjection>(IClientSessionHandle session, FilterDefinition<T> filter, FindOneAndDeleteOptions<T, TProjection> options = null!, CancellationToken cancellationToken = default)
 		{
+			await _upgradeTask.WaitAsync(cancellationToken);
 			using TelemetrySpan span = _tracer.StartMongoDbSpan(nameof(FindOneAndDeleteAsync), _collection);
 			return await _collection.FindOneAndDeleteAsync(session, filter, options,  cancellationToken);
 		}
@@ -337,18 +387,21 @@ namespace Horde.Server.Utilities
 		/// <inheritdoc />
 		public TProjection FindOneAndReplace<TProjection>(FilterDefinition<T> filter, T replacement, FindOneAndReplaceOptions<T, TProjection> options = null!, CancellationToken cancellationToken = default)
 		{
+			_upgradeTask.Wait(cancellationToken);
 			return _collection.FindOneAndReplace(filter, replacement, options,  cancellationToken);
 		}
 
 		/// <inheritdoc />
 		public TProjection FindOneAndReplace<TProjection>(IClientSessionHandle session, FilterDefinition<T> filter, T replacement, FindOneAndReplaceOptions<T, TProjection> options = null!, CancellationToken cancellationToken = default)
 		{
+			_upgradeTask.Wait(cancellationToken);
 			return _collection.FindOneAndReplace(session, filter, replacement, options,  cancellationToken);
 		}
 
 		/// <inheritdoc />
 		public async Task<TProjection> FindOneAndReplaceAsync<TProjection>(FilterDefinition<T> filter, T replacement, FindOneAndReplaceOptions<T, TProjection> options = null!, CancellationToken cancellationToken = default)
 		{
+			await _upgradeTask.WaitAsync(cancellationToken);
 			using TelemetrySpan span = _tracer.StartMongoDbSpan(nameof(FindOneAndReplaceAsync), _collection);
 			return await _collection.FindOneAndReplaceAsync(filter, replacement, options,  cancellationToken);
 		}
@@ -356,6 +409,7 @@ namespace Horde.Server.Utilities
 		/// <inheritdoc />
 		public async Task<TProjection> FindOneAndReplaceAsync<TProjection>(IClientSessionHandle session, FilterDefinition<T> filter, T replacement, FindOneAndReplaceOptions<T, TProjection> options = null!, CancellationToken cancellationToken = default)
 		{
+			await _upgradeTask.WaitAsync(cancellationToken);
 			using TelemetrySpan span = _tracer.StartMongoDbSpan(nameof(FindOneAndReplaceAsync), _collection);
 			return await _collection.FindOneAndReplaceAsync(session, filter, replacement, options,  cancellationToken);
 		}
@@ -364,18 +418,21 @@ namespace Horde.Server.Utilities
 		public TProjection FindOneAndUpdate<TProjection>(FilterDefinition<T> filter, UpdateDefinition<T> update,
 			FindOneAndUpdateOptions<T, TProjection> options = null!, CancellationToken cancellationToken = default)
 		{
+			_upgradeTask.Wait(cancellationToken);
 			return _collection.FindOneAndUpdate(filter, update, options,  cancellationToken);
 		}
 
 		/// <inheritdoc />
 		public TProjection FindOneAndUpdate<TProjection>(IClientSessionHandle session, FilterDefinition<T> filter, UpdateDefinition<T> update, FindOneAndUpdateOptions<T, TProjection> options = null!, CancellationToken cancellationToken = default)
 		{
+			_upgradeTask.Wait(cancellationToken);
 			return _collection.FindOneAndUpdate(session, filter, update, options,  cancellationToken);
 		}
 
 		/// <inheritdoc />
 		public async Task<TProjection> FindOneAndUpdateAsync<TProjection>(FilterDefinition<T> filter, UpdateDefinition<T> update, FindOneAndUpdateOptions<T, TProjection> options = null!, CancellationToken cancellationToken = default)
 		{
+			await _upgradeTask.WaitAsync(cancellationToken);
 			using TelemetrySpan span = _tracer.StartMongoDbSpan(nameof(FindOneAndUpdateAsync), _collection);
 			return await _collection.FindOneAndUpdateAsync(filter, update, options,  cancellationToken);
 		}
@@ -383,6 +440,7 @@ namespace Horde.Server.Utilities
 		/// <inheritdoc />
 		public async Task<TProjection> FindOneAndUpdateAsync<TProjection>(IClientSessionHandle session, FilterDefinition<T> filter, UpdateDefinition<T> update, FindOneAndUpdateOptions<T, TProjection> options = null!, CancellationToken cancellationToken = default)
 		{
+			await _upgradeTask.WaitAsync(cancellationToken);
 			using TelemetrySpan span = _tracer.StartMongoDbSpan(nameof(FindOneAndUpdateAsync), _collection);
 			return await _collection.FindOneAndUpdateAsync(session, filter, update, options,  cancellationToken);
 		}
@@ -390,18 +448,21 @@ namespace Horde.Server.Utilities
 		/// <inheritdoc />
 		public void InsertOne(T document, InsertOneOptions options = null!, CancellationToken cancellationToken = default)
 		{
+			_upgradeTask.Wait(cancellationToken);
 			_collection.InsertOne(document, options,  cancellationToken);
 		}
 
 		/// <inheritdoc />
 		public void InsertOne(IClientSessionHandle session, T document, InsertOneOptions options = null!, CancellationToken cancellationToken = default)
 		{
+			_upgradeTask.Wait(cancellationToken);
 			_collection.InsertOne(session, document, options,  cancellationToken);
 		}
 
 		/// <inheritdoc />
 		public async Task InsertOneAsync(T document, CancellationToken cancellationToken)
 		{
+			await _upgradeTask.WaitAsync(cancellationToken);
 			using TelemetrySpan span = _tracer.StartMongoDbSpan(nameof(InsertOneAsync), _collection);
 			await _collection.InsertOneAsync(document, cancellationToken);
 		}
@@ -409,6 +470,7 @@ namespace Horde.Server.Utilities
 		/// <inheritdoc />
 		public async Task InsertOneAsync(T document, InsertOneOptions options = null!, CancellationToken cancellationToken = default)
 		{
+			await _upgradeTask.WaitAsync(cancellationToken);
 			using TelemetrySpan span = _tracer.StartMongoDbSpan(nameof(InsertOneAsync), _collection);
 			await _collection.InsertOneAsync(document, options,  cancellationToken);
 		}
@@ -416,6 +478,7 @@ namespace Horde.Server.Utilities
 		/// <inheritdoc />
 		public async Task InsertOneAsync(IClientSessionHandle session, T document, InsertOneOptions options = null!, CancellationToken cancellationToken = default)
 		{
+			await _upgradeTask.WaitAsync(cancellationToken);
 			using TelemetrySpan span = _tracer.StartMongoDbSpan(nameof(InsertOneAsync), _collection);
 			await _collection.InsertOneAsync(session, document, options,  cancellationToken);
 		}
@@ -423,18 +486,21 @@ namespace Horde.Server.Utilities
 		/// <inheritdoc />
 		public void InsertMany(IEnumerable<T> documents, InsertManyOptions options = null!, CancellationToken cancellationToken = default)
 		{
+			_upgradeTask.Wait(cancellationToken);
 			_collection.InsertMany(documents, options,  cancellationToken);
 		}
 
 		/// <inheritdoc />
 		public void InsertMany(IClientSessionHandle session, IEnumerable<T> documents, InsertManyOptions options = null!, CancellationToken cancellationToken = default)
 		{
+			_upgradeTask.Wait(cancellationToken);
 			_collection.InsertMany(session, documents, options,  cancellationToken);
 		}
 
 		/// <inheritdoc />
 		public async Task InsertManyAsync(IEnumerable<T> documents, InsertManyOptions options = null!, CancellationToken cancellationToken = default)
 		{
+			await _upgradeTask.WaitAsync(cancellationToken);
 			using TelemetrySpan span = _tracer.StartMongoDbSpan(nameof(InsertManyAsync), _collection);
 			await _collection.InsertManyAsync(documents, options,  cancellationToken);
 		}
@@ -442,6 +508,7 @@ namespace Horde.Server.Utilities
 		/// <inheritdoc />
 		public async Task InsertManyAsync(IClientSessionHandle session, IEnumerable<T> documents, InsertManyOptions options = null!, CancellationToken cancellationToken = default)
 		{
+			await _upgradeTask.WaitAsync(cancellationToken);
 			using TelemetrySpan span = _tracer.StartMongoDbSpan(nameof(InsertManyAsync), _collection);
 			await _collection.InsertManyAsync(session, documents, options,  cancellationToken);
 		}
@@ -449,18 +516,21 @@ namespace Horde.Server.Utilities
 		/// <inheritdoc />
 		public IAsyncCursor<TResult> MapReduce<TResult>(BsonJavaScript map, BsonJavaScript reduce, MapReduceOptions<T, TResult> options = null!, CancellationToken cancellationToken = default)
 		{
+			_upgradeTask.Wait(cancellationToken);
 			return _collection.MapReduce(map, reduce, options,  cancellationToken);
 		}
 
 		/// <inheritdoc />
 		public IAsyncCursor<TResult> MapReduce<TResult>(IClientSessionHandle session, BsonJavaScript map, BsonJavaScript reduce, MapReduceOptions<T, TResult> options = null!, CancellationToken cancellationToken = default)
 		{
+			_upgradeTask.Wait(cancellationToken);
 			return _collection.MapReduce(session, map, reduce, options,  cancellationToken);
 		}
 
 		/// <inheritdoc />
 		public async Task<IAsyncCursor<TResult>> MapReduceAsync<TResult>(BsonJavaScript map, BsonJavaScript reduce, MapReduceOptions<T, TResult> options = null!, CancellationToken cancellationToken = default)
 		{
+			await _upgradeTask.WaitAsync(cancellationToken);
 			using TelemetrySpan span = _tracer.StartMongoDbSpan(nameof(MapReduceAsync), _collection);
 			return await _collection.MapReduceAsync(map, reduce, options,  cancellationToken);
 		}
@@ -468,6 +538,7 @@ namespace Horde.Server.Utilities
 		/// <inheritdoc />
 		public async Task<IAsyncCursor<TResult>> MapReduceAsync<TResult>(IClientSessionHandle session, BsonJavaScript map, BsonJavaScript reduce, MapReduceOptions<T, TResult> options = null!, CancellationToken cancellationToken = default)
 		{
+			await _upgradeTask.WaitAsync(cancellationToken);
 			using TelemetrySpan span = _tracer.StartMongoDbSpan(nameof(MapReduceAsync), _collection);
 			return await _collection.MapReduceAsync(session, map, reduce, options,  cancellationToken);
 		}
@@ -481,30 +552,35 @@ namespace Horde.Server.Utilities
 		/// <inheritdoc />
 		public ReplaceOneResult ReplaceOne(FilterDefinition<T> filter, T replacement, ReplaceOptions options = null!, CancellationToken cancellationToken = default)
 		{
+			_upgradeTask.Wait(cancellationToken);
 			return _collection.ReplaceOne(filter, replacement, options,  cancellationToken);
 		}
 
 		/// <inheritdoc />
 		public ReplaceOneResult ReplaceOne(FilterDefinition<T> filter, T replacement, UpdateOptions options, CancellationToken cancellationToken = default)
 		{
+			_upgradeTask.Wait(cancellationToken);
 			return _collection.ReplaceOne(filter, replacement, options,  cancellationToken);
 		}
 
 		/// <inheritdoc />
 		public ReplaceOneResult ReplaceOne(IClientSessionHandle session, FilterDefinition<T> filter, T replacement, ReplaceOptions options = null!, CancellationToken cancellationToken = default)
 		{
+			_upgradeTask.Wait(cancellationToken);
 			return _collection.ReplaceOne(session, filter, replacement, options,  cancellationToken);
 		}
 
 		/// <inheritdoc />
 		public ReplaceOneResult ReplaceOne(IClientSessionHandle session, FilterDefinition<T> filter, T replacement, UpdateOptions options, CancellationToken cancellationToken = default)
 		{
+			_upgradeTask.Wait(cancellationToken);
 			return _collection.ReplaceOne(session, filter, replacement, options,  cancellationToken);
 		}
 
 		/// <inheritdoc />
 		public async Task<ReplaceOneResult> ReplaceOneAsync(FilterDefinition<T> filter, T replacement, ReplaceOptions options = null!, CancellationToken cancellationToken = default)
 		{
+			await _upgradeTask.WaitAsync(cancellationToken);
 			using TelemetrySpan span = _tracer.StartMongoDbSpan(nameof(ReplaceOneAsync), _collection);
 			return await _collection.ReplaceOneAsync(filter, replacement, options,  cancellationToken);
 		}
@@ -512,6 +588,7 @@ namespace Horde.Server.Utilities
 		/// <inheritdoc />
 		public async Task<ReplaceOneResult> ReplaceOneAsync(FilterDefinition<T> filter, T replacement, UpdateOptions options, CancellationToken cancellationToken = default)
 		{
+			await _upgradeTask.WaitAsync(cancellationToken);
 			using TelemetrySpan span = _tracer.StartMongoDbSpan(nameof(ReplaceOneAsync), _collection);
 			return await _collection.ReplaceOneAsync(filter, replacement, options,  cancellationToken);
 		}
@@ -519,6 +596,7 @@ namespace Horde.Server.Utilities
 		/// <inheritdoc />
 		public async Task<ReplaceOneResult> ReplaceOneAsync(IClientSessionHandle session, FilterDefinition<T> filter, T replacement, ReplaceOptions options = null!, CancellationToken cancellationToken = default)
 		{
+			await _upgradeTask.WaitAsync(cancellationToken);
 			using TelemetrySpan span = _tracer.StartMongoDbSpan(nameof(ReplaceOneAsync), _collection);
 			return await _collection.ReplaceOneAsync(session, filter, replacement, options,  cancellationToken);
 		}
@@ -526,6 +604,7 @@ namespace Horde.Server.Utilities
 		/// <inheritdoc />
 		public async Task<ReplaceOneResult> ReplaceOneAsync(IClientSessionHandle session, FilterDefinition<T> filter, T replacement, UpdateOptions options, CancellationToken cancellationToken = default)
 		{
+			await _upgradeTask.WaitAsync(cancellationToken);
 			using TelemetrySpan span = _tracer.StartMongoDbSpan(nameof(ReplaceOneAsync), _collection);
 			return await _collection.ReplaceOneAsync(session, filter, replacement, options,  cancellationToken);
 		}
@@ -533,18 +612,21 @@ namespace Horde.Server.Utilities
 		/// <inheritdoc />
 		public UpdateResult UpdateMany(FilterDefinition<T> filter, UpdateDefinition<T> update, UpdateOptions options = null!, CancellationToken cancellationToken = default)
 		{
+			_upgradeTask.Wait(cancellationToken);
 			return _collection.UpdateMany(filter, update, options,  cancellationToken);
 		}
 
 		/// <inheritdoc />
 		public UpdateResult UpdateMany(IClientSessionHandle session, FilterDefinition<T> filter, UpdateDefinition<T> update, UpdateOptions options = null!, CancellationToken cancellationToken = default)
 		{
+			_upgradeTask.Wait(cancellationToken);
 			return _collection.UpdateMany(session, filter, update, options,  cancellationToken);
 		}
 
 		/// <inheritdoc />
 		public async Task<UpdateResult> UpdateManyAsync(FilterDefinition<T> filter, UpdateDefinition<T> update, UpdateOptions options = null!, CancellationToken cancellationToken = default)
 		{
+			await _upgradeTask.WaitAsync(cancellationToken);
 			using TelemetrySpan span = _tracer.StartMongoDbSpan(nameof(UpdateManyAsync), _collection);
 			return await _collection.UpdateManyAsync(filter, update, options,  cancellationToken);
 		}
@@ -552,6 +634,7 @@ namespace Horde.Server.Utilities
 		/// <inheritdoc />
 		public async Task<UpdateResult> UpdateManyAsync(IClientSessionHandle session, FilterDefinition<T> filter, UpdateDefinition<T> update, UpdateOptions options = null!, CancellationToken cancellationToken = default)
 		{
+			await _upgradeTask.WaitAsync(cancellationToken);
 			using TelemetrySpan span = _tracer.StartMongoDbSpan(nameof(UpdateManyAsync), _collection);
 			return await _collection.UpdateManyAsync(session, filter, update, options,  cancellationToken);
 		}
@@ -559,18 +642,21 @@ namespace Horde.Server.Utilities
 		/// <inheritdoc />
 		public UpdateResult UpdateOne(FilterDefinition<T> filter, UpdateDefinition<T> update, UpdateOptions options = null!, CancellationToken cancellationToken = default)
 		{
+			_upgradeTask.Wait(cancellationToken);
 			return _collection.UpdateOne(filter, update, options,  cancellationToken);
 		}
 
 		/// <inheritdoc />
 		public UpdateResult UpdateOne(IClientSessionHandle session, FilterDefinition<T> filter, UpdateDefinition<T> update, UpdateOptions options = null!, CancellationToken cancellationToken = default)
 		{
+			_upgradeTask.Wait(cancellationToken);
 			return _collection.UpdateOne(session, filter, update, options,  cancellationToken);
 		}
 
 		/// <inheritdoc />
 		public async Task<UpdateResult> UpdateOneAsync(FilterDefinition<T> filter, UpdateDefinition<T> update, UpdateOptions options = null!, CancellationToken cancellationToken = default)
 		{
+			await _upgradeTask.WaitAsync(cancellationToken);
 			using TelemetrySpan span = _tracer.StartMongoDbSpan(nameof(UpdateOneAsync), _collection);
 			return await _collection.UpdateOneAsync(filter, update, options,  cancellationToken);
 		}
@@ -578,6 +664,7 @@ namespace Horde.Server.Utilities
 		/// <inheritdoc />
 		public async Task<UpdateResult> UpdateOneAsync(IClientSessionHandle session, FilterDefinition<T> filter, UpdateDefinition<T> update, UpdateOptions options = null!, CancellationToken cancellationToken = default)
 		{
+			await _upgradeTask.WaitAsync(cancellationToken);
 			using TelemetrySpan span = _tracer.StartMongoDbSpan(nameof(UpdateOneAsync), _collection);
 			return await _collection.UpdateOneAsync(session, filter, update, options,  cancellationToken);
 		}
@@ -585,18 +672,21 @@ namespace Horde.Server.Utilities
 		/// <inheritdoc />
 		public IChangeStreamCursor<TResult> Watch<TResult>(PipelineDefinition<ChangeStreamDocument<T>, TResult> pipeline, ChangeStreamOptions options = null!, CancellationToken cancellationToken = default)
 		{
+			_upgradeTask.Wait(cancellationToken);
 			return _collection.Watch(pipeline, options,  cancellationToken);
 		}
 
 		/// <inheritdoc />
 		public IChangeStreamCursor<TResult> Watch<TResult>(IClientSessionHandle session, PipelineDefinition<ChangeStreamDocument<T>, TResult> pipeline, ChangeStreamOptions options = null!, CancellationToken cancellationToken = default)
 		{
+			_upgradeTask.Wait(cancellationToken);
 			return _collection.Watch(session, pipeline, options,  cancellationToken);
 		}
 
 		/// <inheritdoc />
 		public async Task<IChangeStreamCursor<TResult>> WatchAsync<TResult>(PipelineDefinition<ChangeStreamDocument<T>, TResult> pipeline, ChangeStreamOptions options = null!, CancellationToken cancellationToken = default)
 		{
+			await _upgradeTask.WaitAsync(cancellationToken);
 			using TelemetrySpan span = _tracer.StartMongoDbSpan(nameof(WatchAsync), _collection);
 			return await _collection.WatchAsync(pipeline, options,  cancellationToken);
 		}
@@ -604,6 +694,7 @@ namespace Horde.Server.Utilities
 		/// <inheritdoc />
 		public async Task<IChangeStreamCursor<TResult>> WatchAsync<TResult>(IClientSessionHandle session, PipelineDefinition<ChangeStreamDocument<T>, TResult> pipeline, ChangeStreamOptions options = null!, CancellationToken cancellationToken = default)
 		{
+			await _upgradeTask.WaitAsync(cancellationToken);
 			using TelemetrySpan span = _tracer.StartMongoDbSpan(nameof(WatchAsync), _collection);
 			return await _collection.WatchAsync(session, pipeline, options,  cancellationToken);
 		}
