Horde: Shuffling some code around trying to track down missing log message issues.

diff --git a/Engine/Source/Programs/Horde/Horde.Server/Server/MongoService.cs b/Engine/Source/Programs/Horde/Horde.Server/Server/MongoService.cs
--- a/Engine/Source/Programs/Horde/Horde.Server/Server/MongoService.cs
+++ b/Engine/Source/Programs/Horde/Horde.Server/Server/MongoService.cs
@@ -264,12 +264,13 @@ namespace Horde.Server.Server
 		/// </summary>
 		/// <param name="settingsSnapshot">The settings instance</param>
 		/// <param name="tracer">Tracer</param>
+		/// <param name="logger">Logger for output</param>
 		/// <param name="loggerFactory">Instance of the logger for this service</param>
-		public MongoService(IOptions<ServerSettings> settingsSnapshot, Tracer tracer, ILoggerFactory loggerFactory)
+		public MongoService(IOptions<ServerSettings> settingsSnapshot, Tracer tracer, ILogger<MongoService> logger, ILoggerFactory loggerFactory)
 		{
 			Settings = settingsSnapshot.Value;
 			_tracer = tracer;
-			_logger = loggerFactory.CreateLogger<MongoService>();
+			_logger = logger;
 			_loggerFactory = loggerFactory;
 
 			try
@@ -747,7 +748,7 @@ namespace Horde.Server.Server
 					}
 					else
 					{
-						_logger.LogInformation("Creating index {IndexName} in {CollectionName}", createIndex.Name, collectionName);
+						_logger.LogWarning("Creating index {IndexName} in {CollectionName}", createIndex.Name, collectionName);
 
 						CreateIndexOptions<T> options = new CreateIndexOptions<T>();
 						options.Name = createIndex.Name;
@@ -763,14 +764,15 @@ namespace Horde.Server.Server
 						}
 						catch (Exception ex)
 						{
+							_logger.LogWarning("Unable to create index {IndexName}", createIndex.Name);
 							_logger.LogError(ex, "Unable to create index {IndexName}: {Message}", createIndex.Name, ex.Message);
 							throw;
 						}
 					}
 				}
-
-				_logger.LogInformation("Finished updating indexes for collection {CollectionName}", collectionName);
 			}
+
+			_logger.LogInformation("Finished updating indexes for collection {CollectionName}", collectionName);
 		}
 
 		class SingletonInfo<T> where T : SingletonBase
