Horde: Build fix.

diff --git a/Engine/Source/Programs/Horde/Horde.Server.Tests/DatabaseIntegrationTest.cs b/Engine/Source/Programs/Horde/Horde.Server.Tests/DatabaseIntegrationTest.cs
--- a/Engine/Source/Programs/Horde/Horde.Server.Tests/DatabaseIntegrationTest.cs
+++ b/Engine/Source/Programs/Horde/Horde.Server.Tests/DatabaseIntegrationTest.cs
@@ -182,7 +182,7 @@ namespace Horde.Server.Tests
 					ss.DatabaseName = _mongoDbInstance.DatabaseName;
 					ss.DatabaseConnectionString = _mongoDbInstance.ConnectionString;
 
-					_mongoService = new MongoService(Options.Create(ss), OpenTelemetryTracers.Horde, _loggerFactory);
+					_mongoService = new MongoService(Options.Create(ss), OpenTelemetryTracers.Horde, _loggerFactory.CreateLogger<MongoService>(), _loggerFactory);
 				}
 			}
 			return _mongoService;
