Horde: Remove warnings.

diff --git a/Engine/Source/Programs/Horde/Horde.Server/Server/MongoService.cs b/Engine/Source/Programs/Horde/Horde.Server/Server/MongoService.cs
--- a/Engine/Source/Programs/Horde/Horde.Server/Server/MongoService.cs
+++ b/Engine/Source/Programs/Horde/Horde.Server/Server/MongoService.cs
@@ -748,7 +748,7 @@ namespace Horde.Server.Server
 					}
 					else
 					{
-						_logger.LogWarning("Creating index {IndexName} in {CollectionName}", createIndex.Name, collectionName);
+						_logger.LogInformation("Creating index {IndexName} in {CollectionName}", createIndex.Name, collectionName);
 
 						CreateIndexOptions<T> options = new CreateIndexOptions<T>();
 						options.Name = createIndex.Name;
@@ -764,7 +764,6 @@ namespace Horde.Server.Server
 						}
 						catch (Exception ex)
 						{
-							_logger.LogWarning("Unable to create index {IndexName}", createIndex.Name);
 							_logger.LogError(ex, "Unable to create index {IndexName}: {Message}", createIndex.Name, ex.Message);
 							throw;
 						}
