Horde: Ensure that the upgrade tasks run before performing any operations on a collection.

diff --git a/Engine/Source/Programs/Horde/Horde.Server.Tests/DatabaseIntegrationTest.cs b/Engine/Source/Programs/Horde/Horde.Server.Tests/DatabaseIntegrationTest.cs
--- a/Engine/Source/Programs/Horde/Horde.Server.Tests/DatabaseIntegrationTest.cs
+++ b/Engine/Source/Programs/Horde/Horde.Server.Tests/DatabaseIntegrationTest.cs
@@ -176,13 +176,15 @@ namespace Horde.Server.Tests
 			{
 				if (_mongoService == null)
 				{
+					RedisService redisService = GetRedisServiceSingleton();
+
 					_mongoDbInstance = new MongoDbInstance();
 
 					ServerSettings ss = new ServerSettings();
 					ss.DatabaseName = _mongoDbInstance.DatabaseName;
 					ss.DatabaseConnectionString = _mongoDbInstance.ConnectionString;
 
-					_mongoService = new MongoService(Options.Create(ss), OpenTelemetryTracers.Horde, _loggerFactory.CreateLogger<MongoService>(), _loggerFactory);
+					_mongoService = new MongoService(Options.Create(ss), redisService, OpenTelemetryTracers.Horde, _loggerFactory.CreateLogger<MongoService>(), _loggerFactory);
 				}
 			}
 			return _mongoService;
diff --git a/Engine/Source/Programs/Horde/Horde.Server/Server/MongoService.cs b/Engine/Source/Programs/Horde/Horde.Server/Server/MongoService.cs
--- a/Engine/Source/Programs/Horde/Horde.Server/Server/MongoService.cs
+++ b/Engine/Source/Programs/Horde/Horde.Server/Server/MongoService.cs
@@ -256,19 +256,20 @@ namespace Horde.Server.Server
 		/// </summary>
 		const int DefaultMongoPort = 27017;
 
-		readonly HashSet<string> _collectionNames = new HashSet<string>(StringComparer.Ordinal);
-		readonly Channel<MongoUpgradeTask> _updateIndexesChannel = Channel.CreateUnbounded<MongoUpgradeTask>();
+		static readonly RedisKey s_schemaLockKey = new RedisKey("server/schema-upgrade/lock");
+
+		readonly RedisService _redisService;
+		readonly SemaphoreSlim _upgradeSema = new SemaphoreSlim(1);
+		readonly Dictionary<string, Task> _collectionUpgradeTasks = new Dictionary<string, Task>(StringComparer.Ordinal);
+		readonly Task<bool> _setSchemaVersionTask;
 
 		/// <summary>
 		/// Constructor
 		/// </summary>
-		/// <param name="settingsSnapshot">The settings instance</param>
-		/// <param name="tracer">Tracer</param>
-		/// <param name="logger">Logger for output</param>
-		/// <param name="loggerFactory">Instance of the logger for this service</param>
-		public MongoService(IOptions<ServerSettings> settingsSnapshot, Tracer tracer, ILogger<MongoService> logger, ILoggerFactory loggerFactory)
+		public MongoService(IOptions<ServerSettings> settingsSnapshot, RedisService redisService, Tracer tracer, ILogger<MongoService> logger, ILoggerFactory loggerFactory)
 		{
 			Settings = settingsSnapshot.Value;
+			_redisService = redisService;
 			_tracer = tracer;
 			_logger = logger;
 			_loggerFactory = loggerFactory;
@@ -340,6 +341,8 @@ namespace Horde.Server.Server
 				_logger.LogError(ex, "Exception while initializing MongoService");
 				throw;
 			}
+
+			_setSchemaVersionTask = SetSchemaVersion(Program.Version);
 		}
 
 		internal const int CtrlCEvent = 0;
@@ -652,42 +655,33 @@ namespace Horde.Server.Server
 		{
 			IMongoCollection<T> collection = Database.GetCollection<T>(name);
 
-			Task upgradeTask = Task.CompletedTask;
+			Task? upgradeTask = Task.CompletedTask;
 			if (indexes.Any())
 			{
-				lock (_collectionNames)
+				lock (_collectionUpgradeTasks)
 				{
-					_logger.LogDebug("Queuing update for collection {Name}", name);
-
-					if (!_collectionNames.Add(name))
+					if (!_collectionUpgradeTasks.TryGetValue(name, out upgradeTask))
 					{
-						throw new NotImplementedException();
-					}
+						_logger.LogDebug("Queuing update for collection {Name}", name);
 
-					MongoIndex<T>[] indexesCopy = indexes.ToArray();
+						MongoIndex<T>[] indexesCopy = indexes.ToArray();
+						upgradeTask = Task.Run(() => UpdateIndexesAsync(name, collection, indexesCopy, CancellationToken.None));
 
-					TaskCompletionSource tcs = new TaskCompletionSource();
-					_updateIndexesChannel.Writer.TryWrite(new MongoUpgradeTask(ctx => UpdateIndexesAsync(name, collection, indexesCopy, ctx), tcs));
-					upgradeTask = tcs.Task;
+						_collectionUpgradeTasks.Add(name, upgradeTask);
+					}
 				}
 			}
 
 			return new MongoTracingCollection<T>(Database.GetCollection<T>(name), upgradeTask, _tracer);
 		}
 
-		/// <summary>
-		/// Pops an upgrade task from the queue. This method should only be called by <see cref="MongoUpgradeService"/>
-		/// </summary>
-		/// <param name="cancellationToken"></param>
-		/// <returns></returns>
-		internal ValueTask<MongoUpgradeTask> ReadNextUpgradeTask(CancellationToken cancellationToken)
-		{
-			return _updateIndexesChannel.Reader.ReadAsync(cancellationToken);
-		}
-
 		private async Task UpdateIndexesAsync<T>(string collectionName, IMongoCollection<T> collection, MongoIndex<T>[] newIndexes, CancellationToken cancellationToken)
 		{
-			_logger.LogDebug("Updating indexes for collection {CollectionName}", collectionName);
+			// Check we're allowed to upgrade the DB
+			if (!await _setSchemaVersionTask)
+			{
+				return;
+			}
 
 			// Find all the current indexes, excluding the default
 			Dictionary<string, MongoIndex> nameToExistingIndex = new Dictionary<string, MongoIndex>(StringComparer.Ordinal);
@@ -725,53 +719,105 @@ namespace Horde.Server.Server
 
 			if (removeIndexNames.Count > 0 || createIndexes.Count > 0)
 			{
-				// Drop any indexes that are no longer needed
-				foreach (string removeIndexName in removeIndexNames)
+				await _upgradeSema.WaitAsync(cancellationToken);
+				try
 				{
-					if (ReadOnlyMode)
+					// Aquire a lock for updating the DB
+					await using RedisLock schemaLock = new(_redisService.GetDatabase(), s_schemaLockKey);
+					while (!await schemaLock.AcquireAsync(TimeSpan.FromMinutes(5.0)))
 					{
-						_logger.LogWarning("Would drop unused index {CollectionName}.{IndexName} - skipping due to read-only setting.", collectionName, removeIndexName);
+						_logger.LogDebug("Unable to acquire lock for upgrade task; pausing for 1s");
+						await Task.Delay(TimeSpan.FromSeconds(1.0), cancellationToken);
 					}
-					else
-					{
-						_logger.LogInformation("Dropping unused index {IndexName}", removeIndexName);
-						await collection.Indexes.DropOneAsync(removeIndexName, cancellationToken);
-					}
-				}
 
-				// Create all the new indexes
-				foreach (MongoIndex<T> createIndex in createIndexes)
-				{
-					if (ReadOnlyMode)
+					_logger.LogDebug("Updating indexes for collection {CollectionName}", collectionName);
+
+					// Drop any indexes that are no longer needed
+					foreach (string removeIndexName in removeIndexNames)
 					{
-						_logger.LogWarning("Would create index {CollectionName}.{IndexName} - skipping due to read-only setting.", collectionName, createIndex.Name);
+						if (ReadOnlyMode)
+						{
+							_logger.LogWarning("Would drop unused index {CollectionName}.{IndexName} - skipping due to read-only setting.", collectionName, removeIndexName);
+						}
+						else
+						{
+							_logger.LogInformation("Dropping unused index {IndexName}", removeIndexName);
+							await collection.Indexes.DropOneAsync(removeIndexName, cancellationToken);
+						}
 					}
-					else
-					{
-						_logger.LogInformation("Creating index {IndexName} in {CollectionName}", createIndex.Name, collectionName);
-
-						CreateIndexOptions<T> options = new CreateIndexOptions<T>();
-						options.Name = createIndex.Name;
-						options.Unique = createIndex.Unique;
-						options.Sparse = createIndex.Sparse;
 
-						CreateIndexModel<T> model = new CreateIndexModel<T>(createIndex.Keys, options);
-
-						try
+					// Create all the new indexes
+					foreach (MongoIndex<T> createIndex in createIndexes)
+					{
+						if (ReadOnlyMode)
 						{
-							string result = await collection.Indexes.CreateOneAsync(model, cancellationToken: cancellationToken);
-							_logger.LogInformation("Created index {IndexName}", result);
+							_logger.LogWarning("Would create index {CollectionName}.{IndexName} - skipping due to read-only setting.", collectionName, createIndex.Name);
 						}
-						catch (Exception ex)
+						else
 						{
-							_logger.LogError(ex, "Unable to create index {IndexName}: {Message}", createIndex.Name, ex.Message);
-							throw;
+							_logger.LogInformation("Creating index {IndexName} in {CollectionName}", createIndex.Name, collectionName);
+
+							CreateIndexOptions<T> options = new CreateIndexOptions<T>();
+							options.Name = createIndex.Name;
+							options.Unique = createIndex.Unique;
+							options.Sparse = createIndex.Sparse;
+
+							CreateIndexModel<T> model = new CreateIndexModel<T>(createIndex.Keys, options);
+
+							try
+							{
+								string result = await collection.Indexes.CreateOneAsync(model, cancellationToken: cancellationToken);
+								_logger.LogInformation("Created index {IndexName}", result);
+							}
+							catch (Exception ex)
+							{
+								_logger.LogError(ex, "Unable to create index {IndexName}: {Message}", createIndex.Name, ex.Message);
+								throw;
+							}
 						}
 					}
 				}
+				finally
+				{
+					_upgradeSema.Release();
+				}
+
+				_logger.LogInformation("Finished updating indexes for collection {CollectionName}", collectionName);
 			}
+		}
+
+		async Task<bool> SetSchemaVersion(SemVer schemaVersion)
+		{
+			// Check we're not downgrading the data
+			for (; ; )
+			{
+				MongoSchemaDocument currentSchema = await GetSingletonAsync<MongoSchemaDocument>();
+				if (!String.IsNullOrEmpty(currentSchema.Version))
+				{
+					SemVer currentVersion = SemVer.Parse(currentSchema.Version);
+					if (schemaVersion < currentVersion)
+					{
+						_logger.LogInformation("Ignoring upgrade command; server is older than current schema version ({ProgramVer} < {CurrentVer})", Program.Version, currentVersion);
+						return false;
+					}
+					if (schemaVersion == currentVersion)
+					{
+						return true;
+					}
+				}
+
+				_logger.LogInformation("Upgrading schema version {OldVersion} -> {NewVersion}", currentSchema.Version, schemaVersion.ToString());
+				currentSchema.Version = schemaVersion.ToString();
 
-			_logger.LogInformation("Finished updating indexes for collection {CollectionName}", collectionName);
+				if (ReadOnlyMode)
+				{
+					return false;
+				}
+				if (await TryUpdateSingletonAsync(currentSchema))
+				{
+					return true;
+				}
+			}
 		}
 
 		class SingletonInfo<T> where T : SingletonBase
@@ -929,141 +975,4 @@ namespace Horde.Server.Server
 		/// </summary>
 		public string? Version { get; set; }
 	}
-
-	/// <summary>
-	/// Service which updates the database state
-	/// </summary>
-	public sealed class MongoUpgradeService : IDisposable, IHostedService
-	{
-		static readonly RedisKey s_schemaLockKey = new RedisKey("server/schema-upgrade/lock");
-
-		readonly MongoService _mongoService;
-		readonly RedisService _redisService;
-		readonly ISingletonDocument<MongoSchemaDocument> _schemaDocument;
-		readonly ILogger _logger;
-		readonly BackgroundTask _updateIndexesTask;
-
-		/// <summary>
-		/// Constructor
-		/// </summary>
-		public MongoUpgradeService(MongoService mongoService, RedisService redisService, ISingletonDocument<MongoSchemaDocument> schemaDocument, ILogger<MongoUpgradeService> logger)
-		{
-			_mongoService = mongoService;
-			_redisService = redisService;
-			_schemaDocument = schemaDocument;
-			_logger = logger;
-			_updateIndexesTask = new BackgroundTask(UpdateAsync);
-		}
-
-		/// <inheritdoc/>
-		public void Dispose()
-		{
-			_updateIndexesTask.Dispose();
-		}
-
-		/// <inheritdoc/>
-		public Task StartAsync(CancellationToken cancellationToken)
-		{
-			_updateIndexesTask.Start();
-			return Task.CompletedTask;
-		}
-
-		/// <inheritdoc/>
-		public async Task StopAsync(CancellationToken cancellationToken)
-		{
-			await _updateIndexesTask.StopAsync();
-		}
-
-		async Task UpdateAsync(CancellationToken cancellationToken)
-		{
-			for (; ; )
-			{
-				try
-				{
-					MongoUpgradeTask upgradeTask = await _mongoService.ReadNextUpgradeTask(cancellationToken);
-					for (; ; )
-					{
-						using (RedisLock schemaLock = new(_redisService.GetDatabase(), s_schemaLockKey))
-						{
-							if (await schemaLock.AcquireAsync(TimeSpan.FromMinutes(5.0)))
-							{
-								await UpdateOneAsync(upgradeTask, cancellationToken);
-								break;
-							}
-							else
-							{
-								_logger.LogDebug("Unable to acquire lock for upgrade task; pausing for 1m");
-								await Task.Delay(TimeSpan.FromMinutes(1.0), cancellationToken);
-							}
-						}
-					}
-				}
-				catch(OperationCanceledException) when (cancellationToken.IsCancellationRequested)
-				{
-					break;
-				}
-				catch (Exception ex)
-				{
-					_logger.LogError(ex, "Exception updating indexes: {Message}", ex.Message);
-					await Task.Delay(TimeSpan.FromSeconds(30.0), cancellationToken);
-				}
-			}
-		}
-
-		async ValueTask UpdateOneAsync(MongoUpgradeTask upgradeTask, CancellationToken cancellationToken)
-		{
-			try
-			{
-				if (await SetSchemaVersion(Program.Version))
-				{
-					await upgradeTask.UpgradeAsync(cancellationToken);
-				}
-				upgradeTask.CompletionSource.SetResult();
-			}
-			catch (MongoCommandException ex)
-			{
-				_logger.LogError(ex, "Command exception while attempting to update indexes ({Code}): {Message}", ex.Code, ex.Message);
-				upgradeTask.CompletionSource.TrySetException(ex);
-			}
-			catch (Exception ex)
-			{
-				_logger.LogError(ex, "Exception while attempting to update indexes: {Message}", ex.Message);
-				upgradeTask.CompletionSource.TrySetException(ex);
-			}
-		}
-
-		async Task<bool> SetSchemaVersion(SemVer schemaVersion)
-		{
-			// Check we're not downgrading the data
-			for (; ; )
-			{
-				MongoSchemaDocument currentSchema = await _schemaDocument.GetAsync();
-				if (!String.IsNullOrEmpty(currentSchema.Version))
-				{
-					SemVer currentVersion = SemVer.Parse(currentSchema.Version);
-					if (schemaVersion < currentVersion)
-					{
-						_logger.LogInformation("Ignoring upgrade command; server is older than current schema version ({ProgramVer} < {CurrentVer})", Program.Version, currentVersion);
-						return false;
-					}
-					if (schemaVersion == currentVersion)
-					{
-						return true;
-					}
-				}
-
-				_logger.LogInformation("Upgrading schema version {OldVersion} -> {NewVersion}", currentSchema.Version, schemaVersion.ToString());
-				currentSchema.Version = schemaVersion.ToString();
-
-				if (_mongoService.ReadOnlyMode)
-				{
-					return false;
-				}
-				if (await _schemaDocument.TryUpdateAsync(currentSchema))
-				{
-					return true;
-				}
-			}
-		}
-	}
 }
diff --git a/Engine/Source/Programs/Horde/Horde.Server/Startup.cs b/Engine/Source/Programs/Horde/Horde.Server/Startup.cs
--- a/Engine/Source/Programs/Horde/Horde.Server/Startup.cs
+++ b/Engine/Source/Programs/Horde/Horde.Server/Startup.cs
@@ -686,9 +686,6 @@ namespace Horde.Server
 			// Always run agent service too; need to be able to listen to Redis for events on any server.
 			services.AddHostedService(provider => provider.GetRequiredService<AgentService>());
 
-			// We need to run the Mongo upgrade service too; some collections only get instantiated when referenced at the moment, and references may come solely from HTTP requests.
-			services.AddHostedService<MongoUpgradeService>();
-
 			if (settings.IsRunModeActive(RunMode.Worker) && !settings.DatabaseReadOnlyMode)
 			{
 				services.AddHostedService<BisectService>();
