Horde: Fix test data test. This was always incorrect, but was not failing due to unique index not being created.

diff --git a/Engine/Source/Programs/Horde/Horde.Server.Tests/TestDataTests.cs b/Engine/Source/Programs/Horde/Horde.Server.Tests/TestDataTests.cs
--- a/Engine/Source/Programs/Horde/Horde.Server.Tests/TestDataTests.cs
+++ b/Engine/Source/Programs/Horde/Horde.Server.Tests/TestDataTests.cs
@@ -256,7 +256,7 @@ namespace Horde.Server.Tests
 			foreach (BsonValue item in items.ToList())
 			{
 				BsonDocument value = item.AsBsonDocument.GetValue("Data").AsBsonDocument;
-				data.Add(("Session Report Key", value));
+				data.Add(($"Session Report Key {data.Count}", value));
 			}
 
 			List<string> dataLines2 = new List<string>();
@@ -272,7 +272,7 @@ namespace Horde.Server.Tests
 			foreach (BsonValue item in items.ToList())
 			{
 				BsonDocument value = item.AsBsonDocument.GetValue("Data").AsBsonDocument;
-				data2.Add(("Session Report Key", value));
+				data2.Add(($"Session Report Key {data2.Count}", value));
 			}
 
 			await TestDataCollection.AddAsync(job, step, data.ToArray());
