Horde: Integrate DDC controllers from Jupiter. Most interfaces are currently just stubbed out.

diff --git a/Engine/Source/Programs/Horde/Horde.Server/Ddc/Acl.cs b/Engine/Source/Programs/Horde/Horde.Server/Ddc/Acl.cs
new file mode 100644
--- /dev/null
+++ b/Engine/Source/Programs/Horde/Horde.Server/Ddc/Acl.cs
@@ -0,0 +1,112 @@
+// Copyright Epic Games, Inc. All Rights Reserved.
+
+using System;
+using System.Collections.Generic;
+using System.Security.Claims;
+using Microsoft.AspNetCore.Authorization;
+
+#pragma warning disable CS1591
+
+namespace Jupiter;
+
+public enum JupiterAclAction
+{
+    /// <summary>
+    /// General read access to refs / blobs and so on
+    /// </summary>
+    ReadObject,
+    /// <summary>
+    /// General write access to upload refs / blobs etc
+    /// </summary>
+    WriteObject,
+    /// <summary>
+    /// Access to delete blobs / refs etc
+    /// </summary>
+    DeleteObject,
+
+    /// <summary>
+    /// Access to delete a particular bucket
+    /// </summary>
+    DeleteBucket,
+    /// <summary>
+    /// Access to delete a whole namespace
+    /// </summary>
+    DeleteNamespace,
+
+    /// <summary>
+    /// Access to read the transaction log
+    /// </summary>
+    ReadTransactionLog,
+
+    /// <summary>
+    /// Access to write the transaction log
+    /// </summary>
+    WriteTransactionLog,
+
+    /// <summary>
+    /// Access to perform administrative task
+    /// </summary>
+    AdminAction
+}
+
+public class AclEntry
+{
+    /// <summary>
+    /// Claims required to be present to be allowed to do the actions - if multiple claims are present *all* of them are required
+    /// </summary>
+    [System.Diagnostics.CodeAnalysis.SuppressMessage("Usage", "CA2227:Collection properties should be read only", Justification = "Used by the configuration system")]
+    public List<string> Claims { get; set; } = new List<string>();
+
+    /// <summary>
+    /// The actions granted if the claims match
+    /// </summary>
+    [System.Diagnostics.CodeAnalysis.SuppressMessage("Usage", "CA2227:Collection properties should be read only", Justification = "Used by the configuration system")]
+    public List<JupiterAclAction> Actions { get; set; } = new List<JupiterAclAction>();
+
+    public IEnumerable<JupiterAclAction> Resolve(ClaimsPrincipal user)
+    {
+        List<JupiterAclAction> allowedActions = new List<JupiterAclAction>();
+
+        bool allClaimsFound = true;
+        // These are ANDed, e.g. all claims needs to be present
+        foreach (string expectedClaim in Claims)
+        {
+            bool claimFound = false;
+            // if expected claim is * then everyone has the associated actions
+            if (expectedClaim == "*")
+            {
+                claimFound = true;
+            }
+            else if (expectedClaim.Contains('=', StringComparison.InvariantCultureIgnoreCase))
+            {
+                int separatorIndex = expectedClaim.IndexOf('=', StringComparison.InvariantCultureIgnoreCase);
+                string claimName = expectedClaim.Substring(0, separatorIndex);
+                string claimValue = expectedClaim.Substring(separatorIndex + 1);
+                if (user.HasClaim(claim => string.Equals(claim.Type, claimName, StringComparison.OrdinalIgnoreCase) && string.Equals(claim.Value, claimValue, StringComparison.OrdinalIgnoreCase)))
+                {
+                    claimFound = true;
+                }
+            }
+            else if (user.HasClaim(claim => string.Equals(claim.Type, expectedClaim, StringComparison.OrdinalIgnoreCase)))
+            {
+                claimFound = true;
+            }
+
+            if (!claimFound)
+            {
+                allClaimsFound = false;
+            }
+        }
+
+        if (allClaimsFound)
+        {
+            allowedActions.AddRange(Actions);
+        }
+        return allowedActions;
+    }
+
+    public IEnumerable<JupiterAclAction> Resolve(AuthorizationHandlerContext context)
+    {
+        return Resolve(context.User);
+    }
+}
\ No newline at end of file
diff --git a/Engine/Source/Programs/Horde/Horde.Server/Ddc/BlobContents.cs b/Engine/Source/Programs/Horde/Horde.Server/Ddc/BlobContents.cs
new file mode 100644
--- /dev/null
+++ b/Engine/Source/Programs/Horde/Horde.Server/Ddc/BlobContents.cs
@@ -0,0 +1,62 @@
+﻿// Copyright Epic Games, Inc. All Rights Reserved.
+
+using System;
+using System.IO;
+using System.Threading.Tasks;
+
+#pragma warning disable CS1591
+
+namespace Jupiter.Implementation
+{
+    public sealed class BlobContents : IDisposable, IAsyncDisposable
+    {
+        private readonly Stream? _stream;
+
+        public BlobContents(Stream stream, long length, string? localPath = null)
+        {
+            _stream = stream;
+
+            Length = length;
+            LocalPath = localPath;
+        }
+        
+        public BlobContents(byte[] payload)
+        {
+            _stream = new MemoryStream(payload);
+            Length = payload.LongLength;
+            LocalPath = null;
+        }
+
+        public BlobContents(Uri redirectUri)
+        {
+            RedirectUri = redirectUri;
+        }
+
+        public Uri? RedirectUri { get; }
+
+        public Stream Stream 
+        { 
+            get
+            {
+                if (_stream == null)
+                {
+                    throw new Exception("Stream not set in blob contents, did you specify support for redirect uris but forgot to check result?");
+                }
+                return _stream;
+            }
+        }
+        public long Length { get; }
+
+        public string? LocalPath { get; }
+
+        public void Dispose()
+        {
+            _stream?.Dispose();
+        }
+        public ValueTask DisposeAsync()
+        {
+            _stream?.Dispose();
+            return ValueTask.CompletedTask;
+        }
+    }
+}
diff --git a/Engine/Source/Programs/Horde/Horde.Server/Ddc/BlobController.cs b/Engine/Source/Programs/Horde/Horde.Server/Ddc/BlobController.cs
new file mode 100644
--- /dev/null
+++ b/Engine/Source/Programs/Horde/Horde.Server/Ddc/BlobController.cs
@@ -0,0 +1,436 @@
+// Copyright Epic Games, Inc. All Rights Reserved.
+
+using System;
+using System.Collections.Concurrent;
+using System.Collections.Generic;
+using System.ComponentModel.DataAnnotations;
+using System.IO;
+using System.Linq;
+using System.Net;
+using System.Net.Mime;
+using System.Text.Json.Serialization;
+using System.Threading.Tasks;
+using EpicGames.AspNet;
+using EpicGames.Horde.Storage;
+using EpicGames.Serialization;
+using Jupiter.Implementation;
+using Jupiter.Common.Implementation;
+using Microsoft.AspNetCore.Authorization;
+using Microsoft.AspNetCore.Http;
+using Microsoft.AspNetCore.Mvc;
+
+#pragma warning disable CS1591
+
+namespace Jupiter.Controllers
+{
+    using IDiagnosticContext = Serilog.IDiagnosticContext;
+    using BlobNotFoundException = Jupiter.Implementation.BlobNotFoundException;
+
+    [ApiController]
+    [Route("api/v1/s", Order = 1)]
+    [Route("api/v1/blobs", Order = 0)]
+    [Authorize]
+    public class BlobsController : ControllerBase
+    {
+        private readonly IBlobService _storage;
+        private readonly IDiagnosticContext _diagnosticContext;
+        private readonly IRequestHelper _requestHelper;
+        private readonly BufferedPayloadFactory _bufferedPayloadFactory;
+        private readonly NginxRedirectHelper _nginxRedirectHelper;
+
+        public BlobsController(IBlobService storage, IDiagnosticContext diagnosticContext, IRequestHelper requestHelper, BufferedPayloadFactory bufferedPayloadFactory, NginxRedirectHelper nginxRedirectHelper)
+        {
+            _storage = storage;
+            _diagnosticContext = diagnosticContext;
+            _requestHelper = requestHelper;
+            _bufferedPayloadFactory = bufferedPayloadFactory;
+            _nginxRedirectHelper = nginxRedirectHelper;
+        }
+
+        [HttpGet("{ns}/{id}")]
+        [ProducesDefaultResponseType]
+        public async Task<IActionResult> Get(
+            [Required] NamespaceId ns,
+            [Required] BlobIdentifier id,
+            [FromQuery] List<string>? storageLayers = null)
+        {
+            ActionResult? result = await _requestHelper.HasAccessToNamespace(User, Request, ns, new [] { JupiterAclAction.ReadObject });
+            if (result != null)
+            {
+                return result;
+            }
+
+            try
+            {
+                BlobContents blobContents = await GetImpl(ns, id, storageLayers, supportsRedirectUri: true);
+
+                if (blobContents.RedirectUri != null)
+                {
+                    return Redirect(blobContents.RedirectUri.ToString());
+                }
+                if (_nginxRedirectHelper.CanRedirect(Request, blobContents))
+                {
+                    return _nginxRedirectHelper.CreateActionResult(blobContents, MediaTypeNames.Application.Octet);
+                }
+                return File(blobContents.Stream, MediaTypeNames.Application.Octet, enableRangeProcessing: true);
+            }
+            catch (BlobNotFoundException e)
+            {
+                return NotFound(new ValidationProblemDetails { Title = $"Blob {e.Blob} not found"});
+            }
+            catch (AuthorizationException e)
+            {
+                return e.Result;
+            }
+        }
+        
+        [HttpHead("{ns}/{id}")]
+        [ProducesDefaultResponseType]
+        public async Task<IActionResult> Head(
+            [Required] NamespaceId ns,
+            [Required] BlobIdentifier id,
+            [FromQuery] List<string>? storageLayers = null)
+        {
+            ActionResult? result = await _requestHelper.HasAccessToNamespace(User, Request, ns, new [] { JupiterAclAction.ReadObject });
+            if (result != null)
+            {
+                return result;
+            }
+            bool exists = await _storage.Exists(ns, id, storageLayers);
+
+            if (!exists)
+            {
+                return NotFound(new ValidationProblemDetails { Title = $"Blob {id} not found"});
+            }
+
+            return Ok();
+        }
+
+        [HttpPost("{ns}/exists")]
+        [ProducesDefaultResponseType]
+        public async Task<IActionResult> ExistsMultiple(
+            [Required] NamespaceId ns,
+            [Required] [FromQuery] List<BlobIdentifier> id)
+        {
+            ActionResult? result = await _requestHelper.HasAccessToNamespace(User, Request, ns, new [] { JupiterAclAction.ReadObject });
+            if (result != null)
+            {
+                return result;
+            }
+
+            ConcurrentBag<BlobIdentifier> missingBlobs = new ConcurrentBag<BlobIdentifier>();
+
+            IEnumerable<Task> tasks = id.Select(async blob =>
+            {
+                if (!await _storage.Exists(ns, blob))
+                {
+                    missingBlobs.Add(blob);
+                }
+            });
+            await Task.WhenAll(tasks);
+
+            return Ok(new HeadMultipleResponse { Needs = missingBlobs.ToArray()});
+        }
+
+        [HttpPost("{ns}/exist")]
+        [ProducesDefaultResponseType]
+        public async Task<IActionResult> ExistsBody(
+            [Required] NamespaceId ns,
+            [FromBody] BlobIdentifier[] bodyIds)
+        {
+            ActionResult? result = await _requestHelper.HasAccessToNamespace(User, Request, ns, new [] { JupiterAclAction.ReadObject });
+            if (result != null)
+            {
+                return result;
+            }
+
+            ConcurrentBag<BlobIdentifier> missingBlobs = new ConcurrentBag<BlobIdentifier>();
+
+            IEnumerable<Task> tasks = bodyIds.Select(async blob =>
+            {
+                if (!await _storage.Exists(ns, blob))
+                {
+                    missingBlobs.Add(blob);
+                }
+            });
+            await Task.WhenAll(tasks);
+
+            return Ok(new HeadMultipleResponse { Needs = missingBlobs.ToArray()});
+        }
+
+        private async Task<BlobContents> GetImpl(NamespaceId ns, BlobIdentifier blob, List<string>? storageLayers = null, bool supportsRedirectUri = false)
+        {
+            try
+            {
+                return await _storage.GetObject(ns, blob, storageLayers, supportsRedirectUri);
+            }
+            catch (BlobNotFoundException)
+            {
+                if (!_storage.ShouldFetchBlobOnDemand(ns))
+                {
+                    throw;
+                }
+
+                return await _storage.ReplicateObject(ns, blob);
+            }
+        }
+
+        [HttpPut("{ns}/{id}")]
+        [RequiredContentType(MediaTypeNames.Application.Octet)]
+        [DisableRequestSizeLimit]
+        public async Task<IActionResult> Put(
+            [Required] NamespaceId ns,
+            [Required] BlobIdentifier id)
+        {
+            ActionResult? result = await _requestHelper.HasAccessToNamespace(User, Request, ns, new [] { JupiterAclAction.WriteObject });
+            if (result != null)
+            {
+                return result;
+            }
+
+            _diagnosticContext.Set("Content-Length", Request.ContentLength ?? -1);
+
+            try
+            {
+                Uri? uri = await _storage.MaybePutObjectWithRedirect(ns, id);
+                if (uri != null)
+                {
+                    return Ok(new
+                    {
+                        Identifier = id.ToString(),
+                        RedirectUri = uri,
+                    });
+                }
+                using IBufferedPayload payload = await _bufferedPayloadFactory.CreateFromRequest(Request);
+
+                BlobIdentifier identifier = await _storage.PutObject(ns, payload, id);
+                return Ok(new
+                {
+                    Identifier = identifier.ToString()
+                });
+            }
+            catch (ResourceHasToManyRequestsException)
+            {
+                return StatusCode(StatusCodes.Status429TooManyRequests);
+            }
+            catch (ClientSendSlowException e)
+            {
+                return Problem(e.Message, null, (int)HttpStatusCode.RequestTimeout);
+            }
+        }
+
+        [HttpPost("{ns}")]
+        [RequiredContentType(MediaTypeNames.Application.Octet)]
+        [DisableRequestSizeLimit]
+        public async Task<IActionResult> Post(
+            [Required] NamespaceId ns)
+        {
+            ActionResult? result = await _requestHelper.HasAccessToNamespace(User, Request, ns, new [] { JupiterAclAction.WriteObject });
+            if (result != null)
+            {
+                return result;
+            }
+
+            _diagnosticContext.Set("Content-Length", Request.ContentLength ?? -1);
+            try
+            {
+                using IBufferedPayload payload = await _bufferedPayloadFactory.CreateFromRequest(Request);
+
+                await using Stream stream = payload.GetStream();
+
+                BlobIdentifier id = await BlobIdentifier.FromStream(stream);
+                await _storage.PutObjectKnownHash(ns, payload, id);
+                
+                return Ok(new
+                {
+                    Identifier = id.ToString()
+                });
+            }
+            catch (ClientSendSlowException e)
+            {
+                return Problem(e.Message, null, (int)HttpStatusCode.RequestTimeout);
+            }
+        }
+
+        [HttpDelete("{ns}/{id}")]
+        public async Task<IActionResult> Delete(
+            [Required] NamespaceId ns,
+            [Required] BlobIdentifier id)
+        {
+            ActionResult? result = await _requestHelper.HasAccessToNamespace(User, Request, ns, new [] { JupiterAclAction.DeleteObject });
+            if (result != null)
+            {
+                return result;
+            }
+
+            await DeleteImpl(ns, id);
+
+            return NoContent();
+        }
+
+        
+        [HttpDelete("{ns}")]
+        public async Task<IActionResult> DeleteNamespace(
+            [Required] NamespaceId ns)
+        {
+            ActionResult? result = await _requestHelper.HasAccessToNamespace(User, Request, ns, new [] { JupiterAclAction.DeleteNamespace });
+            if (result != null)
+            {
+                return result;
+            }
+
+            await  _storage.DeleteNamespace(ns);
+
+            return NoContent();
+        }
+
+        private async Task DeleteImpl(NamespaceId ns, BlobIdentifier id)
+        {
+            await _storage.DeleteObject(ns, id);
+        }
+
+        // ReSharper disable UnusedAutoPropertyAccessor.Global
+        // ReSharper disable once ClassNeverInstantiated.Global
+        [System.Diagnostics.CodeAnalysis.SuppressMessage("Design", "CA1034:Nested types should not be visible", Justification = "Used by serialization")]
+        public class BatchOp
+        {
+            // ReSharper disable once InconsistentNaming
+            public enum Operation
+            {
+                INVALID,
+                GET,
+                PUT,
+                DELETE,
+                HEAD
+            }
+
+            [Required] public NamespaceId? Namespace { get; set; }
+
+            public BlobIdentifier? Id { get; set; }
+
+            [Required] public Operation Op { get; set; }
+
+            public byte[]? Content { get; set; }
+        }
+
+        [System.Diagnostics.CodeAnalysis.SuppressMessage("Design", "CA1034:Nested types should not be visible", Justification = "Used by serialization")]
+        public class BatchCall
+        {
+            public BatchOp[]? Operations { get; set; }
+        }
+        // ReSharper restore UnusedAutoPropertyAccessor.Global
+
+        [HttpPost("")]
+        public async Task<IActionResult> Post([FromBody] BatchCall batch)
+        {
+            JupiterAclAction MapToAclAction(BatchOp.Operation op)
+            {
+                switch (op)
+                {
+                    case BatchOp.Operation.GET:
+                    case BatchOp.Operation.HEAD:
+                        return JupiterAclAction.ReadObject;
+                    case BatchOp.Operation.PUT:
+                        return JupiterAclAction.WriteObject;
+                    case BatchOp.Operation.DELETE:
+                        return JupiterAclAction.DeleteObject;
+                    default:
+                        throw new ArgumentOutOfRangeException(nameof(op), op, null);
+                }
+            }
+
+            if (batch?.Operations == null)
+            {
+                throw new ();
+            }
+
+            Task<object?>[] tasks = new Task<object?>[batch.Operations.Length];
+            for (int index = 0; index < batch.Operations.Length; index++)
+            {
+                BatchOp op = batch.Operations[index];
+                if (op.Namespace == null)
+                {
+                    throw new Exception(nameof(op.Namespace));
+                }
+
+                ActionResult? result = await _requestHelper.HasAccessToNamespace(User, Request, op.Namespace!.Value, new[] { MapToAclAction(op.Op) });
+
+                if (result != null)
+                {
+                    return result;
+                }
+
+                switch (op.Op)
+                {
+                    case BatchOp.Operation.INVALID:
+                        throw new InvalidOperationException($"Op type {BatchOp.Operation.INVALID} is not a valid op type for operation id {op.Id}");
+                    case BatchOp.Operation.GET:
+                        if (op.Id == null)
+                        {
+                            return BadRequest();
+                        }
+
+                        tasks[index] = GetImpl(op.Namespace.Value, op.Id).ContinueWith((t, _) =>
+                        {
+                            // TODO: This is very allocation heavy but given that the end result is a json object we can not really stream this anyway
+                            using BlobContents blobContents = t.Result;
+
+                            using MemoryStream ms = new MemoryStream();
+                            blobContents.Stream.CopyTo(ms);
+                            ms.Seek(0, SeekOrigin.Begin);
+                            string str = Convert.ToBase64String(ms.ToArray());
+                            return (object?) str;
+                        }, null, TaskScheduler.Current);
+                        break;
+                    case BatchOp.Operation.HEAD:
+                        if (op.Id == null)
+                        {
+                            return BadRequest();
+                        }
+
+                        tasks[index] = _storage.Exists(op.Namespace.Value, op.Id)
+                            .ContinueWith((t,_) => t.Result ? (object?) null : op.Id, null, TaskScheduler.Current);
+                        break;
+                    case BatchOp.Operation.PUT:
+                        {
+                            if (op.Content == null)
+                            {
+                                return BadRequest();
+                            }
+
+                            if (op.Id == null)
+                            {
+                                return BadRequest();
+                            }
+
+                            using MemoryBufferedPayload payload = new MemoryBufferedPayload(op.Content);
+                            tasks[index] = _storage.PutObject(op.Namespace.Value, payload, op.Id).ContinueWith((t, _) => (object?) t.Result, null, TaskScheduler.Current);
+                            break;
+                        }
+                    case BatchOp.Operation.DELETE:
+                        if (op.Id == null)
+                        {
+                            return BadRequest();
+                        }
+
+                        tasks[index] = DeleteImpl(op.Namespace.Value, op.Id).ContinueWith((t, _) => (object?) null, null, TaskScheduler.Current);
+                        break;
+                    default:
+                        throw new NotImplementedException($"{op.Op} is not a support op type");
+                }
+            }
+
+            await Task.WhenAll(tasks);
+
+            object?[] results = tasks.Select(t => t.Result).ToArray();
+
+            return Ok(results);
+        }
+    }
+
+    public class HeadMultipleResponse
+    {
+        [CbField("needs")]
+        public BlobIdentifier[] Needs { get; set; } = null!;
+    }
+}
diff --git a/Engine/Source/Programs/Horde/Horde.Server/Ddc/BlobIdentifier.cs b/Engine/Source/Programs/Horde/Horde.Server/Ddc/BlobIdentifier.cs
new file mode 100644
--- /dev/null
+++ b/Engine/Source/Programs/Horde/Horde.Server/Ddc/BlobIdentifier.cs
@@ -0,0 +1,241 @@
+﻿// Copyright Epic Games, Inc. All Rights Reserved.
+
+using System;
+using System.ComponentModel;
+using System.Globalization;
+using System.IO;
+using System.Text;
+using System.Text.Json;
+using System.Text.Json.Serialization;
+using System.Threading.Tasks;
+using Blake3;
+using EpicGames.Core;
+using EpicGames.Horde.Storage;
+using EpicGames.Horde.Storage.Bundles;
+using EpicGames.Serialization;
+
+#pragma warning disable CS1591
+
+namespace Jupiter.Implementation
+{
+    [TypeConverter(typeof(BlobIdentifierTypeConverter))]
+    [JsonConverter(typeof(BlobIdentifierJsonConverter))]
+    [CbConverter(typeof(BlobIdentifierCbConverter))]
+    public class BlobIdentifier : JupiterContentHash,  IEquatable<BlobIdentifier>
+    {
+        // multi thread the hashing for blobs larger then this size
+        private const int MultiThreadedSize = 1_000_000;
+        private string? _stringIdentifier;
+
+        public BlobIdentifier(byte[] identifier) : base(identifier)
+        {
+        }
+
+        [JsonConstructor]
+        public BlobIdentifier(string identifier) : base(identifier)
+        {
+
+        }
+
+        public override int GetHashCode()
+        {
+            return Comparer.GetHashCode(Identifier);
+        }
+
+        public bool Equals(BlobIdentifier? other)
+        {
+            if (other == null)
+            {
+                return false;
+            }
+
+            return Comparer.Equals(Identifier, other.Identifier);
+        }
+
+        public override bool Equals(object? obj)
+        {
+            if (ReferenceEquals(null, obj))
+            {
+                return false;
+            }
+
+            if (ReferenceEquals(this, obj))
+            {
+                return true;
+            }
+
+            if (obj.GetType() != GetType())
+            {
+                return false;
+            }
+
+            return Equals((BlobIdentifier) obj);
+        }
+
+        public override string ToString()
+        {
+            if (_stringIdentifier == null)
+            {
+                _stringIdentifier = StringUtils.FormatAsHexString(Identifier);
+            }
+
+            return _stringIdentifier;
+        }
+
+        public static new BlobIdentifier FromBlob(byte[] blobMemory)
+        {
+            Hash blake3Hash;
+            if (blobMemory.Length < MultiThreadedSize)
+            {
+                using Hasher hasher = Hasher.New();
+                hasher.Update(blobMemory);
+                blake3Hash = hasher.Finalize();
+            }
+            else
+            {
+                using Hasher hasher = Hasher.New();
+                hasher.UpdateWithJoin(blobMemory);
+                blake3Hash = hasher.Finalize();
+            }
+            
+            // we only keep the first 20 bytes of the Blake3 hash
+            Span<byte> hash = blake3Hash.AsSpanUnsafe().Slice(0, 20);
+            return new BlobIdentifier(hash.ToArray());
+        }
+
+        public static BlobIdentifier FromBlob(in Memory<byte> blobMemory)
+        {
+            Hash blake3Hash;
+            if (blobMemory.Length < MultiThreadedSize)
+            {
+                using Hasher hasher = Hasher.New();
+                hasher.Update(blobMemory.Span);
+                blake3Hash = hasher.Finalize();
+            }
+            else
+            {
+                using Hasher hasher = Hasher.New();
+                hasher.UpdateWithJoin(blobMemory.Span);
+                blake3Hash = hasher.Finalize();
+            }
+
+            // we only keep the first 20 bytes of the Blake3 hash
+            Span<byte> hash = blake3Hash.AsSpanUnsafe().Slice(0, 20);
+            return new BlobIdentifier(hash.ToArray());
+        }
+
+        public static BlobIdentifier FromContentHash(JupiterContentHash testObjectHash)
+        {
+            return new BlobIdentifier(testObjectHash.HashData);
+        }
+
+        public static async Task<BlobIdentifier> FromStream(Stream stream)
+        {
+            using Hasher hasher = Hasher.New();
+            const int bufferSize = 1024 * 1024 * 5;
+            byte[] buffer = new byte[bufferSize];
+            int read = await stream.ReadAsync(buffer, 0, buffer.Length);
+            while (read > 0)
+            {
+                hasher.UpdateWithJoin(new ReadOnlySpan<byte>(buffer, 0, read));
+                read = await stream.ReadAsync(buffer, 0, buffer.Length);
+            }
+            Hash blake3Hash = hasher.Finalize();
+
+            // we only keep the first 20 bytes of the Blake3 hash
+            byte[] hash = blake3Hash.AsSpanUnsafe().Slice(0, 20).ToArray();
+            return new BlobIdentifier(hash);
+        }
+
+        public static BlobIdentifier FromIoHash(IoHash blobIdentifier)
+        {
+            return new BlobIdentifier(blobIdentifier.ToByteArray());
+        }
+
+        public IoHash AsIoHash()
+        {
+            return new IoHash(HashData);
+        }
+
+        public static BlobIdentifier FromBlobLocator(BundleLocator locator)
+        {
+            return new BlobIdentifier(Encoding.UTF8.GetBytes(locator.ToString()));
+        }
+
+        public BundleLocator AsBlobLocator()
+        {
+            return new BundleLocator(Encoding.UTF8.GetString(HashData));
+        }
+    }
+
+    public class BlobIdentifierTypeConverter : TypeConverter
+    {
+        public override bool CanConvertFrom(ITypeDescriptorContext? context, Type sourceType)
+        {  
+            if (sourceType == typeof(string))  
+            {  
+                return true;
+            }  
+            return base.CanConvertFrom(context, sourceType);
+        }  
+  
+        public override object? ConvertFrom(ITypeDescriptorContext? context, CultureInfo? culture, object value)  
+        {
+            if (value is string s)
+            {
+                return new BlobIdentifier(s);
+            }
+
+            return base.ConvertFrom(context, culture, value);  
+        }
+
+        public override bool CanConvertTo(ITypeDescriptorContext? context, Type? destinationType)
+        {
+            if (destinationType == typeof(string))
+            {
+                return true;
+            }
+            return base.CanConvertTo(context, destinationType);
+        }
+
+        public override object? ConvertTo(ITypeDescriptorContext? context, CultureInfo? culture, object? value, Type destinationType)
+        {
+            if (destinationType == typeof(string))
+            {
+                BlobIdentifier? identifier = (BlobIdentifier?)value;
+                return identifier?.ToString();
+            }
+            return base.ConvertTo(context, culture, value, destinationType);
+        }
+    }
+
+    public class BlobIdentifierJsonConverter : JsonConverter<BlobIdentifier>
+    {
+        public override BlobIdentifier? Read(ref Utf8JsonReader reader, Type typeToConvert, JsonSerializerOptions options)
+        {
+            string? str = reader.GetString();
+            if (str == null)
+            {
+                throw new InvalidDataException("Unable to parse blob identifier");
+            }
+
+            return new BlobIdentifier(str);
+        }
+
+        public override void Write(Utf8JsonWriter writer, BlobIdentifier value, JsonSerializerOptions options)
+        {
+            writer.WriteStringValue(value.ToString());
+        }
+    }
+
+    public class BlobIdentifierCbConverter : CbConverterBase<BlobIdentifier>
+    {
+        public override BlobIdentifier Read(CbField field) => new BlobIdentifier(field.AsHash().ToByteArray());
+
+        /// <inheritdoc/>
+        public override void Write(CbWriter writer, BlobIdentifier value) => writer.WriteBinaryAttachmentValue(new IoHash(value.HashData));
+
+        /// <inheritdoc/>
+        public override void WriteNamed(CbWriter writer, Utf8String name, BlobIdentifier value) => writer.WriteBinaryAttachment(name, new IoHash(value.HashData));
+    }
+}
diff --git a/Engine/Source/Programs/Horde/Horde.Server/Ddc/BlobService.cs b/Engine/Source/Programs/Horde/Horde.Server/Ddc/BlobService.cs
new file mode 100644
--- /dev/null
+++ b/Engine/Source/Programs/Horde/Horde.Server/Ddc/BlobService.cs
@@ -0,0 +1,100 @@
+// Copyright Epic Games, Inc. All Rights Reserved.
+
+using System;
+using System.Collections.Generic;
+using System.IO;
+using System.Threading.Tasks;
+using EpicGames.Horde.Storage;
+using Jupiter.Common.Implementation;
+using Jupiter.Implementation;
+
+namespace Horde.Server.Ddc
+{
+	class BlobService : IBlobService
+	{
+		public Task DeleteNamespace(NamespaceId ns)
+		{
+			throw new NotImplementedException();
+		}
+
+		public Task DeleteObject(NamespaceId ns, BlobIdentifier blob)
+		{
+			throw new NotImplementedException();
+		}
+
+		public Task<bool> Exists(NamespaceId ns, BlobIdentifier blob, List<string>? storageLayers = null)
+		{
+			throw new NotImplementedException();
+		}
+
+		public Task<bool> ExistsInRootStore(NamespaceId ns, BlobIdentifier blob)
+		{
+			throw new NotImplementedException();
+		}
+
+		public Task<BlobIdentifier[]> FilterOutKnownBlobs(NamespaceId ns, IEnumerable<BlobIdentifier> blobs)
+		{
+			throw new NotImplementedException();
+		}
+
+		public Task<BlobIdentifier[]> FilterOutKnownBlobs(NamespaceId ns, IAsyncEnumerable<BlobIdentifier> blobs)
+		{
+			throw new NotImplementedException();
+		}
+
+		public Task<BlobContents> GetObject(NamespaceId ns, BlobIdentifier blob, List<string>? storageLayers = null, bool supportsRedirectUri = false)
+		{
+			throw new NotImplementedException();
+		}
+
+		public Task<BlobContents> GetObjects(NamespaceId ns, BlobIdentifier[] refRequestBlobReferences)
+		{
+			throw new NotImplementedException();
+		}
+
+		public Task<Uri?> GetObjectWithRedirect(NamespaceId ns, BlobIdentifier blobIdentifier, List<string>? storageLayers = null)
+		{
+			throw new NotImplementedException();
+		}
+
+		public IAsyncEnumerable<(BlobIdentifier, DateTime)> ListObjects(NamespaceId ns)
+		{
+			throw new NotImplementedException();
+		}
+
+		public Task<Uri?> MaybePutObjectWithRedirect(NamespaceId ns, BlobIdentifier identifier)
+		{
+			throw new NotImplementedException();
+		}
+
+		public Task<BlobIdentifier> PutObject(NamespaceId ns, IBufferedPayload payload, BlobIdentifier identifier)
+		{
+			throw new NotImplementedException();
+		}
+
+		public Task<BlobIdentifier> PutObject(NamespaceId ns, byte[] payload, BlobIdentifier identifier)
+		{
+			throw new NotImplementedException();
+		}
+
+		public Task<BlobIdentifier> PutObjectKnownHash(NamespaceId ns, IBufferedPayload content, BlobIdentifier identifier)
+		{
+			throw new NotImplementedException();
+		}
+
+		public Task<BlobContents> ReplicateObject(NamespaceId ns, BlobIdentifier blob, bool force = false)
+		{
+			throw new NotImplementedException();
+		}
+
+		public bool ShouldFetchBlobOnDemand(NamespaceId ns)
+		{
+			throw new NotImplementedException();
+		}
+
+		public Task<JupiterContentHash> VerifyContentMatchesHash(Stream content, JupiterContentHash identifier)
+		{
+			throw new NotImplementedException();
+		}
+	}
+}
diff --git a/Engine/Source/Programs/Horde/Horde.Server/Ddc/BlobStore.cs b/Engine/Source/Programs/Horde/Horde.Server/Ddc/BlobStore.cs
new file mode 100644
--- /dev/null
+++ b/Engine/Source/Programs/Horde/Horde.Server/Ddc/BlobStore.cs
@@ -0,0 +1,64 @@
+// Copyright Epic Games, Inc. All Rights Reserved.
+
+using System;
+using System.Collections.Generic;
+using System.IO;
+using System.Threading.Tasks;
+using EpicGames.Horde.Storage;
+using Jupiter.Implementation;
+
+namespace Horde.Server.Ddc
+{
+	class BlobStore : IBlobStore
+	{
+		public Task DeleteNamespace(NamespaceId ns)
+		{
+			throw new NotImplementedException();
+		}
+
+		public Task DeleteObject(NamespaceId ns, BlobIdentifier blob)
+		{
+			throw new NotImplementedException();
+		}
+
+		public Task<bool> Exists(NamespaceId ns, BlobIdentifier blob, bool forceCheck = false)
+		{
+			throw new NotImplementedException();
+		}
+
+		public Task<BlobContents> GetObject(NamespaceId ns, BlobIdentifier blob, LastAccessTrackingFlags flags = LastAccessTrackingFlags.DoTracking, bool supportsRedirectUri = false)
+		{
+			throw new NotImplementedException();
+		}
+
+		public Task<Uri?> GetObjectByRedirect(NamespaceId ns, BlobIdentifier blob)
+		{
+			throw new NotImplementedException();
+		}
+
+		public IAsyncEnumerable<(BlobIdentifier, DateTime)> ListObjects(NamespaceId ns)
+		{
+			throw new NotImplementedException();
+		}
+
+		public Task<BlobIdentifier> PutObject(NamespaceId ns, byte[] blob, BlobIdentifier identifier)
+		{
+			throw new NotImplementedException();
+		}
+
+		public Task<BlobIdentifier> PutObject(NamespaceId ns, ReadOnlyMemory<byte> blob, BlobIdentifier identifier)
+		{
+			throw new NotImplementedException();
+		}
+
+		public Task<BlobIdentifier> PutObject(NamespaceId ns, Stream content, BlobIdentifier identifier)
+		{
+			throw new NotImplementedException();
+		}
+
+		public Task<Uri?> PutObjectWithRedirect(NamespaceId ns, BlobIdentifier identifier)
+		{
+			throw new NotImplementedException();
+		}
+	}
+}
diff --git a/Engine/Source/Programs/Horde/Horde.Server/Ddc/BucketId.cs b/Engine/Source/Programs/Horde/Horde.Server/Ddc/BucketId.cs
new file mode 100644
--- /dev/null
+++ b/Engine/Source/Programs/Horde/Horde.Server/Ddc/BucketId.cs
@@ -0,0 +1,88 @@
+// Copyright Epic Games, Inc. All Rights Reserved.
+
+using System;
+using System.ComponentModel;
+using System.Globalization;
+using System.Text.Json;
+using System.Text.Json.Serialization;
+using EpicGames.Core;
+using EpicGames.Serialization;
+
+namespace EpicGames.Horde.Storage
+{
+	/// <summary>
+	/// Identifier for a storage bucket
+	/// </summary>
+	[JsonConverter(typeof(BucketIdJsonConverter))]
+	[TypeConverter(typeof(BucketIdTypeConverter))]
+	[CbConverter(typeof(BucketIdCbConverter))]
+	public struct BucketId : IEquatable<BucketId>
+	{
+		/// <summary>
+		/// The text representing this id
+		/// </summary>
+		readonly StringId _inner;
+
+		/// <summary>
+		/// Constructor
+		/// </summary>
+		/// <param name="input">Unique id for the string</param>
+		public BucketId(string input)
+		{
+			_inner = new StringId(input);
+		}
+
+		/// <inheritdoc/>
+		public override bool Equals(object? obj) => obj is BucketId id && _inner.Equals(id._inner);
+
+		/// <inheritdoc/>
+		public override int GetHashCode() => _inner.GetHashCode();
+
+		/// <inheritdoc/>
+		public bool Equals(BucketId other) => _inner.Equals(other._inner);
+
+		/// <inheritdoc/>
+		public override string ToString() => _inner.ToString();
+
+		/// <inheritdoc cref="StringId.op_Equality"/>
+		public static bool operator ==(BucketId left, BucketId right) => left._inner == right._inner;
+
+		/// <inheritdoc cref="StringId.op_Inequality"/>
+		public static bool operator !=(BucketId left, BucketId right) => left._inner != right._inner;
+	}
+
+	/// <summary>
+	/// Type converter for BucketId to and from JSON
+	/// </summary>
+	sealed class BucketIdJsonConverter : JsonConverter<BucketId>
+	{
+		/// <inheritdoc/>
+		public override BucketId Read(ref Utf8JsonReader reader, Type typeToConvert, JsonSerializerOptions options) => new BucketId(reader.GetString() ?? String.Empty);
+
+		/// <inheritdoc/>
+		public override void Write(Utf8JsonWriter writer, BucketId value, JsonSerializerOptions options) => writer.WriteStringValue(value.ToString());
+	}
+
+	/// <summary>
+	/// Type converter from strings to BucketId objects
+	/// </summary>
+	sealed class BucketIdTypeConverter : TypeConverter
+	{
+		/// <inheritdoc/>
+		public override bool CanConvertFrom(ITypeDescriptorContext? context, Type? sourceType) => sourceType == typeof(string);
+
+		/// <inheritdoc/>
+		public override object ConvertFrom(ITypeDescriptorContext? context, CultureInfo? culture, object? value) => new BucketId((string)value!);
+	}
+
+	sealed class BucketIdCbConverter : CbConverterBase<BucketId>
+	{
+		public override BucketId Read(CbField field) => new BucketId(field.AsString());
+
+		/// <inheritdoc/>
+		public override void Write(CbWriter writer, BucketId value) => writer.WriteStringValue(value.ToString());
+
+		/// <inheritdoc/>
+		public override void WriteNamed(CbWriter writer, Utf8String name, BucketId value) => writer.WriteString(name, value.ToString());
+	}
+}
diff --git a/Engine/Source/Programs/Horde/Horde.Server/Ddc/CommonHeaders.cs b/Engine/Source/Programs/Horde/Horde.Server/Ddc/CommonHeaders.cs
new file mode 100644
--- /dev/null
+++ b/Engine/Source/Programs/Horde/Horde.Server/Ddc/CommonHeaders.cs
@@ -0,0 +1,16 @@
+// Copyright Epic Games, Inc. All Rights Reserved.
+
+#pragma warning disable CS1591
+
+namespace Jupiter.Implementation
+{
+    public static class CommonHeaders
+    {
+        public const string HashHeaderSHA1Name = "X-Jupiter-Sha1";
+        public const string HashHeaderName = "X-Jupiter-IoHash";
+
+        public const string LastAccessHeaderName = "X-Jupiter-LastAccess";
+
+        public const string InlinePayloadHash = "X-Jupiter-InlinePayloadHash";
+    }
+}
diff --git a/Engine/Source/Programs/Horde/Horde.Server/Ddc/CompressedBufferUtils.cs b/Engine/Source/Programs/Horde/Horde.Server/Ddc/CompressedBufferUtils.cs
new file mode 100644
--- /dev/null
+++ b/Engine/Source/Programs/Horde/Horde.Server/Ddc/CompressedBufferUtils.cs
@@ -0,0 +1,468 @@
+﻿// Copyright Epic Games, Inc. All Rights Reserved.
+
+using System;
+using System.Buffers.Binary;
+using System.Collections.Generic;
+using System.IO;
+using System.Linq;
+using EpicGames.Core;
+using EpicGames.Compression;
+using Force.Crc32;
+using K4os.Compression.LZ4;
+using OpenTelemetry.Trace;
+
+#pragma warning disable CS1591
+
+namespace Jupiter.Implementation
+{
+    public class CompressedBufferUtils
+    {
+        private readonly Tracer _tracer;
+
+        public CompressedBufferUtils(Tracer tracer)
+        {
+            _tracer = tracer;
+        }
+
+        private class Header
+        {
+            [System.Diagnostics.CodeAnalysis.SuppressMessage("Design", "CA1028:Enum Storage should be Int32", Justification = "Interop requires byte")]
+            public enum CompressionMethod : byte
+            {
+                // Header is followed by one uncompressed block. 
+                None = 0,
+                // Header is followed by an array of compressed block sizes then the compressed blocks. 
+                Oodle = 3,
+                LZ4 = 4,
+            }
+
+            public const uint ExpectedMagic = 0xb7756362; // <dot>ucb
+            public const uint HeaderLength = 64;
+
+            // A magic number to identify a compressed buffer. Always 0xb7756362.
+
+            public uint Magic { get; set; }
+            // A CRC-32 used to check integrity of the buffer. Uses the polynomial 0x04c11db7.
+
+            public uint Crc32 { get; set; }
+
+            // The method used to compress the buffer. Affects layout of data following the header. 
+            public CompressionMethod Method { get; set; }
+            public byte CompressionLevel { get; set; }
+            public byte CompressionMethodUsed { get; set; }
+
+            // The power of two size of every uncompressed block except the last. Size is 1 << BlockSizeExponent. 
+            public byte BlockSizeExponent { get; set; }
+
+            // The number of blocks that follow the header. 
+            public uint BlockCount { get; set; }
+
+            // The total size of the uncompressed data. 
+            public ulong TotalRawSize { get; set; }
+
+            // The total size of the compressed data including the header. 
+            public ulong TotalCompressedSize { get; set; }
+
+            /** The hash of the uncompressed data. */
+            public byte[] RawHash { get; set; } = Array.Empty<byte>();
+
+            public void ByteSwap()
+            {
+                Magic = BinaryPrimitives.ReverseEndianness(Magic);
+                Crc32 = BinaryPrimitives.ReverseEndianness(Crc32);
+                BlockCount = BinaryPrimitives.ReverseEndianness(BlockCount);
+                TotalRawSize = BinaryPrimitives.ReverseEndianness(TotalRawSize);
+                TotalCompressedSize = BinaryPrimitives.ReverseEndianness(TotalCompressedSize);
+            }
+        }
+
+        private static Header ExtractHeader(byte[] content)
+        {
+            // the header is always stored big endian
+            bool needsByteSwap = BitConverter.IsLittleEndian;
+            if (content.Length < Header.HeaderLength)
+            {
+                throw new ArgumentOutOfRangeException(nameof(content), $"Content was less then {Header.HeaderLength} bytes and thus is not a compressed buffer");
+            }
+
+            using MemoryStream ms = new MemoryStream(content);
+            using BinaryReader reader = new BinaryReader(ms);
+
+            Header header = new Header
+            {
+                Magic = reader.ReadUInt32(),
+                Crc32 = reader.ReadUInt32(),
+                Method = (Header.CompressionMethod)reader.ReadByte(),
+                CompressionLevel = reader.ReadByte(),
+                CompressionMethodUsed = reader.ReadByte(),
+                BlockSizeExponent = reader.ReadByte(),
+                BlockCount = reader.ReadUInt32(),
+                TotalRawSize = reader.ReadUInt64(),
+                TotalCompressedSize = reader.ReadUInt64()
+            };
+            byte[] hash = reader.ReadBytes(32); // a full blake3 hash
+            header.RawHash = hash;
+
+            if (needsByteSwap)
+            {
+                header.ByteSwap();
+            }
+
+            if (ms.Position != Header.HeaderLength)
+            {
+                throw new Exception($"Read {ms.Position} bytes but expected to read {Header.HeaderLength}");
+            }
+
+            if (header.Magic != Header.ExpectedMagic)
+            {
+                throw new InvalidMagicException(header.Magic, Header.ExpectedMagic);
+            }
+
+            // calculate the crc from the start of the method field (skipping magic which is a constant and the crc field itself)
+            const int MethodOffset = sizeof(uint) + sizeof(uint);
+
+            // none compressed objects have no extra blocks
+            uint blocksByteUsed = header.Method != Header.CompressionMethod.None ? header.BlockCount * (uint)sizeof(uint) : 0;
+            uint calculatedCrc = Crc32Algorithm.Compute(content, MethodOffset, (int)(Header.HeaderLength - MethodOffset + blocksByteUsed));
+
+            if (header.Crc32 != calculatedCrc)
+            {
+                throw new InvalidHashException(header.Crc32, calculatedCrc);
+            }
+
+            return header;
+        }
+
+        static void WriteHeader(Header header, BinaryWriter writer)
+        {
+            // the header is always stored big endian
+            bool needsByteSwap = BitConverter.IsLittleEndian;
+            if (needsByteSwap)
+            {
+                header.ByteSwap();
+            }
+
+            writer.Write(header.Magic);
+            writer.Write(header.Crc32);
+            writer.Write((byte)header.Method);
+            writer.Write((byte)header.CompressionLevel);
+            writer.Write((byte)header.CompressionMethodUsed);
+            writer.Write((byte)header.BlockSizeExponent);
+            writer.Write(header.BlockCount);
+            writer.Write(header.TotalRawSize);
+            writer.Write(header.TotalCompressedSize);
+            writer.Write(header.RawHash, 0, 20); // write the first 20 bytes as iohashes are 20 bytes
+            for (int i = 0; i < 12; i++)
+            {
+                // the last 12 bytes should be 0 as they are reserved
+                writer.Write((byte)0);
+            }
+
+            if (needsByteSwap)
+            {
+                header.ByteSwap();
+            }
+        }
+
+        public byte[] DecompressContent(byte[] content)
+        {
+            Header header = ExtractHeader(content);
+
+            if (content.LongLength < (long)header.TotalCompressedSize)
+            {
+                throw new Exception($"Expected buffer to be {header.TotalCompressedSize} but it was {content.LongLength}");
+            }
+
+            ulong decompressedPayloadOffset = 0;
+            byte[] decompressedPayload = new byte[header.TotalRawSize];
+
+            ReadOnlySpan<byte> memory = new ReadOnlySpan<byte>(content);
+            memory = memory.Slice((int)Header.HeaderLength);
+
+            bool willHaveBlocks = header.Method != Header.CompressionMethod.None;
+            if (willHaveBlocks)
+            {
+                uint[] compressedBlockSizes = new uint[header.BlockCount];
+                for (int i = 0; i < header.BlockCount; i++)
+                {
+                    uint compressedBlockSize = BinaryPrimitives.ReadUInt32BigEndian(memory);
+                    compressedBlockSizes[i] = compressedBlockSize;
+                    memory = memory.Slice(sizeof(uint));
+                }
+
+                ulong blockSize = 1ul << header.BlockSizeExponent;
+                ulong compressedOffset = 0;
+
+                foreach (uint compressedBlockSize in compressedBlockSizes)
+                {
+                    ulong rawBlockSize = Math.Min(header.TotalRawSize - decompressedPayloadOffset, blockSize);
+                    ReadOnlySpan<byte> compressedPayload = memory.Slice((int)compressedOffset, (int)compressedBlockSize);
+                    Span<byte> targetSpan = new Span<byte>(decompressedPayload, (int)decompressedPayloadOffset, (int)rawBlockSize);
+
+                    int writtenBytes;
+                    // if a block has the same raw and compressed size its uncompressed and we should not attempt to decompress it
+                    if (rawBlockSize == compressedBlockSize)
+                    {
+                        writtenBytes = (int)rawBlockSize;
+                        compressedPayload.CopyTo(targetSpan);
+                    }
+                    else
+                    {
+                        writtenBytes = DecompressPayload(compressedPayload, header, rawBlockSize, targetSpan);
+                    }
+
+                    decompressedPayloadOffset += (uint)writtenBytes;
+                    compressedOffset += compressedBlockSize;
+                }
+            }
+            else
+            {
+                // if no compression is applied there are no extra blocks and just a single chunk that is uncompressed
+                Span<byte> targetSpan = new Span<byte>(decompressedPayload, 0, (int)header.TotalRawSize);
+                DecompressPayload(memory, header, header.TotalRawSize, targetSpan);
+                decompressedPayloadOffset = header.TotalRawSize;
+            }
+
+            if (header.TotalRawSize != decompressedPayloadOffset)
+            {
+                throw new Exception("Did not decompress the full payload");
+            }
+
+            {
+                using TelemetrySpan _ = _tracer.StartActiveSpan("web.hash").SetAttribute("operation.name", "web.hash");
+
+                // only read the first 20 bytes of the hash field as IoHashes are 20 bytes and not 32 bytes
+                byte[] slicedHash = new byte[20];
+                Array.Copy(header.RawHash, 0, slicedHash, 0, 20);
+
+                BlobIdentifier headerIdentifier = new BlobIdentifier(slicedHash);
+                BlobIdentifier contentHash = BlobIdentifier.FromBlob(decompressedPayload);
+
+                if (!headerIdentifier.Equals(contentHash))
+                {
+                    throw new Exception($"Payload was expected to be {headerIdentifier} but was {contentHash}");
+                }
+            }
+
+            return decompressedPayload;
+        }
+
+        private static int DecompressPayload(ReadOnlySpan<byte> compressedPayload, Header header, ulong rawBlockSize, Span<byte> target)
+        {
+            switch (header.Method)
+            {
+                case Header.CompressionMethod.None:
+                    compressedPayload.CopyTo(target);
+                    return compressedPayload.Length;
+                case Header.CompressionMethod.Oodle:
+                    {
+                        byte[] result = new byte[rawBlockSize];
+                        long writtenBytes = Oodle.Decompress(compressedPayload, result);
+                        if (writtenBytes == 0)
+                        {
+                            throw new Exception("Failed to run oodle decompress");
+                        }
+                        result.CopyTo(target);
+                        return (int)writtenBytes;
+                    }
+                case Header.CompressionMethod.LZ4:
+                    {
+                        int writtenBytes = LZ4Codec.Decode(compressedPayload, target);
+                        return writtenBytes;
+                    }
+                default:
+                    throw new NotImplementedException($"Method {header.Method} is not a support value");
+            }
+        }
+
+        public byte[] CompressContent(OoodleCompressorMethod method, OoodleCompressionLevel compressionLevel, byte[] rawContents)
+        {
+            OodleCompressorType oodleMethod = OodleUtils.ToOodleApiCompressor(method);
+            OodleCompressionLevel oodleLevel = OodleUtils.ToOodleApiCompressionLevel(compressionLevel);
+
+            const long DefaultBlockSize = 256 * 1024;
+            long blockSize = DefaultBlockSize;
+            long blockCount = (rawContents.LongLength + blockSize - 1) / blockSize;
+
+            byte blockSizeExponent = (byte)Math.Floor(Math.Log2(blockSize));
+
+            Span<byte> contentsSpan = new Span<byte>(rawContents);
+            List<byte[]> compressedBlocks = new List<byte[]>();
+
+            for (int i = 0; i < blockCount; i++)
+            {
+                int rawBlockSize = Math.Min(rawContents.Length - (i * (int)blockSize), (int)blockSize);
+                Span<byte> bufferToCompress = contentsSpan.Slice((int)(i * blockSize), rawBlockSize);
+
+                int maxSize = Oodle.MaximumOutputSize(oodleMethod, bufferToCompress.Length);
+                byte[] compressedBlock = new byte[maxSize];
+                long encodedSize = Oodle.Compress(oodleMethod, bufferToCompress.ToArray(), compressedBlock, oodleLevel);
+
+                if (encodedSize == 0)
+                {
+                    throw new Exception("Failed to compress content");
+                }
+
+                compressedBlocks.Add(compressedBlock);
+            }
+
+            uint compressedContentLength = (uint)compressedBlocks.Sum(b => b.LongLength);
+
+            Header header = new Header
+            {
+                Magic = Header.ExpectedMagic,
+                Crc32 = 0,
+                Method = Header.CompressionMethod.Oodle,
+                CompressionLevel = (byte)compressionLevel,
+                CompressionMethodUsed = (byte)method,
+                BlockSizeExponent = blockSizeExponent,
+                BlockCount = (uint)blockCount,
+                TotalRawSize = (ulong)rawContents.LongLength,
+                TotalCompressedSize = (ulong)compressedContentLength,
+                RawHash = IoHash.Compute(rawContents).ToByteArray()
+            };
+
+            uint blocksByteUsed = (uint)blockCount * sizeof(uint);
+
+            byte[] headerBuffer = new byte[Header.HeaderLength + blocksByteUsed + compressedContentLength];
+
+            // write the compressed buffer, but with the wrong crc which we update and rewrite later
+            {
+                using MemoryStream ms = new MemoryStream(headerBuffer);
+                using BinaryWriter writer = new BinaryWriter(ms);
+
+                WriteHeader(header, writer);
+
+                for (int i = 0; i < blockCount; i++)
+                {
+                    uint value = (uint)compressedBlocks[i].Length;
+                    if (BitConverter.IsLittleEndian)
+                    {
+                        value = BinaryPrimitives.ReverseEndianness(value);
+                    }
+                    writer.Write(value);
+                }
+
+                for (int i = 0; i < blockCount; i++)
+                {
+                    writer.Write(compressedBlocks[i]);
+                }
+            }
+
+            // calculate the crc from the start of the method field (skipping magic which is a constant and the crc field itself)
+            const int MethodOffset = sizeof(uint) + sizeof(uint);
+            uint calculatedCrc = Crc32Algorithm.Compute(headerBuffer, MethodOffset, (int)(Header.HeaderLength - MethodOffset + blocksByteUsed));
+            header.Crc32 = calculatedCrc;
+
+            // write the header again now that we have the crc
+            {
+                using MemoryStream ms = new MemoryStream(headerBuffer);
+                using BinaryWriter writer = new BinaryWriter(ms);
+                WriteHeader(header, writer);
+            }
+
+            return headerBuffer;
+        }
+    }
+
+    public class InvalidHashException : Exception
+    {
+        public InvalidHashException(uint headerCrc32, uint calculatedCrc) : base($"Header specified crc \"{headerCrc32}\" but calculated hash was \"{calculatedCrc}\"")
+        {
+
+        }
+    }
+
+    public class InvalidMagicException : Exception
+    {
+        public InvalidMagicException(uint headerMagic, uint expectedMagic) : base($"Header magic \"{headerMagic}\" was incorrect, expected to be {expectedMagic}")
+        {
+        }
+    }
+
+    // from OodleDataCompression.h , we define our own enums for oodle compressions used and convert to the ones expected in the oodle api
+#pragma warning disable CA1028 // Enum Storage should be Int32
+    public enum OoodleCompressorMethod : byte
+
+    {
+        NotSet = 0,
+        Selkie = 1,
+        Mermaid = 2,
+        Kraken = 3,
+        Leviathan = 4
+    }
+
+    public enum OoodleCompressionLevel : sbyte
+    {
+        HyperFast4 = -4,
+        HyperFast3 = -3,
+        HyperFast2 = -2,
+        HyperFast1 = -1,
+        None = 0,
+        SuperFast = 1,
+        VeryFast = 2,
+        Fast = 3,
+        Normal = 4,
+        Optimal1 = 5,
+        Optimal2 = 6,
+        Optimal3 = 7,
+        Optimal4 = 8,
+    }
+#pragma warning restore CA1028 // Enum Storage should be Int32
+    public static class OodleUtils
+    {
+        public static OodleCompressionLevel ToOodleApiCompressionLevel(OoodleCompressionLevel compressionLevel)
+        {
+            switch (compressionLevel)
+            {
+                case OoodleCompressionLevel.HyperFast4:
+                    return OodleCompressionLevel.HyperFast4;
+                case OoodleCompressionLevel.HyperFast3:
+                    return OodleCompressionLevel.HyperFast3;
+                case OoodleCompressionLevel.HyperFast2:
+                    return OodleCompressionLevel.HyperFast2;
+                case OoodleCompressionLevel.HyperFast1:
+                    return OodleCompressionLevel.HyperFast1;
+                case OoodleCompressionLevel.None:
+                    return OodleCompressionLevel.None;
+                case OoodleCompressionLevel.SuperFast:
+                    return OodleCompressionLevel.SuperFast;
+                case OoodleCompressionLevel.VeryFast:
+                    return OodleCompressionLevel.VeryFast;
+                case OoodleCompressionLevel.Fast:
+                    return OodleCompressionLevel.Fast;
+                case OoodleCompressionLevel.Normal:
+                    return OodleCompressionLevel.Normal;
+                case OoodleCompressionLevel.Optimal1:
+                    return OodleCompressionLevel.Optimal1;
+                case OoodleCompressionLevel.Optimal2:
+                    return OodleCompressionLevel.Optimal2;
+                case OoodleCompressionLevel.Optimal3:
+                    return OodleCompressionLevel.Optimal3;
+                case OoodleCompressionLevel.Optimal4:
+                    return OodleCompressionLevel.Optimal4;
+                default:
+                    throw new ArgumentOutOfRangeException(nameof(compressionLevel), compressionLevel, null);
+            }
+        }
+
+        public static OodleCompressorType ToOodleApiCompressor(OoodleCompressorMethod compressor)
+        {
+            switch (compressor)
+            {
+                case OoodleCompressorMethod.NotSet:
+                    return OodleCompressorType.None;
+                case OoodleCompressorMethod.Selkie:
+                    return OodleCompressorType.Selkie;
+                case OoodleCompressorMethod.Mermaid:
+                    return OodleCompressorType.Mermaid;
+                case OoodleCompressorMethod.Kraken:
+                    return OodleCompressorType.Kraken;
+                case OoodleCompressorMethod.Leviathan:
+                    return OodleCompressorType.Leviathan;
+                default:
+                    throw new ArgumentOutOfRangeException(nameof(compressor), compressor, null);
+            }
+        }
+    }
+}
diff --git a/Engine/Source/Programs/Horde/Horde.Server/Ddc/ContentId.cs b/Engine/Source/Programs/Horde/Horde.Server/Ddc/ContentId.cs
new file mode 100644
--- /dev/null
+++ b/Engine/Source/Programs/Horde/Horde.Server/Ddc/ContentId.cs
@@ -0,0 +1,108 @@
+// Copyright Epic Games, Inc. All Rights Reserved.
+
+using System;
+using System.ComponentModel;
+using System.Globalization;
+using System.Text.Json.Serialization;
+using EpicGames.Core;
+
+#pragma warning disable CS1591
+
+namespace Jupiter.Implementation
+{
+    [TypeConverter(typeof(ContentIdTypeConverter))]
+    public class ContentId : JupiterContentHash, IEquatable<ContentId>
+    {
+        public ContentId(byte[] identifier) : base(identifier)
+        {
+        }
+
+        [JsonConstructor]
+        public ContentId(string identifier) : base(identifier)
+        {
+
+        }
+
+        public override int GetHashCode()
+        {
+            return Comparer.GetHashCode(Identifier);
+        }
+
+        public bool Equals(ContentId? other)
+        {
+            if (other == null)
+            {
+                return false;
+            }
+
+            return Comparer.Equals(Identifier, other.Identifier);
+        }
+
+        public override bool Equals(object? obj)
+        {
+            if (ReferenceEquals(null, obj))
+            {
+                return false;
+            }
+
+            if (ReferenceEquals(this, obj))
+            {
+                return true;
+            }
+
+            if (obj.GetType() != GetType())
+            {
+                return false;
+            }
+
+            return Equals((ContentId) obj);
+        }
+
+        public static ContentId FromContentHash(JupiterContentHash contentHash)
+        {
+            return new ContentId(contentHash.HashData);
+        }
+
+        public static ContentId FromBlobIdentifier(BlobIdentifier blobIdentifier)
+        {
+            return new ContentId(blobIdentifier.HashData);
+        }
+
+        public BlobIdentifier AsBlobIdentifier()
+        {
+            return new BlobIdentifier(HashData);
+        }
+
+        public static ContentId FromIoHash(IoHash ioHash)
+        {
+            return new ContentId(ioHash.ToByteArray());
+        }
+
+        public IoHash AsIoHash()
+        {
+            return new IoHash(HashData);
+        }
+    }
+
+    public class ContentIdTypeConverter : TypeConverter
+    {
+        public override bool CanConvertFrom(ITypeDescriptorContext? context, Type sourceType)
+        {  
+            if (sourceType == typeof(string))  
+            {  
+                return true;
+            }  
+            return base.CanConvertFrom(context, sourceType);
+        }  
+  
+        public override object? ConvertFrom(ITypeDescriptorContext? context, CultureInfo? culture, object value)  
+        {
+            if (value is string s)
+            {
+                return new ContentId(s);
+            }
+
+            return base.ConvertFrom(context, culture, value);  
+        }  
+    }
+}
diff --git a/Engine/Source/Programs/Horde/Horde.Server/Ddc/ContentIdController.cs b/Engine/Source/Programs/Horde/Horde.Server/Ddc/ContentIdController.cs
new file mode 100644
--- /dev/null
+++ b/Engine/Source/Programs/Horde/Horde.Server/Ddc/ContentIdController.cs
@@ -0,0 +1,90 @@
+// Copyright Epic Games, Inc. All Rights Reserved.
+
+using System.Net.Mime;
+using System.Threading.Tasks;
+using EpicGames.AspNet;
+using EpicGames.Horde.Storage;
+using Jupiter.Implementation;
+using Microsoft.AspNetCore.Authorization;
+using Microsoft.AspNetCore.Mvc;
+
+#pragma warning disable CS1591
+
+using ContentId = Jupiter.Implementation.ContentId;
+
+namespace Jupiter.Controllers
+{
+    [ApiController]
+    [FormatFilter]
+    [Produces(MediaTypeNames.Application.Json, MediaTypeNames.Application.Octet, CustomMediaTypeNames.UnrealCompactBinary)]
+    [Route("api/v1/content-id")]
+    [Authorize]
+    public class ContentIdController : ControllerBase
+    {
+        private readonly IRequestHelper _requestHelper;
+        private readonly IContentIdStore _contentIdStore;
+
+        public ContentIdController(IRequestHelper requestHelper, IContentIdStore contentIdStore)
+        {
+            _requestHelper = requestHelper;
+            _contentIdStore = contentIdStore;
+        }
+
+        /// <summary>
+        /// Returns which blobs a content id maps to
+        /// </summary>
+        /// <param name="ns">Namespace. Each namespace is completely separated from each other. Use for different types of data that is never expected to be similar (between two different games for instance). Example: `uc4.ddc`</param>
+        /// <param name="contentId">The content id to resolve </param>
+        [HttpGet("{ns}/{contentId}.{format?}", Order = 500)]
+        [ProducesDefaultResponseType]
+        [ProducesResponseType(type: typeof(ProblemDetails), 400)]
+        public async Task<IActionResult> Resolve(NamespaceId ns, ContentId contentId)
+        {
+            ActionResult? result = await _requestHelper.HasAccessToNamespace(User, Request, ns, new [] { JupiterAclAction.ReadObject });
+            if (result != null)
+            {
+                return result;
+            }
+            BlobIdentifier[]? blobs = await _contentIdStore.Resolve(ns, contentId, mustBeContentId: true);
+
+            if (blobs == null)
+            {
+                return NotFound(new ProblemDetails { Title = $"Unable to resolve content id {contentId} ({ns})." });
+            }
+
+            return Ok(new ResolvedContentIdResponse(blobs));
+        }
+
+        /// <summary>
+        /// Update a content id mapping
+        /// </summary>
+        /// <param name="ns">Namespace. Each namespace is completely separated from each other. Use for different types of data that is never expected to be similar (between two different games for instance). Example: `uc4.ddc`</param>
+        /// <param name="contentId">The content id to resolve </param>
+        /// <param name="blobIdentifier">The blob identifier to map the content id to</param>
+        /// <param name="contentWeight">The weight of this mapping, higher means more important</param>
+        [HttpPut("{ns}/{contentId}/update/{blobIdentifier}/{contentWeight}", Order = 500)]
+        [ProducesDefaultResponseType]
+        [ProducesResponseType(type: typeof(ProblemDetails), 400)]
+        public async Task<IActionResult> UpdateContentIdMapping(NamespaceId ns, ContentId contentId, BlobIdentifier blobIdentifier, int contentWeight)
+        {
+            ActionResult? result = await _requestHelper.HasAccessToNamespace(User, Request, ns, new [] { JupiterAclAction.WriteObject });
+            if (result != null)
+            {
+                return result;
+            }
+            await _contentIdStore.Put(ns, contentId, blobIdentifier, contentWeight);
+
+            return Ok();
+        }
+    }
+
+    public class ResolvedContentIdResponse
+    {
+        public BlobIdentifier[] Blobs { get; set; }
+
+        public ResolvedContentIdResponse(BlobIdentifier[] blobs)
+        {
+            Blobs = blobs;
+        }
+    }
+}
diff --git a/Engine/Source/Programs/Horde/Horde.Server/Ddc/ContentIdStore.cs b/Engine/Source/Programs/Horde/Horde.Server/Ddc/ContentIdStore.cs
new file mode 100644
--- /dev/null
+++ b/Engine/Source/Programs/Horde/Horde.Server/Ddc/ContentIdStore.cs
@@ -0,0 +1,22 @@
+// Copyright Epic Games, Inc. All Rights Reserved.
+
+using EpicGames.Horde.Storage;
+using System;
+using System.Threading.Tasks;
+using Jupiter.Implementation;
+
+namespace Horde.Server.Ddc
+{
+	class ContentIdStore : IContentIdStore
+	{
+		public Task<BlobIdentifier[]?> Resolve(NamespaceId ns, ContentId contentId, bool mustBeContentId = false)
+		{
+			throw new NotImplementedException();
+		}
+
+		public Task Put(NamespaceId ns, ContentId contentId, BlobIdentifier blobIdentifier, int contentWeight)
+		{
+			throw new NotImplementedException();
+		}
+	}
+}
diff --git a/Engine/Source/Programs/Horde/Horde.Server/Ddc/FormatResolver.cs b/Engine/Source/Programs/Horde/Horde.Server/Ddc/FormatResolver.cs
new file mode 100644
--- /dev/null
+++ b/Engine/Source/Programs/Horde/Horde.Server/Ddc/FormatResolver.cs
@@ -0,0 +1,72 @@
+// Copyright Epic Games, Inc. All Rights Reserved.
+
+using System;
+using System.Linq;
+using System.Net.Mime;
+using EpicGames.AspNet;
+using Microsoft.AspNetCore.Http;
+using Microsoft.AspNetCore.Mvc;
+using Microsoft.Extensions.Options;
+using Microsoft.Extensions.Primitives;
+
+#pragma warning disable CS1591
+
+namespace Jupiter.Controllers
+{
+    public class FormatResolver
+    {
+        private readonly IOptionsMonitor<MvcOptions> _mvcOptions;
+
+        private readonly string[] _validContentTypes = {
+            MediaTypeNames.Application.Octet, 
+            MediaTypeNames.Application.Json, 
+            CustomMediaTypeNames.UnrealCompactBinary, 
+            CustomMediaTypeNames.JupiterInlinedPayload,
+            CustomMediaTypeNames.UnrealCompactBinaryPackage
+        };
+
+        public FormatResolver(IOptionsMonitor<MvcOptions> mvcOptions)
+        {
+            _mvcOptions = mvcOptions;
+        }
+
+        public string GetResponseType(HttpRequest request, string? format, string defaultContentType)
+        {
+            // if format specifier is used it takes precedence over the accept header
+            if (format != null)
+            {
+                string? typeMapping = _mvcOptions.CurrentValue.FormatterMappings.GetMediaTypeMappingForFormat(format);
+                if (typeMapping == null)
+                {
+                    throw new Exception($"No mapping defined from format {format} to mime type");
+                }
+
+                return typeMapping;
+            }
+            
+            StringValues acceptHeader = request.Headers["Accept"];
+
+            if (acceptHeader.Count == 0)
+            {
+                // no accept header specified, return default type
+                return defaultContentType;
+            }
+
+            // */* means to accept anything, so we use the default content type
+            if (acceptHeader == "*/*")
+            {
+                return defaultContentType;
+            }
+
+            foreach (string header in acceptHeader)
+            {
+                if (_validContentTypes.Contains(header, StringComparer.OrdinalIgnoreCase))
+                {
+                    return header;
+                }
+            }
+
+            throw new Exception($"Unable to determine response type for header: {acceptHeader}");
+        }
+    }
+}
diff --git a/Engine/Source/Programs/Horde/Horde.Server/Ddc/IBlobService.cs b/Engine/Source/Programs/Horde/Horde.Server/Ddc/IBlobService.cs
new file mode 100644
--- /dev/null
+++ b/Engine/Source/Programs/Horde/Horde.Server/Ddc/IBlobService.cs
@@ -0,0 +1,156 @@
+﻿// Copyright Epic Games, Inc. All Rights Reserved.
+
+using System;
+using System.Collections.Generic;
+using System.IO;
+using System.Net.Mime;
+using System.Threading.Tasks;
+using EpicGames.AspNet;
+using EpicGames.Horde.Storage;
+using Jupiter.Common.Implementation;
+using Jupiter.Utils;
+using Microsoft.Extensions.DependencyInjection;
+using OpenTelemetry.Trace;
+
+#pragma warning disable CS1591
+
+namespace Jupiter.Implementation;
+
+public interface IBlobService
+{
+    Task<JupiterContentHash> VerifyContentMatchesHash(Stream content, JupiterContentHash identifier);
+    Task<BlobIdentifier> PutObjectKnownHash(NamespaceId ns, IBufferedPayload content, BlobIdentifier identifier);
+    Task<BlobIdentifier> PutObject(NamespaceId ns, IBufferedPayload payload, BlobIdentifier identifier);
+    Task<BlobIdentifier> PutObject(NamespaceId ns, byte[] payload, BlobIdentifier identifier);
+    Task<Uri?> MaybePutObjectWithRedirect(NamespaceId ns, BlobIdentifier identifier);
+
+    Task<BlobContents> GetObject(NamespaceId ns, BlobIdentifier blob, List<string>? storageLayers = null, bool supportsRedirectUri = false);
+    
+    Task<Uri?> GetObjectWithRedirect(NamespaceId ns, BlobIdentifier blobIdentifier, List<string>? storageLayers = null);
+
+    Task<BlobContents> ReplicateObject(NamespaceId ns, BlobIdentifier blob, bool force = false);
+
+    Task<bool> Exists(NamespaceId ns, BlobIdentifier blob, List<string>? storageLayers = null);
+
+    /// <summary>
+    /// Checks that the blob exists in the root store, the store which is last in the list and thus is intended to have every blob in it
+    /// </summary>
+    /// <param name="ns">The namespace</param>
+    /// <param name="blob">The identifier of the blob</param>
+    /// <returns></returns>
+    Task<bool> ExistsInRootStore(NamespaceId ns, BlobIdentifier blob);
+
+    // Delete a object
+    Task DeleteObject(NamespaceId ns, BlobIdentifier blob);
+
+    // delete the whole namespace
+    Task DeleteNamespace(NamespaceId ns);
+
+    IAsyncEnumerable<(BlobIdentifier,DateTime)> ListObjects(NamespaceId ns);
+    Task<BlobIdentifier[]> FilterOutKnownBlobs(NamespaceId ns, IEnumerable<BlobIdentifier> blobs);
+    Task<BlobIdentifier[]> FilterOutKnownBlobs(NamespaceId ns, IAsyncEnumerable<BlobIdentifier> blobs);
+    Task<BlobContents> GetObjects(NamespaceId ns, BlobIdentifier[] refRequestBlobReferences);
+
+    bool ShouldFetchBlobOnDemand(NamespaceId ns);
+}
+
+public static class BlobServiceExtensions
+{
+    public static async Task<ContentId> PutCompressedObject(this IBlobService blobService, NamespaceId ns, IBufferedPayload payload, ContentId? id, IServiceProvider provider)
+    {
+        IContentIdStore contentIdStore = provider.GetService<IContentIdStore>()!;
+        CompressedBufferUtils compressedBufferUtils = provider.GetService<CompressedBufferUtils>()!;
+        Tracer tracer = provider.GetService<Tracer>()!;
+
+        // decompress the content and generate a identifier from it to verify the identifier we got
+        await using Stream decompressStream = payload.GetStream();
+        // TODO: we should add a overload for decompress content that can work on streams, otherwise we are still limited to 2GB compressed blobs
+        byte[] decompressedContent = compressedBufferUtils.DecompressContent(await decompressStream.ToByteArray());
+
+        await using MemoryStream decompressedStream = new MemoryStream(decompressedContent);
+        ContentId identifierDecompressedPayload;
+        if (id != null)
+        {
+            identifierDecompressedPayload = ContentId.FromContentHash(await blobService.VerifyContentMatchesHash(decompressedStream, id));
+        }
+        else
+        {
+            JupiterContentHash blobHash;
+            {
+                using TelemetrySpan _ = tracer.StartActiveSpan("web.hash").SetAttribute("operation.name", "web.hash");
+                blobHash = await BlobIdentifier.FromStream(decompressedStream);
+            }
+
+            identifierDecompressedPayload = ContentId.FromContentHash(blobHash);
+        }
+
+        BlobIdentifier identifierCompressedPayload;
+        {
+            using TelemetrySpan _ = tracer.StartActiveSpan("web.hash").SetAttribute("operation.name", "web.hash");
+            await using Stream hashStream = payload.GetStream();
+            identifierCompressedPayload = await BlobIdentifier.FromStream(hashStream);
+        }
+
+        // commit the mapping from the decompressed hash to the compressed hash, we run this in parallel with the blob store submit
+        // TODO: let users specify weight of the blob compared to previously submitted content ids
+        int contentIdWeight = (int)payload.Length;
+        Task contentIdStoreTask = contentIdStore.Put(ns, identifierDecompressedPayload, identifierCompressedPayload, contentIdWeight);
+
+        // we still commit the compressed buffer to the object store using the hash of the compressed content
+        {
+            await blobService.PutObjectKnownHash(ns, payload, identifierCompressedPayload);
+        }
+
+        await contentIdStoreTask;
+
+        return identifierDecompressedPayload;
+    }
+
+    public static async Task<(BlobContents, string)> GetCompressedObject(this IBlobService blobService, NamespaceId ns, ContentId contentId, IServiceProvider provider, bool supportsRedirectUri = false)
+    {
+        IContentIdStore contentIdStore = provider.GetService<IContentIdStore>()!;
+        Tracer tracer = provider.GetService<Tracer>()!;
+
+        BlobIdentifier[]? chunks = await contentIdStore.Resolve(ns, contentId, mustBeContentId: false);
+        if (chunks == null || chunks.Length == 0)
+        {
+            throw new ContentIdResolveException(contentId);
+        }
+
+        // single chunk, we just return that chunk
+        if (chunks.Length == 1)
+        {
+            BlobIdentifier blobToReturn = chunks[0];
+            string mimeType = CustomMediaTypeNames.UnrealCompressedBuffer;
+            if (contentId.Equals(blobToReturn))
+            {
+                // this was actually the unmapped blob, meaning its not a compressed buffer
+                mimeType = MediaTypeNames.Application.Octet;
+            }
+
+            return (await blobService.GetObject(ns, blobToReturn, supportsRedirectUri: supportsRedirectUri), mimeType);
+        }
+
+        // chunked content, combine the chunks into a single stream
+        using TelemetrySpan _ = tracer.StartActiveSpan("blob.combine").SetAttribute("operation.name", "blob.combine");
+        Task<BlobContents>[] tasks = new Task<BlobContents>[chunks.Length];
+        for (int i = 0; i < chunks.Length; i++)
+        {
+            // even if it was requested to support redirect, since we need to combine the chunks using redirects is not possible
+            tasks[i] = blobService.GetObject(ns, chunks[i], supportsRedirectUri: false);
+        }
+
+        MemoryStream ms = new MemoryStream();
+        foreach (Task<BlobContents> task in tasks)
+        {
+            BlobContents blob = await task;
+            await using Stream s = blob.Stream;
+            await s.CopyToAsync(ms);
+        }
+
+        ms.Seek(0, SeekOrigin.Begin);
+
+        // chunking could not have happened for a non compressed buffer so assume it is compressed
+        return (new BlobContents(ms, ms.Length), CustomMediaTypeNames.UnrealCompressedBuffer);
+    }
+}
diff --git a/Engine/Source/Programs/Horde/Horde.Server/Ddc/IBlobStore.cs b/Engine/Source/Programs/Horde/Horde.Server/Ddc/IBlobStore.cs
new file mode 100644
--- /dev/null
+++ b/Engine/Source/Programs/Horde/Horde.Server/Ddc/IBlobStore.cs
@@ -0,0 +1,90 @@
+﻿// Copyright Epic Games, Inc. All Rights Reserved.
+
+using System;
+using System.Collections.Generic;
+using System.IO;
+using System.Threading;
+using System.Threading.Tasks;
+using EpicGames.Horde.Storage;
+
+#pragma warning disable CS1591
+
+namespace Jupiter.Implementation
+{
+    public enum LastAccessTrackingFlags
+    {
+        DoTracking = 0,
+        SkipTracking = 1,
+    }
+
+    public interface IBlobStore
+    {
+        Task<BlobIdentifier> PutObject(NamespaceId ns, byte[] blob, BlobIdentifier identifier);
+        Task<BlobIdentifier> PutObject(NamespaceId ns, ReadOnlyMemory<byte> blob, BlobIdentifier identifier);
+        Task<BlobIdentifier> PutObject(NamespaceId ns, Stream content, BlobIdentifier identifier);
+
+        Task<BlobContents> GetObject(NamespaceId ns, BlobIdentifier blob, LastAccessTrackingFlags flags = LastAccessTrackingFlags.DoTracking, bool supportsRedirectUri = false);
+        Task<bool> Exists(NamespaceId ns, BlobIdentifier blob, bool forceCheck = false);
+
+        // Delete a object
+        Task DeleteObject(NamespaceId ns, BlobIdentifier blob);
+
+        // delete the whole namespace
+        Task DeleteNamespace(NamespaceId ns);
+
+        IAsyncEnumerable<(BlobIdentifier,DateTime)> ListObjects(NamespaceId ns);
+        Task<Uri?> PutObjectWithRedirect(NamespaceId ns, BlobIdentifier identifier);
+        Task<Uri?> GetObjectByRedirect(NamespaceId ns, BlobIdentifier blob);
+    }
+
+    public interface IStorageBackend
+    {
+        Task WriteAsync(string path, Stream content, CancellationToken cancellationToken = default);
+        Task<BlobContents?> TryReadAsync(string path, LastAccessTrackingFlags flags = LastAccessTrackingFlags.DoTracking, CancellationToken cancellationToken = default);
+        Task<bool> ExistsAsync(string path, CancellationToken cancellationToken = default);
+        Task DeleteAsync(string path, CancellationToken cancellationToken = default);
+        IAsyncEnumerable<(string, DateTime)> ListAsync(CancellationToken cancellationToken = default);
+    }
+
+    public class BlobNotFoundException : Exception
+    {
+        public NamespaceId Ns { get; }
+        public BlobIdentifier Blob { get; }
+
+        public BlobNotFoundException(NamespaceId ns, BlobIdentifier blob) : base($"No Blob in Namespace {ns} with id {blob}")
+        {
+            Ns = ns;
+            Blob = blob;
+        }
+
+        public BlobNotFoundException(NamespaceId ns, BlobIdentifier blob, string message) : base(message)
+        {
+            Ns = ns;
+            Blob = blob;
+        }
+    }
+
+    public class BlobReplicationException : BlobNotFoundException
+    {
+        public BlobReplicationException(NamespaceId ns, BlobIdentifier blob, string message) : base(ns, blob, message)
+        {
+        }
+    }
+
+    public class BlobToLargeException : Exception
+    {
+        public BlobIdentifier Blob { get; }
+
+        public BlobToLargeException(BlobIdentifier blob) : base($"Blob {blob} was to large to cache")
+        {
+            Blob = blob;
+        }
+    }
+
+    public class ResourceHasToManyRequestsException : Exception
+    {
+        public ResourceHasToManyRequestsException(Exception originalException) : base($"To many requests to resource", originalException)
+        {
+        }
+    }
+}
diff --git a/Engine/Source/Programs/Horde/Horde.Server/Ddc/IContentIdStore.cs b/Engine/Source/Programs/Horde/Horde.Server/Ddc/IContentIdStore.cs
new file mode 100644
--- /dev/null
+++ b/Engine/Source/Programs/Horde/Horde.Server/Ddc/IContentIdStore.cs
@@ -0,0 +1,50 @@
+// Copyright Epic Games, Inc. All Rights Reserved.
+
+using EpicGames.Horde.Storage;
+using System;
+using System.Threading.Tasks;
+
+#pragma warning disable CS1591
+
+namespace Jupiter.Implementation
+{
+    public interface IContentIdStore
+    {
+        /// <summary>
+        /// Resolve a content id from its hash into the actual blob (that can in turn be chunked into a set of blobs)
+        /// </summary>
+        /// <param name="ns">The namespace to operate in</param>
+        /// <param name="contentId">The identifier for the content id</param>
+        /// <param name="mustBeContentId"></param>
+        /// <returns></returns>
+        Task<BlobIdentifier[]?> Resolve(NamespaceId ns, ContentId contentId, bool mustBeContentId = false);
+
+        /// <summary>
+        /// Add a mapping from contentId to blobIdentifier
+        /// </summary>
+        /// <param name="ns">The namespace to operate in</param>
+        /// <param name="contentId">The contentId</param>
+        /// <param name="blobIdentifier">The blob the content id maps to</param>
+        /// <param name="contentWeight">Weight of this identifier compared to previous mappings, used to determine which is more important, lower weight is considered a better fit</param>
+        /// <returns></returns>
+        Task Put(NamespaceId ns, ContentId contentId, BlobIdentifier blobIdentifier, int contentWeight);
+    }
+
+    public class InvalidContentIdException : Exception
+    {
+        public InvalidContentIdException(ContentId contentId) : base($"Unknown content id {contentId}")
+        {
+
+        }
+    }
+
+    public class ContentIdResolveException : Exception
+    {
+        public ContentId ContentId { get; }
+
+        public ContentIdResolveException(ContentId contentId) : base($"Unable to find any mapping of contentId {contentId} that has all blobs present")
+        {
+            ContentId = contentId;
+        }
+    }
+}
diff --git a/Engine/Source/Programs/Horde/Horde.Server/Ddc/IObjectService.cs b/Engine/Source/Programs/Horde/Horde.Server/Ddc/IObjectService.cs
new file mode 100644
--- /dev/null
+++ b/Engine/Source/Programs/Horde/Horde.Server/Ddc/IObjectService.cs
@@ -0,0 +1,35 @@
+// Copyright Epic Games, Inc. All Rights Reserved.
+
+using System;
+using System.Collections.Generic;
+using System.Threading.Tasks;
+using EpicGames.Horde.Storage;
+using EpicGames.Serialization;
+
+#pragma warning disable CS1591
+
+namespace Jupiter.Implementation
+{
+    public interface IObjectService
+    {
+        Task<(ObjectRecord, BlobContents?)> Get(NamespaceId ns, BucketId bucket, IoHashKey key, string[] fields, bool doLastAccessTracking = true);
+        Task<(ContentId[], BlobIdentifier[])> Put(NamespaceId ns, BucketId bucket, IoHashKey key, BlobIdentifier blobHash, CbObject payload);
+        Task<(ContentId[], BlobIdentifier[])> Finalize(NamespaceId ns, BucketId bucket, IoHashKey key, BlobIdentifier blobHash);
+
+        IAsyncEnumerable<NamespaceId> GetNamespaces();
+
+        Task<bool> Delete(NamespaceId ns, BucketId bucket, IoHashKey key);
+        Task<long> DropNamespace(NamespaceId ns);
+        Task<long> DeleteBucket(NamespaceId ns, BucketId bucket);
+
+        Task<bool> Exists(NamespaceId ns, BucketId bucket, IoHashKey key);
+        Task<List<BlobIdentifier>> GetReferencedBlobs(NamespaceId ns, BucketId bucket, IoHashKey key);
+    }
+
+    public class ObjectHashMismatchException : Exception
+    {
+        public ObjectHashMismatchException(NamespaceId ns, BucketId bucket, IoHashKey name, BlobIdentifier suppliedHash, BlobIdentifier actualHash) : base($"Object {name} in bucket {bucket} and namespace {ns} did not reference hash {suppliedHash} was referencing {actualHash}")
+        {
+        }
+    }
+}
diff --git a/Engine/Source/Programs/Horde/Horde.Server/Ddc/IReferenceResolver.cs b/Engine/Source/Programs/Horde/Horde.Server/Ddc/IReferenceResolver.cs
new file mode 100644
--- /dev/null
+++ b/Engine/Source/Programs/Horde/Horde.Server/Ddc/IReferenceResolver.cs
@@ -0,0 +1,370 @@
+// Copyright Epic Games, Inc. All Rights Reserved.
+
+using System;
+using System.Collections.Generic;
+using System.Linq;
+using System.Threading.Tasks;
+using EpicGames.Core;
+using EpicGames.Horde.Storage;
+using EpicGames.Serialization;
+using Jupiter.Utils;
+using OpenTelemetry.Trace;
+
+#pragma warning disable CS1591
+
+namespace Jupiter.Implementation
+{
+    public abstract class Attachment
+    {
+        public abstract IoHash AsIoHash();
+    }
+
+    public class BlobAttachment : Attachment
+    {
+        public BlobIdentifier Identifier { get; }
+
+        public BlobAttachment(BlobIdentifier blobIdentifier)
+        {
+            Identifier = blobIdentifier;
+        }
+
+        public override IoHash AsIoHash()
+        {
+            return Identifier.AsIoHash();
+        }
+    }
+
+    public class ObjectAttachment : Attachment
+    {
+        public BlobIdentifier Identifier { get; }
+
+        public ObjectAttachment(BlobIdentifier blobIdentifier)
+        {
+            Identifier = blobIdentifier;
+        }
+
+        public override IoHash AsIoHash()
+        {
+            return Identifier.AsIoHash();
+        }
+    }
+
+    public class ContentIdAttachment : Attachment
+    {
+        public ContentId Identifier { get; }
+        public BlobIdentifier[] ReferencedBlobs { get; }
+
+        public ContentIdAttachment(ContentId contentId, BlobIdentifier[] referencedBlobs)
+        {
+            Identifier = contentId;
+            ReferencedBlobs = referencedBlobs;
+        }
+
+        public override IoHash AsIoHash()
+        {
+            return Identifier.AsIoHash();
+        }
+    }
+
+    public interface IReferenceResolver
+    {
+        /// <summary>
+        /// Returns the blobs referenced from the cb object and any children
+        /// These blobs are guaranteed to exist, if any blob is missing
+        /// PartialReferenceResolveException or ReferenceIsMissingBlobsException are raised
+        /// </summary>
+        /// <param name="ns">The namespace to check</param>
+        /// <param name="cb">The compact binary object to resolve references for</param>
+        /// <returns></returns>
+        IAsyncEnumerable<BlobIdentifier> GetReferencedBlobs(NamespaceId ns, CbObject cb);
+
+        /// <summary>
+        /// Returns which attachments exist in the cb object or any children
+        /// These attachments are not guaranteed to reference blobs that actually exist
+        /// </summary>
+        /// <param name="ns">The namespace to check</param>
+        /// <param name="cb">The compact binary object to resolve references for</param>
+        /// <returns></returns>
+        IAsyncEnumerable<Attachment> GetAttachments(NamespaceId ns, CbObject cb);
+    }
+
+    public class ReferenceResolver : IReferenceResolver
+    {
+        private readonly IBlobService _blobStore;
+        private readonly IContentIdStore _contentIdStore;
+        private readonly Tracer _tracer;
+
+        public ReferenceResolver(IBlobService blobStore, IContentIdStore contentIdStore, Tracer tracer)
+        {
+            _blobStore = blobStore;
+            _contentIdStore = contentIdStore;
+            _tracer = tracer;
+        }
+
+        public async IAsyncEnumerable<Attachment> GetAttachments(NamespaceId ns, CbObject cb)
+        {
+            Queue<CbObject> objectsToVisit = new Queue<CbObject>();
+            objectsToVisit.Enqueue(cb);
+            List<BlobIdentifier> unresolvedBlobReferences = new List<BlobIdentifier>();
+
+            List<Task<CbObject>> pendingCompactBinaryAttachments = new();
+            List<Task< (ContentId, BlobIdentifier[]?)>> pendingContentIdResolves = new();
+
+            while (pendingCompactBinaryAttachments.Count != 0 || pendingContentIdResolves.Count != 0 || objectsToVisit.Count != 0)
+            {
+                List<Attachment> attachments = new List<Attachment>();
+
+                if (objectsToVisit.TryDequeue(out CbObject? parent))
+                {
+                    parent.IterateAttachments(field =>
+                    {
+                        IoHash attachmentHash = field.AsAttachment();
+
+                        BlobIdentifier blobIdentifier = BlobIdentifier.FromIoHash(attachmentHash);
+                        ContentId contentId = ContentId.FromIoHash(attachmentHash);
+
+                        if (field.IsBinaryAttachment())
+                        {
+                            Task<(ContentId, BlobIdentifier[]?)> resolveContentId = ResolveContentId(ns, contentId);
+                            pendingContentIdResolves.Add(resolveContentId);
+                        }
+                        else if (field.IsObjectAttachment())
+                        {
+                            attachments.Add(new ObjectAttachment(blobIdentifier));
+                            pendingCompactBinaryAttachments.Add(ParseCompactBinaryAttachment(ns, blobIdentifier));
+                        }
+                        else
+                        {
+                            throw new NotImplementedException($"Unknown attachment type for field {field}");
+                        }
+                    });
+                }
+
+                // check for any content id resolves to finish
+                List<Task<(ContentId, BlobIdentifier[]?)>> finishedContentIdResolves = new();
+                foreach (Task<(ContentId, BlobIdentifier[]?)> pendingContentIdResolve in pendingContentIdResolves)
+                {
+                    if (pendingContentIdResolve.IsCompleted)
+                    {
+                        ContentId? contentId = null;
+                        BlobIdentifier[]? resolvedBlobs = null;
+                        BlobIdentifier? blobIdentifier = null;
+                        bool wasContentId = false;
+                        try
+                        {
+                            (contentId, resolvedBlobs) = await pendingContentIdResolve;
+                            blobIdentifier = contentId.AsBlobIdentifier();
+                            wasContentId = !(resolvedBlobs is { Length: 1 } && resolvedBlobs[0].Equals(blobIdentifier));
+                        }
+                        catch (InvalidContentIdException)
+                        {
+                            resolvedBlobs = null;
+                        }
+                        catch (BlobNotFoundException e)
+                        {
+                            unresolvedBlobReferences.Add(e.Blob);
+                        }
+
+                        if (wasContentId && resolvedBlobs != null)
+                        {
+                            attachments.Add(new ContentIdAttachment(contentId!, resolvedBlobs));
+                        }
+                        else
+                        {
+                            attachments.Add(new BlobAttachment(blobIdentifier!));
+                        }
+
+                        finishedContentIdResolves.Add(pendingContentIdResolve);
+                    }
+                }
+
+                // cleanup finished tasks
+                foreach (Task<(ContentId, BlobIdentifier[]?)> finishedTask in finishedContentIdResolves)
+                {
+                    pendingContentIdResolves.Remove(finishedTask);
+                }
+
+                // check for any compact binary attachment fetches and add those to the objects we are handling
+                List<Task<CbObject>> finishedCompactBinaryResolves = new();
+                foreach (Task<CbObject> pendingCompactBinaryAttachment in pendingCompactBinaryAttachments)
+                {
+                    if (pendingCompactBinaryAttachment.IsCompleted)
+                    {
+                        try
+                        {
+                            CbObject childBinaryObject = await pendingCompactBinaryAttachment;
+                            objectsToVisit.Enqueue(childBinaryObject);
+                        }
+                        catch (BlobNotFoundException e)
+                        {
+                            unresolvedBlobReferences.Add(e.Blob);
+                        }
+                        finishedCompactBinaryResolves.Add(pendingCompactBinaryAttachment);
+                    }
+                }
+
+                // cleanup finished tasks
+                foreach (Task<CbObject> finishedTask in finishedCompactBinaryResolves)
+                {
+                    pendingCompactBinaryAttachments.Remove(finishedTask);
+                }
+
+                bool hasWaited = false;
+                // if there are pending resolves left, wait for one of them to finish to avoid busy waiting
+                if (pendingCompactBinaryAttachments.Any())
+                {
+                    hasWaited = true;
+                    await Task.WhenAny(pendingCompactBinaryAttachments);
+                }
+
+                // if we did not waiting for compact binary attachment resolves we consider waiting for content id resolves if there are any
+                if (!hasWaited && pendingContentIdResolves.Any())
+                {
+                    await Task.WhenAny(pendingContentIdResolves);
+                }
+
+                foreach (Attachment attachment in attachments)
+                {
+                    yield return attachment;
+                }
+            }
+
+            if (unresolvedBlobReferences.Count != 0)
+            {
+                throw new ReferenceIsMissingBlobsException(unresolvedBlobReferences);
+            }
+        }
+
+        public async IAsyncEnumerable<BlobIdentifier> GetReferencedBlobs(NamespaceId ns, CbObject cb)
+        {
+            List<Task<(BlobIdentifier, bool)>> pendingBlobExistsChecks = new();
+            List<Task<(ContentIdAttachment, bool)>> pendingContentIdChecks = new();
+            List<ContentId> unresolvedContentIdReferences = new List<ContentId>();
+            List<BlobIdentifier> unresolvedBlobReferences = new List<BlobIdentifier>();
+
+            // Resolve all the attachments
+            await foreach (Attachment attachment in GetAttachments(ns, cb))
+            {
+                if (attachment is BlobAttachment blobAttachment)
+                {
+                    pendingBlobExistsChecks.Add(CheckBlobExists(ns, blobAttachment.Identifier));
+                }
+                else if (attachment is ContentIdAttachment contentIdAttachment)
+                {
+                    // If we find a content id we resolve that into the actual blobs it references
+                    pendingContentIdChecks.Add(CheckContentIdExists(ns, contentIdAttachment));
+                }
+                else if (attachment is ObjectAttachment objectAttachment)
+                {
+                    // a object just references the same blob, traversing the object attachment is done in GetAttachments
+                    pendingBlobExistsChecks.Add(CheckBlobExists(ns, objectAttachment.Identifier));
+                }
+                else
+                {
+                    throw new NotSupportedException($"Unknown attachment type {attachment.GetType()}");
+                }
+            }
+
+            // return any verified blobs
+            foreach (Task<(BlobIdentifier, bool)> pendingBlobExistsTask in pendingBlobExistsChecks)
+            {
+                (BlobIdentifier blob, bool exists) = await pendingBlobExistsTask;
+                if (exists)
+                {
+                    yield return blob;
+                }
+                else
+                {
+                    unresolvedBlobReferences.Add(blob);
+                }
+            }
+
+            foreach (Task<(ContentIdAttachment, bool)> pendingContentIdTask in pendingContentIdChecks)
+            {
+                (ContentIdAttachment contentIdAttachment, bool exists) = await pendingContentIdTask;
+
+                if (!exists)
+                {
+                    unresolvedContentIdReferences.Add(contentIdAttachment.Identifier);
+                    continue;
+                }
+
+                foreach (BlobIdentifier b in contentIdAttachment.ReferencedBlobs)
+                {
+                    yield return b;
+                }
+            }
+
+            // if there were any content ids we did not recognize we throw a partial reference exception
+            if (unresolvedContentIdReferences.Count != 0)
+            {
+                throw new PartialReferenceResolveException(unresolvedContentIdReferences);
+            }
+            // if there were any blobs missing we throw a partial reference exception
+            if (unresolvedBlobReferences.Count != 0)
+            {
+                throw new ReferenceIsMissingBlobsException(unresolvedBlobReferences);
+            }
+        }
+
+        private async Task<(ContentIdAttachment, bool)> CheckContentIdExists(NamespaceId ns, ContentIdAttachment contentIdAttachment)
+        {
+            bool allBlobsExist = true;
+            foreach (BlobIdentifier b in contentIdAttachment.ReferencedBlobs)
+            {
+                (BlobIdentifier _, bool exists) = await CheckBlobExists(ns, b);
+
+                if (!exists)
+                {
+                    allBlobsExist = false;
+                    break;
+                }
+            }
+
+            return (contentIdAttachment, allBlobsExist);
+        }
+
+        private async Task<CbObject> ParseCompactBinaryAttachment(NamespaceId ns, BlobIdentifier blobIdentifier)
+        {
+            BlobContents contents = await _blobStore.GetObject(ns, blobIdentifier);
+            byte[] data = await contents.Stream.ToByteArray();
+            CbObject childBinaryObject = new CbObject(data);
+
+            return childBinaryObject;
+        }
+
+        private async Task<(ContentId, BlobIdentifier[]?)> ResolveContentId(NamespaceId ns, ContentId contentId)
+        {
+            using TelemetrySpan scope = _tracer.StartActiveSpan("ReferenceResolver.ResolveContentId")
+                .SetAttribute("operation.name", "ReferenceResolver.ResolveContentId")
+                .SetAttribute("resource.name", contentId.ToString());
+            BlobIdentifier[]? resolvedBlobs = await _contentIdStore.Resolve(ns, contentId);
+            return (contentId, resolvedBlobs);
+        }
+
+        private async Task<(BlobIdentifier, bool)> CheckBlobExists(NamespaceId ns, BlobIdentifier blob)
+        {
+            return (blob, await _blobStore.Exists(ns, blob));
+        }
+    }
+
+    public class PartialReferenceResolveException : Exception
+    {
+        public List<ContentId> UnresolvedReferences { get; }
+
+        public PartialReferenceResolveException(List<ContentId> unresolvedReferences) : base($"References missing: {string.Join(',', unresolvedReferences)}")
+        {
+            UnresolvedReferences = unresolvedReferences;
+        }
+    }
+
+    public class ReferenceIsMissingBlobsException : Exception
+    {
+        public List<BlobIdentifier> MissingBlobs { get; }
+
+        public ReferenceIsMissingBlobsException(List<BlobIdentifier> missingBlobs) : base($"References is missing these blobs: {string.Join(',', missingBlobs)}")
+        {
+            MissingBlobs = missingBlobs;
+        }
+    }
+}
diff --git a/Engine/Source/Programs/Horde/Horde.Server/Ddc/IReferencesStore.cs b/Engine/Source/Programs/Horde/Horde.Server/Ddc/IReferencesStore.cs
new file mode 100644
--- /dev/null
+++ b/Engine/Source/Programs/Horde/Horde.Server/Ddc/IReferencesStore.cs
@@ -0,0 +1,71 @@
+// Copyright Epic Games, Inc. All Rights Reserved.
+
+using System;
+using System.Collections.Generic;
+using System.Threading.Tasks;
+using EpicGames.Horde.Storage;
+
+#pragma warning disable CS1591
+
+namespace Jupiter.Implementation
+{
+    public interface IReferencesStore
+    {
+        Task<ObjectRecord> Get(NamespaceId ns, BucketId bucket, IoHashKey key, FieldFlags flags);
+
+        [Flags]
+        public enum FieldFlags
+        {
+            None = 0,
+            IncludePayload = 1,
+            All = IncludePayload
+        }
+
+        Task Put(NamespaceId ns, BucketId bucket, IoHashKey key, BlobIdentifier blobHash, byte[] blob, bool isFinalized);
+        Task Finalize(NamespaceId ns, BucketId bucket, IoHashKey key, BlobIdentifier blobIdentifier);
+
+        Task UpdateLastAccessTime(NamespaceId ns, BucketId bucket, IoHashKey key, DateTime newLastAccessTime);
+        IAsyncEnumerable<(NamespaceId, BucketId, IoHashKey, DateTime)> GetRecords();
+
+        IAsyncEnumerable<NamespaceId> GetNamespaces();
+        Task<bool> Delete(NamespaceId ns, BucketId bucket, IoHashKey key);
+        Task<long> DropNamespace(NamespaceId ns);
+        Task<long> DeleteBucket(NamespaceId ns, BucketId bucket);
+    }
+
+    public class ObjectRecord
+    {
+        public ObjectRecord(NamespaceId ns, BucketId bucket, IoHashKey name, DateTime lastAccess, byte[]? inlinePayload, BlobIdentifier blobIdentifier, bool isFinalized)
+        {
+            Namespace = ns;
+            Bucket = bucket;
+            Name = name;
+            LastAccess = lastAccess;
+            InlinePayload = inlinePayload;
+            BlobIdentifier = blobIdentifier;
+            IsFinalized = isFinalized;
+        }
+
+        public NamespaceId Namespace { get; }
+        public BucketId Bucket { get; }
+        public IoHashKey Name { get; }
+        public DateTime LastAccess { get; }
+        public byte[]? InlinePayload { get; set; }
+        public BlobIdentifier BlobIdentifier { get; set; }
+        public bool IsFinalized {get;}
+    }
+
+    public class ObjectNotFoundException : Exception
+    {
+        public ObjectNotFoundException(NamespaceId ns, BucketId bucket, IoHashKey key) : base($"Object not found {key} in bucket {bucket} namespace {ns}")
+        {
+            Namespace = ns;
+            Bucket = bucket;
+            Key = key;
+        }
+
+        public NamespaceId Namespace { get; }
+        public BucketId Bucket { get; }
+        public IoHashKey Key { get; }
+    }
+}
diff --git a/Engine/Source/Programs/Horde/Horde.Server/Ddc/IRequestHelper.cs b/Engine/Source/Programs/Horde/Horde.Server/Ddc/IRequestHelper.cs
new file mode 100644
--- /dev/null
+++ b/Engine/Source/Programs/Horde/Horde.Server/Ddc/IRequestHelper.cs
@@ -0,0 +1,29 @@
+// Copyright Epic Games, Inc. All Rights Reserved.
+
+using System;
+using System.Security.Claims;
+using System.Threading.Tasks;
+using EpicGames.Horde.Storage;
+using Microsoft.AspNetCore.Http;
+using Microsoft.AspNetCore.Mvc;
+
+#pragma warning disable CS1591
+
+namespace Jupiter.Controllers
+{
+    public interface IRequestHelper
+    {
+        public Task<ActionResult?> HasAccessToNamespace(ClaimsPrincipal user, HttpRequest request, NamespaceId ns, JupiterAclAction[] aclActions);
+        public Task<ActionResult?> HasAccessForGlobalOperations(ClaimsPrincipal user, JupiterAclAction[] aclActions);
+    }
+
+    public class AuthorizationException : Exception
+    {
+        public ActionResult Result { get; }
+
+        public AuthorizationException(ActionResult result, string errorMessage) : base(errorMessage)
+        {
+            Result = result;
+        }
+    }
+}
diff --git a/Engine/Source/Programs/Horde/Horde.Server/Ddc/IoHashKey.cs b/Engine/Source/Programs/Horde/Horde.Server/Ddc/IoHashKey.cs
new file mode 100644
--- /dev/null
+++ b/Engine/Source/Programs/Horde/Horde.Server/Ddc/IoHashKey.cs
@@ -0,0 +1,145 @@
+// Copyright Epic Games, Inc. All Rights Reserved.
+
+using System;
+using System.ComponentModel;
+using System.Globalization;
+using System.IO;
+using System.Text;
+using System.Text.Json;
+using System.Text.Json.Serialization;
+using EpicGames.Core;
+using EpicGames.Serialization;
+
+#pragma warning disable CS1591
+
+namespace Jupiter.Implementation
+{
+    [TypeConverter(typeof(IoHashKeyTypeConverter))]
+    [JsonConverter(typeof(IoHashKeyJsonConverter))]
+    [CbConverter(typeof(IoHashKeyCbConverter))]
+
+    public readonly struct IoHashKey: IEquatable<IoHashKey>
+    {
+        public IoHashKey(string key)
+        {
+             _text = key.ToLower();
+ 
+            if (_text.Length != 40)
+            {
+                throw new ArgumentException("IoHashKeys must be exactly 40 bytes.");
+            }
+ 
+            for (int Idx = 0; Idx < _text.Length; Idx++)
+            {
+                if (!IsValidCharacter(_text[Idx]))
+                {
+                    throw new ArgumentException($"{_text} is not a valid namespace id");
+                }
+            }
+        }
+
+        static bool IsValidCharacter(char c)
+        {
+            if (c >= 'a' && c <= 'z')
+            {
+                return true;
+            }
+            if (c >= '0' && c <= '9')
+            {
+                return true;
+            }
+            return false;
+        }
+        readonly string _text;
+
+        public bool Equals(IoHashKey other)
+        {
+            return string.Equals(_text , other._text, StringComparison.Ordinal);
+        }
+
+        public override bool Equals(object? obj)
+        {
+            return obj is IoHashKey other && Equals(other);
+        }
+
+        public override int GetHashCode()
+        {
+            return _text.GetHashCode(StringComparison.Ordinal);
+        }
+
+        public override string ToString()
+        {
+            return _text;
+        }
+
+        public static bool operator ==(IoHashKey left, IoHashKey right)
+        {
+            return left.Equals(right);
+        }
+ 
+        public static bool operator !=(IoHashKey left, IoHashKey right)
+        {
+            return !left.Equals(right);
+        }
+
+        public static IoHashKey FromName(string s)
+        {
+            return new IoHashKey(JupiterContentHash.FromBlob(Encoding.UTF8.GetBytes(s)).ToString());
+        }
+    }
+
+    public sealed class IoHashKeyTypeConverter : TypeConverter
+    {
+        /// <inheritdoc/>
+        public override bool CanConvertFrom(ITypeDescriptorContext? context, Type sourceType)
+        {
+            if (sourceType == typeof(string))
+            {
+                return true;
+            }
+
+            return base.CanConvertFrom(context, sourceType);
+        }
+ 
+        public override object? ConvertFrom(ITypeDescriptorContext? context, CultureInfo? culture, object value)
+        {
+            if (value is string s)
+            {
+                return new IoHashKey(s);
+            }
+
+            return base.ConvertFrom(context, culture, value);  
+        }
+    }
+
+    public class IoHashKeyJsonConverter : JsonConverter<IoHashKey>
+    {
+        public override IoHashKey Read(ref Utf8JsonReader reader, Type typeToConvert, JsonSerializerOptions options)
+        {
+            string? str = reader.GetString();
+            if (str == null)
+            {
+                throw new InvalidDataException("Unable to parse io hash key");
+            }
+
+            return new IoHashKey(str);
+        }
+
+        public override void Write(Utf8JsonWriter writer, IoHashKey value, JsonSerializerOptions options)
+        {
+            writer.WriteStringValue(value.ToString());
+        }
+    }
+
+    sealed class IoHashKeyCbConverter : CbConverterBase<IoHashKey>
+    {
+        public override IoHashKey Read(CbField field) => new IoHashKey(field.AsString());
+
+        /// <inheritdoc/>
+        public override void Write(CbWriter writer, IoHashKey value) => writer.WriteStringValue(value.ToString());
+
+        /// <inheritdoc/>
+        public override void WriteNamed(CbWriter writer, Utf8String name, IoHashKey value) => writer.WriteString(name, value.ToString());
+    }
+}
+
diff --git a/Engine/Source/Programs/Horde/Horde.Server/Ddc/JupiterContentHash.cs b/Engine/Source/Programs/Horde/Horde.Server/Ddc/JupiterContentHash.cs
new file mode 100644
--- /dev/null
+++ b/Engine/Source/Programs/Horde/Horde.Server/Ddc/JupiterContentHash.cs
@@ -0,0 +1,185 @@
+// Copyright Epic Games, Inc. All Rights Reserved.
+
+using System;
+using System.ComponentModel;
+using System.Globalization;
+using System.IO;
+using System.Text.Json;
+using System.Text.Json.Serialization;
+using Blake3;
+using EpicGames.Core;
+using EpicGames.Serialization;
+
+#pragma warning disable CS1591
+
+namespace Jupiter.Implementation
+{
+    [TypeConverter(typeof(ContentHashTypeConverter))]
+    [JsonConverter(typeof(ContentHashJsonConverter))]
+    [CbConverter(typeof(ContentHashCbConverter))]
+    public class JupiterContentHash : IEquatable<JupiterContentHash>, IEquatable<byte[]>
+    {
+        protected ByteArrayComparer Comparer { get; } = new ByteArrayComparer();
+        protected byte[] Identifier { get; init; }
+        public const int HashLength = 20;
+
+        public byte[] HashData => Identifier;
+        public JupiterContentHash(byte[] identifier)
+        {
+            Identifier = identifier;
+            /*if (identifier.Length != HashLength)
+            {
+                throw new ArgumentException("Supplied identifier was not 20 bytes, this is not a valid identifier", nameof(identifier));
+            }*/
+        }
+
+        [JsonConstructor]
+        public JupiterContentHash(string identifier)
+        {
+            if (identifier == null)
+            {
+                throw new ArgumentNullException(nameof(identifier));
+            }
+
+            byte[] byteIdentifier = StringUtils.ToHashFromHexString(identifier);
+            Identifier = byteIdentifier;
+            /*if (byteIdentifier.Length != HashLength)
+            {
+                throw new ArgumentException("Supplied identifier was not 20 bytes, this is not a valid identifier", nameof(identifier));
+            }*/
+        }
+
+        public override int GetHashCode()
+        {
+            return Comparer.GetHashCode(Identifier);
+        }
+
+        public bool Equals(JupiterContentHash? other)
+        {
+            if (other == null)
+            {
+                return false;
+            }
+
+            return Comparer.Equals(Identifier, other.Identifier);
+        }
+
+        public bool Equals(byte[]? other)
+        {
+            if (other == null)
+            {
+                return false;
+            }
+
+            return Comparer.Equals(Identifier, other);
+        }
+
+        public override bool Equals(object? obj)
+        {
+            if (ReferenceEquals(null, obj))
+            {
+                return false;
+            }
+
+            if (ReferenceEquals(this, obj))
+            {
+                return true;
+            }
+
+            if (obj.GetType() != GetType())
+            {
+                return false;
+            }
+
+            return Equals((JupiterContentHash) obj);
+        }
+
+        public override string ToString()
+        {
+            return StringUtils.FormatAsHexString(Identifier);
+        }
+
+        public static JupiterContentHash FromBlob(byte[] blobMemory)
+        {
+            using Hasher hasher = Hasher.New();
+            hasher.UpdateWithJoin(blobMemory);
+            Hash blake3Hash = hasher.Finalize();
+
+            // we only keep the first 20 bytes of the Blake3 hash
+            Span<byte> hash = blake3Hash.AsSpanUnsafe().Slice(0, HashLength);
+
+            return new JupiterContentHash(hash.ToArray());
+        }
+    }
+    
+    public class ContentHashCbConverter : CbConverterBase<JupiterContentHash>
+    {
+        public override JupiterContentHash Read(CbField field) => new JupiterContentHash(field.AsHash().ToByteArray());
+
+        /// <inheritdoc/>
+        public override void Write(CbWriter writer, JupiterContentHash value) => writer.WriteHashValue(new IoHash(value.HashData));
+
+        /// <inheritdoc/>
+        public override void WriteNamed(CbWriter writer, Utf8String name, JupiterContentHash value) => writer.WriteHash(name, new IoHash(value.HashData));
+    }
+
+    public class ContentHashJsonConverter : JsonConverter<JupiterContentHash>
+    {
+        public override JupiterContentHash Read(ref Utf8JsonReader reader, Type typeToConvert, JsonSerializerOptions options)
+        {
+            string? str = reader.GetString();
+            if (str == null)
+            {
+                throw new InvalidDataException("Unable to parse content hash");
+            }
+
+            return new JupiterContentHash(str);
+        }
+
+        public override void Write(Utf8JsonWriter writer, JupiterContentHash value, JsonSerializerOptions options)
+        {
+            writer.WriteStringValue(value.ToString());
+        }
+    }
+
+    public class ContentHashTypeConverter : TypeConverter
+    {
+        public override bool CanConvertFrom(ITypeDescriptorContext? context, Type sourceType)
+        {  
+            if (sourceType == typeof(string))  
+            {  
+                return true;
+            }  
+            return base.CanConvertFrom(context, sourceType);
+        }  
+  
+        public override object? ConvertFrom(ITypeDescriptorContext? context, CultureInfo? culture, object value)  
+        {
+            if (value is string s)
+            {
+                return new JupiterContentHash(s);
+            }
+
+            return base.ConvertFrom(context, culture, value);  
+        }
+
+        public override bool CanConvertTo(ITypeDescriptorContext? context, Type? destinationType)
+        {
+            if (destinationType == typeof(string))
+            {
+                return true;
+            }
+            return base.CanConvertTo(context, destinationType);
+        }
+
+        public override object? ConvertTo(ITypeDescriptorContext? context, CultureInfo? culture, object? value, Type destinationType)
+        {
+            if (destinationType == typeof(string))
+            {
+                JupiterContentHash? identifier = (JupiterContentHash?)value;
+                return identifier?.ToString();
+            }
+            return base.ConvertTo(context, culture, value, destinationType);
+        }
+    }
+}
diff --git a/Engine/Source/Programs/Horde/Horde.Server/Ddc/NamespaceNotFoundException.cs b/Engine/Source/Programs/Horde/Horde.Server/Ddc/NamespaceNotFoundException.cs
new file mode 100644
--- /dev/null
+++ b/Engine/Source/Programs/Horde/Horde.Server/Ddc/NamespaceNotFoundException.cs
@@ -0,0 +1,24 @@
+// Copyright Epic Games, Inc. All Rights Reserved.
+
+using System;
+using EpicGames.Horde.Storage;
+
+#pragma warning disable CS1591
+
+namespace Jupiter.Common
+{
+    public class NamespaceNotFoundException : Exception
+    {
+        public NamespaceId Namespace { get; }
+
+        public NamespaceNotFoundException(NamespaceId @namespace) : base($"Could not find namespace {@namespace}")
+        {
+            Namespace = @namespace;
+        }
+
+        public NamespaceNotFoundException(NamespaceId @namespace, string message) : base(message)
+        {
+            Namespace = @namespace;
+        }
+    }
+}
diff --git a/Engine/Source/Programs/Horde/Horde.Server/Ddc/NginxRedirectHelper.cs b/Engine/Source/Programs/Horde/Horde.Server/Ddc/NginxRedirectHelper.cs
new file mode 100644
--- /dev/null
+++ b/Engine/Source/Programs/Horde/Horde.Server/Ddc/NginxRedirectHelper.cs
@@ -0,0 +1,83 @@
+// Copyright Epic Games, Inc. All Rights Reserved.
+
+using System;
+using System.Threading.Tasks;
+using Microsoft.AspNetCore.Http;
+using Microsoft.AspNetCore.Mvc;
+using Microsoft.Extensions.Options;
+
+#pragma warning disable CS1591
+
+namespace Jupiter.Implementation
+{
+    public class NginxRedirectHelper
+    {
+        private readonly IOptionsMonitor<NginxSettings> _nginxSettings;
+
+        public NginxRedirectHelper(IOptionsMonitor<NginxSettings> nginxSettings)
+        {
+            _nginxSettings = nginxSettings;
+        }
+
+        public IActionResult CreateActionResult(BlobContents blobContents, string contentType)
+        {
+            if (blobContents.LocalPath == null)
+            {
+                throw new Exception("No local path available for blob contents when attempting to use nginx redirect");
+            }
+
+            string s = blobContents.LocalPath.Replace("\\", "/", StringComparison.OrdinalIgnoreCase);
+            // dispose of the open stream because we will just use the path to it and have nginx read and send it
+            blobContents.Dispose();
+            return new NginxRedirectOkResult($"{_nginxSettings.CurrentValue.NginxRootPath}/{s}", contentType);
+        }
+
+        public bool CanRedirect(HttpRequest request, BlobContents blobContents)
+        {
+            if (!_nginxSettings.CurrentValue.UseNginxRedirect)
+            {
+                return false;
+            }
+
+            if (blobContents.LocalPath == null)
+            {
+                return false;
+            }
+
+            if (!request.Headers.ContainsKey(_nginxSettings.CurrentValue.RequiredNginxHeader))
+            {
+                return false;
+            }
+
+            return true;
+        }
+    }
+
+    public class NginxSettings
+    {
+        public bool UseNginxRedirect { get; set; } = false;
+        public string NginxRootPath { get; set; } = "/nginx-blobs";
+        public string RequiredNginxHeader { get; set; } = "X-Jupiter-XAccel-Supported";
+    }
+
+    public class NginxRedirectOkResult : OkResult
+    {
+        private readonly string _nginxRedirectPath;
+        private readonly string _contentType;
+
+        public NginxRedirectOkResult(string nginxRedirectPath, string contentType)
+        {
+            _nginxRedirectPath = nginxRedirectPath;
+            _contentType = contentType;
+        }
+
+        public override Task ExecuteResultAsync(ActionContext context)
+        {
+            context.HttpContext.Response.Headers.ContentType = _contentType;
+            context.HttpContext.Response.Headers.Add("X-Accel-Redirect", _nginxRedirectPath);
+            // disable buffering in nginx as this is mostly used for large objects which should be streamed
+            context.HttpContext.Response.Headers.Add("X-Accel-Buffering", "no");
+            return base.ExecuteResultAsync(context);
+        }
+    }
+}
diff --git a/Engine/Source/Programs/Horde/Horde.Server/Ddc/ObjectController.cs b/Engine/Source/Programs/Horde/Horde.Server/Ddc/ObjectController.cs
new file mode 100644
--- /dev/null
+++ b/Engine/Source/Programs/Horde/Horde.Server/Ddc/ObjectController.cs
@@ -0,0 +1,299 @@
+﻿// Copyright Epic Games, Inc. All Rights Reserved.
+
+using System;
+using System.Collections.Concurrent;
+using System.Collections.Generic;
+using System.ComponentModel.DataAnnotations;
+using System.Linq;
+using System.Net;
+using System.Net.Mime;
+using System.Threading.Tasks;
+using EpicGames.AspNet;
+using EpicGames.Horde.Storage;
+using EpicGames.Serialization;
+using Jupiter.Implementation;
+using Jupiter.Common.Implementation;
+using Jupiter.Utils;
+using Microsoft.AspNetCore.Authorization;
+using Microsoft.AspNetCore.Mvc;
+using Microsoft.Extensions.Logging;
+
+#pragma warning disable CS1591
+
+namespace Jupiter.Controllers
+{
+    using IDiagnosticContext = Serilog.IDiagnosticContext;
+    using BlobNotFoundException = Jupiter.Implementation.BlobNotFoundException;
+
+    [ApiController]
+    [Route("api/v1/objects", Order = 0)]
+    [Authorize]
+    [Produces(CustomMediaTypeNames.UnrealCompactBinary, MediaTypeNames.Application.Json)]
+    public class ObjectController : ControllerBase
+    {
+        private readonly IBlobService _storage;
+        private readonly IDiagnosticContext _diagnosticContext;
+        private readonly IRequestHelper _requestHelper;
+        private readonly IReferenceResolver _referenceResolver;
+        private readonly BufferedPayloadFactory _bufferedPayloadFactory;
+
+        private readonly ILogger _logger;
+
+        public ObjectController(IBlobService storage, IDiagnosticContext diagnosticContext, IRequestHelper requestHelper, IReferenceResolver referenceResolver, BufferedPayloadFactory bufferedPayloadFactory, ILogger<ObjectController> logger)
+        {
+            _storage = storage;
+            _diagnosticContext = diagnosticContext;
+            _requestHelper = requestHelper;
+            _referenceResolver = referenceResolver;
+            _bufferedPayloadFactory = bufferedPayloadFactory;
+            _logger = logger;
+        }
+
+        [HttpGet("{ns}/{id}")]
+        [ProducesDefaultResponseType]
+        public async Task<IActionResult> Get(
+            [Required] NamespaceId ns,
+            [Required] BlobIdentifier id)
+        {
+            ActionResult? result = await _requestHelper.HasAccessToNamespace(User, Request, ns, new [] { JupiterAclAction.ReadObject });
+            if (result != null)
+            {
+                return result;
+            }
+
+            try
+            {
+                BlobContents blobContents = await _storage.GetObject(ns, id);
+
+                return File(blobContents.Stream, CustomMediaTypeNames.UnrealCompactBinary);
+            }
+            catch (BlobNotFoundException e)
+            {
+                return NotFound(new ValidationProblemDetails {Title = $"Object {e.Blob} not found"});
+            }
+        }
+
+        [HttpHead("{ns}/{id}")]
+        [ProducesDefaultResponseType]
+        public async Task<IActionResult> Head(
+            [Required] NamespaceId ns,
+            [Required] BlobIdentifier id)
+        {
+            ActionResult? result = await _requestHelper.HasAccessToNamespace(User, Request, ns, new [] { JupiterAclAction.ReadObject });
+            if (result != null)
+            {
+                return result;
+            }
+
+            bool exists = await _storage.Exists(ns, id);
+
+            if (!exists)
+            {
+                return NotFound(new ValidationProblemDetails {Title = $"Object {id} not found"});
+            }
+
+            return Ok();
+        }
+
+        [HttpPost("{ns}/exists")]
+        [ProducesDefaultResponseType]
+        public async Task<IActionResult> ExistsMultiple(
+            [Required] NamespaceId ns,
+            [Required] [FromQuery] List<BlobIdentifier> id)
+        {
+            ActionResult? result = await _requestHelper.HasAccessToNamespace(User, Request, ns, new [] { JupiterAclAction.ReadObject });
+            if (result != null)
+            {
+                return result;
+            }
+
+            ConcurrentBag<BlobIdentifier> missingBlobs = new ConcurrentBag<BlobIdentifier>();
+
+            IEnumerable<Task> tasks = id.Select(async blob =>
+            {
+                if (!await _storage.Exists(ns, blob))
+                {
+                    missingBlobs.Add(blob);
+                }
+            });
+            await Task.WhenAll(tasks);
+
+            return Ok(new HeadMultipleResponse {Needs = missingBlobs.ToArray()});
+        }
+
+        [HttpPost("{ns}/exist")]
+        [ProducesDefaultResponseType]
+        public async Task<IActionResult> ExistsBody(
+            [Required] NamespaceId ns,
+            [FromBody] BlobIdentifier[] bodyIds)
+        {
+            ActionResult? result = await _requestHelper.HasAccessToNamespace(User, Request, ns, new [] { JupiterAclAction.ReadObject });
+            if (result != null)
+            {
+                return result;
+            }
+
+            ConcurrentBag<BlobIdentifier> missingBlobs = new ConcurrentBag<BlobIdentifier>();
+
+            IEnumerable<Task> tasks = bodyIds.Select(async blob =>
+            {
+                if (!await _storage.Exists(ns, blob))
+                {
+                    missingBlobs.Add(blob);
+                }
+            });
+            await Task.WhenAll(tasks);
+
+            return Ok(new HeadMultipleResponse { Needs = missingBlobs.ToArray() });
+        }
+
+        [HttpPut("{ns}/{id}")]
+        [RequiredContentType(CustomMediaTypeNames.UnrealCompactBinary)]
+        public async Task<IActionResult> Put(
+            [Required] NamespaceId ns,
+            [Required] BlobIdentifier id)
+        {
+            ActionResult? result = await _requestHelper.HasAccessToNamespace(User, Request, ns, new [] { JupiterAclAction.WriteObject });
+            if (result != null)
+            {
+                return result;
+            }
+
+            _diagnosticContext.Set("Content-Length", Request.ContentLength ?? -1);
+            try
+            {
+                using IBufferedPayload payload = await _bufferedPayloadFactory.CreateFromRequest(Request);
+
+                BlobIdentifier identifier = await _storage.PutObject(ns, payload, id);
+                return Ok(new PutBlobResponse(identifier));
+            }
+            catch (ClientSendSlowException e)
+            {
+                return Problem(e.Message, null, (int)HttpStatusCode.RequestTimeout);
+            }
+        }
+
+        [HttpGet("{ns}/{id}/references")]
+        public async Task<IActionResult> ResolveReferences(
+            [Required] NamespaceId ns,
+            [Required] BlobIdentifier id)
+        {
+            ActionResult? result = await _requestHelper.HasAccessToNamespace(User, Request, ns, new [] { JupiterAclAction.ReadObject });
+            if (result != null)
+            {
+                return result;
+            }
+
+            BlobContents blob;
+            try
+            {
+                blob = await _storage.GetObject(ns, id);
+            }
+            catch (BlobNotFoundException e)
+            {
+                return NotFound(new ValidationProblemDetails {Title = $"Object {e.Blob} not found"});
+            }
+           
+            byte[] blobContents = await blob.Stream.ToByteArray();
+            if (blobContents.Length == 0)
+            {
+                _logger.LogWarning("0 byte object found for {Id} {Namespace}", id, ns);
+            }
+
+            CbObject compactBinaryObject;
+            try
+            {
+                compactBinaryObject = new CbObject(blobContents);
+            }
+            catch (IndexOutOfRangeException)
+            {
+                return Problem(title: $"{id} was not a proper compact binary object.", detail: "Index out of range");
+            }
+
+            try
+            {
+                BlobIdentifier[] references = await _referenceResolver.GetReferencedBlobs(ns, compactBinaryObject).ToArrayAsync();
+                return Ok(new ResolvedReferencesResult(references));
+            }
+            catch (PartialReferenceResolveException e)
+            {
+                return BadRequest(new ValidationProblemDetails {Title = $"Object {id} is missing content ids", Detail = $"Following content ids are invalid: {string.Join(",", e.UnresolvedReferences)}"});
+            }
+            catch (ReferenceIsMissingBlobsException e)
+            {
+                return BadRequest(new ValidationProblemDetails {Title = $"Object {id} is missing blobs", Detail = $"Following blobs are missing: {string.Join(",", e.MissingBlobs)}"});
+            }
+        }
+
+        [HttpDelete("{ns}/{id}")]
+        public async Task<IActionResult> Delete(
+            [Required] NamespaceId ns,
+            [Required] BlobIdentifier id)
+        {
+            ActionResult? result = await _requestHelper.HasAccessToNamespace(User, Request, ns, new [] { JupiterAclAction.DeleteObject });
+            if (result != null)
+            {
+                return result;
+            }
+
+            await _storage.DeleteObject(ns, id);
+
+            return Ok( new DeletedResponse
+            {
+                DeletedCount = 1
+            });
+        }
+
+        [HttpDelete("{ns}")]
+        public async Task<IActionResult> DeleteNamespace(
+            [Required] NamespaceId ns)
+        {
+            ActionResult? result = await _requestHelper.HasAccessToNamespace(User, Request, ns, new [] { JupiterAclAction.DeleteNamespace });
+            if (result != null)
+            {
+                return result;
+            }
+
+            await _storage.DeleteNamespace(ns);
+
+            return Ok();
+        }
+    }
+
+    public class PutBlobResponse
+    {
+        public PutBlobResponse()
+        {
+            Identifier = null!;
+        }
+
+        public PutBlobResponse(BlobIdentifier identifier)
+        {
+            Identifier = identifier;
+        }
+
+        [CbField("identifier")]
+        public BlobIdentifier Identifier { get; set; }
+    }
+
+    public class DeletedResponse
+    {
+        public int DeletedCount { get; set; }
+    }
+
+    public class ResolvedReferencesResult
+    {
+        public ResolvedReferencesResult()
+        {
+            References = null!;
+        }
+
+        public ResolvedReferencesResult(BlobIdentifier[] references)
+        {
+            References = references;
+        }
+
+        [CbField("references")]
+        public BlobIdentifier[] References { get; set; }
+    }
+}
diff --git a/Engine/Source/Programs/Horde/Horde.Server/Ddc/ObjectService.cs b/Engine/Source/Programs/Horde/Horde.Server/Ddc/ObjectService.cs
new file mode 100644
--- /dev/null
+++ b/Engine/Source/Programs/Horde/Horde.Server/Ddc/ObjectService.cs
@@ -0,0 +1,58 @@
+// Copyright Epic Games, Inc. All Rights Reserved.
+
+using System.Collections.Generic;
+using System.Threading.Tasks;
+using EpicGames.Horde.Storage;
+using EpicGames.Serialization;
+using Jupiter.Implementation;
+
+namespace Horde.Server.Ddc
+{
+	class ObjectService : IObjectService
+	{
+		public Task<bool> Delete(NamespaceId ns, BucketId bucket, IoHashKey key)
+		{
+			throw new System.NotImplementedException();
+		}
+
+		public Task<long> DeleteBucket(NamespaceId ns, BucketId bucket)
+		{
+			throw new System.NotImplementedException();
+		}
+
+		public Task<long> DropNamespace(NamespaceId ns)
+		{
+			throw new System.NotImplementedException();
+		}
+
+		public Task<bool> Exists(NamespaceId ns, BucketId bucket, IoHashKey key)
+		{
+			throw new System.NotImplementedException();
+		}
+
+		public Task<(ContentId[], BlobIdentifier[])> Finalize(NamespaceId ns, BucketId bucket, IoHashKey key, BlobIdentifier blobHash)
+		{
+			throw new System.NotImplementedException();
+		}
+
+		public Task<(ObjectRecord, BlobContents?)> Get(NamespaceId ns, BucketId bucket, IoHashKey key, string[] fields, bool doLastAccessTracking = true)
+		{
+			throw new System.NotImplementedException();
+		}
+
+		public IAsyncEnumerable<NamespaceId> GetNamespaces()
+		{
+			throw new System.NotImplementedException();
+		}
+
+		public Task<List<BlobIdentifier>> GetReferencedBlobs(NamespaceId ns, BucketId bucket, IoHashKey key)
+		{
+			throw new System.NotImplementedException();
+		}
+
+		public Task<(ContentId[], BlobIdentifier[])> Put(NamespaceId ns, BucketId bucket, IoHashKey key, BlobIdentifier blobHash, CbObject payload)
+		{
+			throw new System.NotImplementedException();
+		}
+	}
+}
diff --git a/Engine/Source/Programs/Horde/Horde.Server/Ddc/ReferencesController.cs b/Engine/Source/Programs/Horde/Horde.Server/Ddc/ReferencesController.cs
new file mode 100644
--- /dev/null
+++ b/Engine/Source/Programs/Horde/Horde.Server/Ddc/ReferencesController.cs
@@ -0,0 +1,1352 @@
+// Copyright Epic Games, Inc. All Rights Reserved.
+
+using System;
+using System.Collections.Concurrent;
+using System.Collections.Generic;
+using System.ComponentModel.DataAnnotations;
+using System.Globalization;
+using System.IO;
+using System.Linq;
+using System.Net;
+using System.Net.Mime;
+using System.Text;
+using System.Text.Json.Serialization;
+using System.Threading;
+using System.Threading.Tasks;
+using EpicGames.AspNet;
+using EpicGames.Core;
+using EpicGames.Horde.Storage;
+using EpicGames.Serialization;
+using Jupiter.Implementation;
+using JetBrains.Annotations;
+using Jupiter.Common;
+using Jupiter.Common.Implementation;
+using Jupiter.Utils;
+using Microsoft.AspNetCore.Authorization;
+using Microsoft.AspNetCore.Http;
+using Microsoft.AspNetCore.Http.Extensions;
+using Microsoft.AspNetCore.Mvc;
+using Microsoft.Extensions.DependencyInjection;
+using Microsoft.Extensions.Logging;
+using OpenTelemetry.Trace;
+
+#pragma warning disable CS1591
+
+using JupiterContentHash = Jupiter.Implementation.JupiterContentHash;
+using ContentId = Jupiter.Implementation.ContentId;
+
+namespace Jupiter.Controllers
+{
+	using IDiagnosticContext = Serilog.IDiagnosticContext;
+    using BlobNotFoundException = Jupiter.Implementation.BlobNotFoundException;
+
+    [ApiController]
+    [FormatFilter]
+    [Produces(MediaTypeNames.Application.Json, MediaTypeNames.Application.Octet, CustomMediaTypeNames.UnrealCompactBinary)]
+    [Route("api/v1/refs")]
+    [Authorize]
+    public class ReferencesController : ControllerBase
+    {
+        private readonly IDiagnosticContext _diagnosticContext;
+        private readonly FormatResolver _formatResolver;
+        private readonly BufferedPayloadFactory _bufferedPayloadFactory;
+        private readonly IReferenceResolver _referenceResolver;
+        private readonly NginxRedirectHelper _nginxRedirectHelper;
+        private readonly IRequestHelper _requestHelper;
+        private readonly Tracer _tracer;
+
+        private readonly ILogger _logger;
+        private readonly IObjectService _objectService;
+        private readonly IBlobService _blobStore;
+
+        public ReferencesController(IObjectService objectService, IBlobService blobStore, IDiagnosticContext diagnosticContext, FormatResolver formatResolver, BufferedPayloadFactory bufferedPayloadFactory, IReferenceResolver referenceResolver, NginxRedirectHelper nginxRedirectHelper, IRequestHelper requestHelper, Tracer tracer, ILogger<ReferencesController> logger)
+        {
+            _objectService = objectService;
+            _blobStore = blobStore;
+            _diagnosticContext = diagnosticContext;
+            _formatResolver = formatResolver;
+            _bufferedPayloadFactory = bufferedPayloadFactory;
+            _referenceResolver = referenceResolver;
+            _nginxRedirectHelper = nginxRedirectHelper;
+            _requestHelper = requestHelper;
+            _tracer = tracer;
+            _logger = logger;
+        }
+
+        /// <summary>
+        /// Returns all the known namespace the token has access to
+        /// </summary>
+        /// <returns></returns>
+        [HttpGet("")]
+        [ProducesDefaultResponseType]
+        [ProducesResponseType(type: typeof(ProblemDetails), 400)]
+        public async Task<IActionResult> GetNamespaces()
+        {
+            NamespaceId[] namespaces = await _objectService.GetNamespaces().ToArrayAsync();
+
+            // filter namespaces down to only the namespaces the user has access to
+            List<NamespaceId> namespacesWithAccess = new();
+            foreach (NamespaceId ns in namespaces)
+            {
+                ActionResult? accessResult = await _requestHelper.HasAccessToNamespace(User, Request, ns, new [] { JupiterAclAction.ReadObject });
+                if (accessResult == null)
+                {
+                    namespacesWithAccess.Add(ns);
+                }
+            }
+
+            return Ok(new GetNamespacesResponse(namespacesWithAccess.ToArray()));
+        }
+
+        /// <summary>
+        /// Returns a refs key
+        /// </summary>
+        /// <param name="ns">Namespace. Each namespace is completely separated from each other. Use for different types of data that is never expected to be similar (between two different games for instance). Example: `uc4.ddc`</param>
+        /// <param name="bucket">The category/type of record you are caching. Is a clustered key together with the actual key, but all records in the same bucket can be dropped easily. Example: `terrainTexture` </param>
+        /// <param name="key">The unique name of this particular key. `iAmAVeryValidKey`</param>
+        /// <param name="format">Optional specifier to set which output format is used json/raw/cb</param>
+        [HttpGet("{ns}/{bucket}/{key}.{format?}", Order = 500)]
+        [Produces(MediaTypeNames.Application.Json, MediaTypeNames.Application.Octet, CustomMediaTypeNames.UnrealCompactBinary, CustomMediaTypeNames.JupiterInlinedPayload, CustomMediaTypeNames.UnrealCompactBinaryPackage)]
+        public async Task<IActionResult> Get(
+            [FromRoute] [Required] NamespaceId ns,
+            [FromRoute] [Required] BucketId bucket,
+            [FromRoute] [Required] IoHashKey key,
+            [FromRoute] string? format = null)
+        {
+            ActionResult? accessResult = await _requestHelper.HasAccessToNamespace(User, Request, ns, new [] { JupiterAclAction.ReadObject });
+            if (accessResult != null)
+            {
+                return accessResult;
+            }
+
+            try
+            {
+                (ObjectRecord objectRecord, BlobContents? blob) = await _objectService.Get(ns, bucket, key, Array.Empty<string>());
+
+                if (blob == null)
+                {
+                    throw new InvalidOperationException($"Blob was null when attempting to fetch {ns} {bucket} {key}");
+                }
+
+                if (!objectRecord.IsFinalized)
+                {
+                    // we do not consider un-finalized objects as valid
+                    return BadRequest(new ProblemDetails { Title = $"Object {objectRecord.Bucket} {objectRecord.Name} is not finalized." });
+                }
+
+                Response.Headers[CommonHeaders.HashHeaderName] = objectRecord.BlobIdentifier.ToString();
+                Response.Headers[CommonHeaders.LastAccessHeaderName] = objectRecord.LastAccess.ToString(CultureInfo.InvariantCulture);
+
+                async Task WriteBody(BlobContents blobContents, string contentType)
+                {
+                    IServerTiming? serverTiming = Request.HttpContext.RequestServices.GetService<IServerTiming>();
+                    using ServerTimingMetricScoped? serverTimingScope = serverTiming?.CreateServerTimingMetricScope("body.write", "Time spent writing body");
+
+                    long contentLength = blobContents.Length;
+                    using TelemetrySpan scope = _tracer.StartActiveSpan("body.write").SetAttribute("operation.name", "body.write");
+                    scope.SetAttribute("content-length", contentLength);
+                    const int BufferSize = 64 * 1024;
+                    Stream outputStream = Response.Body;
+                    Response.ContentLength = contentLength;
+                    Response.ContentType = contentType;
+                    Response.StatusCode = StatusCodes.Status200OK;
+                    try
+                    {
+                        await StreamCopyOperation.CopyToAsync(blobContents.Stream, outputStream, count: null, bufferSize: BufferSize, cancel: Response.HttpContext.RequestAborted);
+                    }
+                    catch (OperationCanceledException)
+                    {
+                        // do not raise exceptions for cancelled writes
+                        // as we have already started writing a response we can not change the status code
+                        // so we just drop a warning and proceed
+                        _logger.LogWarning("The operation was canceled while writing the body");
+                    }
+                }
+
+                string responseType = _formatResolver.GetResponseType(Request, format, CustomMediaTypeNames.UnrealCompactBinary);
+                Tracer.CurrentSpan.SetAttribute("response-type", responseType);
+                switch (responseType)
+                {
+                    case CustomMediaTypeNames.UnrealCompactBinary:
+                    {
+                        // for compact binary we can just serialize our internal object
+                        await WriteBody(blob, CustomMediaTypeNames.UnrealCompactBinary);
+
+                        break;
+                    }
+                    case MediaTypeNames.Application.Octet:
+                    {
+                        byte[] blobMemory = await blob.Stream.ToByteArray();
+                        CbObject cb = new CbObject(blobMemory);
+
+                        (int,CbField?) CountFields(CbObject o)
+                        {
+                            int count = 0;
+                            CbField? foundField = null;
+                            cb.IterateAttachments(field =>
+                            {
+                                ++count;
+                                if (field.IsBinaryAttachment())
+                                {
+                                    foundField = field;
+                                }
+                            });
+
+                            return (count, foundField);
+                        }
+
+                        // breaking lambda call into private method to workaround incorrect triggering of CA1508 - [REDACTED_PROFILE_URL]/roslyn-analyzers/issues/5254
+                        (int countOfAttachmentFields,CbField? binaryAttachmentField) = CountFields(cb);
+
+                        if (countOfAttachmentFields == 1 && binaryAttachmentField != null)
+                        {
+                            // there is a single attachment field and that is of the binary attachment type, fetch that attachment and return it instead of the compact binary
+                            // this is so that we match the uploaded that happened as a octet-stream which generates a small cb object with a single attachment
+
+                            IoHash hash = binaryAttachmentField.AsBinaryAttachment();
+
+                            BlobContents referencedBlobContents = await _blobStore.GetObject(ns, BlobIdentifier.FromIoHash(hash));
+
+                            if (_nginxRedirectHelper.CanRedirect(Request, referencedBlobContents))
+                            {
+                                return _nginxRedirectHelper.CreateActionResult(referencedBlobContents, MediaTypeNames.Application.Octet);
+                            }
+                            await WriteBody(referencedBlobContents, MediaTypeNames.Application.Octet);
+                            break;
+                        }
+
+                        // this doesn't look like the generated compact binary so we just return the payload
+                        await using BlobContents contents = new(blobMemory);
+                        await WriteBody(contents, MediaTypeNames.Application.Octet);
+                        break;
+                    }
+                    case MediaTypeNames.Application.Json:
+                    {
+                        byte[] blobMemory;
+                        {
+                            using TelemetrySpan scope = _tracer.StartActiveSpan("json.readblob").SetAttribute("operation.name", "json.readblob");
+                            blobMemory = await blob.Stream.ToByteArray();
+                        }
+                        CbObject cb = new CbObject(blobMemory);
+                        string s = cb.ToJson();
+                        await using BlobContents contents = new BlobContents(Encoding.UTF8.GetBytes(s));
+                        await WriteBody(contents, MediaTypeNames.Application.Json);
+                        break;
+
+                    }
+                    case CustomMediaTypeNames.UnrealCompactBinaryPackage:
+                    {
+                        using TelemetrySpan packageScope = _tracer.StartActiveSpan("cbpackage.fetch").SetAttribute("operation.name", "cbpackage.fetch");
+                        byte[] blobMemory = await blob.Stream.ToByteArray();
+                        CbObject cb = new CbObject(blobMemory);
+
+                        IAsyncEnumerable<Attachment> attachments = _referenceResolver.GetAttachments(ns, cb);
+
+                        using CbPackageBuilder writer = new CbPackageBuilder();
+                        writer.AddAttachment(objectRecord.BlobIdentifier.AsIoHash(), CbPackageAttachmentFlags.IsObject, blobMemory);
+
+                        await Parallel.ForEachAsync(attachments, async (attachment, token) =>
+                        {
+                            IoHash attachmentHash = attachment.AsIoHash();
+                            CbPackageAttachmentFlags flags = 0;
+
+                            try
+                            {
+                                BlobContents attachmentContents;
+                                if (attachment is BlobAttachment blobAttachment)
+                                {
+                                    BlobIdentifier referencedBlob = blobAttachment.Identifier;
+                                    attachmentContents = await _blobStore.GetObject(ns, referencedBlob);
+                                }
+                                else if (attachment is ObjectAttachment objectAttachment)
+                                {
+                                    flags |= CbPackageAttachmentFlags.IsObject;
+                                    BlobIdentifier referencedBlob = objectAttachment.Identifier;
+                                    attachmentContents = await _blobStore.GetObject(ns, referencedBlob);
+                                }
+                                else if (attachment is ContentIdAttachment contentIdAttachment)
+                                {
+
+                                    ContentId contentId = contentIdAttachment.Identifier;
+                                    (attachmentContents, string mime) = await _blobStore.GetCompressedObject(ns, contentId, HttpContext.RequestServices);
+                                    if (mime == CustomMediaTypeNames.UnrealCompressedBuffer)
+                                    {
+                                        flags |= CbPackageAttachmentFlags.IsCompressed;
+                                    }
+                                    else
+                                    {
+                                        // this resolved to a uncompressed blob, the content id existed the the compressed blob didn't
+                                        // so resetting flags to indicate this.
+                                        flags = 0;
+                                    }
+                                }
+                                else
+                                {
+                                    throw new NotSupportedException($"Unknown attachment type {attachment.GetType()}");
+                                }
+
+                                writer.AddAttachment(attachmentHash, flags, attachmentContents.Stream, (ulong)attachmentContents.Length);
+                            }
+                            catch (Exception e)
+                            {
+                                (CbObject errorObject, HttpStatusCode _) = ToErrorResult(e);
+                                writer.AddAttachment(attachmentHash,
+                                    CbPackageAttachmentFlags.IsError | CbPackageAttachmentFlags.IsObject,
+                                    errorObject.GetView().ToArray());
+                            }
+                        });
+
+                        byte[] packageBytes;
+                        {
+                            using TelemetrySpan _ = _tracer.StartActiveSpan("cbpackage.buffer").SetAttribute("operation.name", "cbpackage.buffer");
+                            packageBytes = await writer.ToByteArray();
+                        }
+                        await using BlobContents contents = new BlobContents(packageBytes);
+                        await WriteBody(contents, CustomMediaTypeNames.UnrealCompactBinaryPackage);
+                        break;
+                    }
+                    case CustomMediaTypeNames.JupiterInlinedPayload:
+                    {
+                        byte[] blobMemory = await blob.Stream.ToByteArray();
+                        CbObject cb = new CbObject(blobMemory);
+
+                        static (int, int) CountFields(CbObject o)
+                        {
+                            int countOfBinaryAttachmentFields = 0;
+                            int countOfAttachmentFields = 0;
+
+                            o.IterateAttachments(field =>
+                            {
+                                if (field.IsBinaryAttachment())
+                                {
+                                    ++countOfBinaryAttachmentFields;
+                                }
+
+                                if (field.IsAttachment())
+                                {
+                                    ++countOfAttachmentFields;
+                                }
+                            });
+
+                            return (countOfBinaryAttachmentFields, countOfAttachmentFields);
+                        }
+                        // breaking lambda call into private method to workaround incorrect triggering of CA1508 - [REDACTED_PROFILE_URL]/roslyn-analyzers/issues/5254
+                        (int countOfAttachmentFields, int countOfBinaryAttachmentFields) = CountFields(cb);
+
+                        // if the object consists of a single attachment field we return this attachment field instead
+                        if (countOfBinaryAttachmentFields == 1 && countOfAttachmentFields == 1)
+                        {
+                            // fetch the blob so we can resolve any content ids in it
+                            List<BlobIdentifier> referencedBlobs;
+                            try
+                            {
+                                IAsyncEnumerable<BlobIdentifier> referencedBlobsEnumerable = _referenceResolver.GetReferencedBlobs(ns, cb);
+                                referencedBlobs = await referencedBlobsEnumerable.ToListAsync();
+                            }
+                            catch (PartialReferenceResolveException)
+                            {
+                                return NotFound(new ProblemDetails {Title = $"Object {bucket} {key} in namespace {ns} was missing some content ids"});
+                            }
+                            catch (ReferenceIsMissingBlobsException)
+                            {
+                                return NotFound(new ProblemDetails {Title = $"Object {bucket} {key} in namespace {ns} was missing some blobs"});
+                            }
+
+                            if (referencedBlobs.Count == 1)
+                            {
+                                BlobIdentifier attachmentToSend = referencedBlobs.First();
+                                try
+                                {
+                                    BlobContents referencedBlobContents = await _blobStore.GetObject(ns, attachmentToSend);
+                                    Response.Headers[CommonHeaders.InlinePayloadHash] = attachmentToSend.ToString();
+
+                                    if (_nginxRedirectHelper.CanRedirect(Request, referencedBlobContents))
+                                    {
+                                        return _nginxRedirectHelper.CreateActionResult(referencedBlobContents, CustomMediaTypeNames.JupiterInlinedPayload);
+                                    }
+                                    await WriteBody(referencedBlobContents, CustomMediaTypeNames.JupiterInlinedPayload);
+                                }
+                                catch (BlobNotFoundException)
+                                {
+                                    return NotFound(new ProblemDetails {Title = $"Object {bucket} {key} in namespace {ns} was missing blob {attachmentToSend}"});
+                                }
+                                catch (Exception ex)
+                                {
+                                    _logger.LogError(ex, "Unknown exception encountered while writing body for jupiter inlined payload.");
+                                    throw;
+                                }
+                                return new EmptyResult();
+                            }
+                            else if (referencedBlobs.Count == 0)
+                            {
+                                return NotFound(new ProblemDetails
+                                {
+                                    Title =
+                                        $"Object {objectRecord.Bucket} {objectRecord.Name} did not resolve into any objects that we could find."
+                                });
+                            }
+
+                            return BadRequest(new ProblemDetails
+                            {
+                                Title =
+                                    $"Object {objectRecord.Bucket} {objectRecord.Name} contained a content id which resolved to more then 1 blob, unable to inline this object. Use compact object response instead."
+                            });
+                        }
+                        else if (countOfBinaryAttachmentFields == 0 && countOfAttachmentFields == 0)
+                        {
+                            // no attachments so we just return the compact object instead
+                            await using BlobContents contents = new BlobContents(blobMemory);
+                            await WriteBody(contents, CustomMediaTypeNames.JupiterInlinedPayload);
+                            return new EmptyResult();
+                        }
+
+                        return BadRequest(new ProblemDetails
+                        {
+                            Title =
+                                $"Object {objectRecord.Bucket} {objectRecord.Name} had more then 1 binary attachment field, unable to inline this object. Use compact object response instead."
+                        });
+                    }
+                    default:
+                        throw new NotImplementedException($"Unknown expected response type {responseType}");
+                }
+                
+                // this result is ignored as we write to the body explicitly
+                return new EmptyResult();
+
+            }
+            catch (NamespaceNotFoundException e)
+            {
+                return NotFound(new ProblemDetails {Title = $"Namespace {e.Namespace} did not exist"});
+            }
+            catch (ObjectNotFoundException e)
+            {
+                return NotFound(new ProblemDetails { Title = $"Object {e.Bucket} {e.Key} did not exist" });
+            }
+            catch (BlobNotFoundException e)
+            {
+                return NotFound(new ProblemDetails { Title = $"Object {e.Blob} in {e.Ns} not found" });
+            }
+        }
+
+       
+
+    /// <summary>
+    /// Returns the metadata about a ref key
+    /// </summary>
+    /// <param name="ns">Namespace. Each namespace is completely separated from each other. Use for different types of data that is never expected to be similar (between two different games for instance). Example: `uc4.ddc`</param>
+    /// <param name="bucket">The category/type of record you are caching. Is a clustered key together with the actual key, but all records in the same bucket can be dropped easily. Example: `terrainTexture` </param>
+    /// <param name="key">The unique name of this particular key. `iAmAVeryValidKey`</param>
+    /// <param name="fields">The fields to include in the response, omit this to include everything.</param>
+    [HttpGet("{ns}/{bucket}/{key}/metadata", Order = 500)]
+    public async Task<IActionResult> GetMetadata(
+        [FromRoute] [Required] NamespaceId ns,
+        [FromRoute] [Required] BucketId bucket,
+        [FromRoute] [Required] IoHashKey key,
+        [FromQuery] string[] fields)
+    {
+        ActionResult? accessResult = await _requestHelper.HasAccessToNamespace(User, Request, ns, new [] { JupiterAclAction.ReadObject });
+        if (accessResult != null)
+        {
+            return accessResult;
+        }
+
+        try
+        {
+            (ObjectRecord objectRecord, BlobContents? _) = await _objectService.Get(ns, bucket, key, fields);
+
+            return Ok(new RefMetadataResponse(objectRecord));
+        }
+        catch (NamespaceNotFoundException e)
+        {
+            return NotFound(new ProblemDetails {Title = $"Namespace {e.Namespace} did not exist"});
+        }
+        catch (ObjectNotFoundException e)
+        {
+            return NotFound(new ProblemDetails { Title = $"Object {e.Bucket} {e.Key} did not exist" });
+        }
+        catch (BlobNotFoundException e)
+        {
+            return NotFound(new ProblemDetails { Title = $"Object {e.Blob} in {e.Ns} not found" });
+        }
+    }
+
+        /// <summary>
+        /// Checks if a object exists
+        /// </summary>
+        /// <param name="ns">Namespace. Each namespace is completely separated from each other. Use for different types of data that is never expected to be similar (between two different games for instance). Example: `uc4.ddc`</param>
+        /// <param name="bucket">The category/type of record you are caching. Is a clustered key together with the actual key, but all records in the same bucket can be dropped easily. Example: `terrainTexture` </param>
+        /// <param name="key">The unique name of this particular key. `iAmAVeryValidKey`</param>
+        /// <returns>200 if it existed, 400 otherwise</returns>
+        [HttpHead("{ns}/{bucket}/{key}", Order = 500)]
+        [ProducesResponseType(type: typeof(OkResult), 200)]
+        [ProducesResponseType(type: typeof(ValidationProblemDetails), 400)]
+        public async Task<IActionResult> Head(
+            [FromRoute] [Required] NamespaceId ns,
+            [FromRoute] [Required] BucketId bucket,
+            [FromRoute] [Required] IoHashKey key)
+        {
+            ActionResult? accessResult = await _requestHelper.HasAccessToNamespace(User, Request, ns, new [] { JupiterAclAction.ReadObject });
+            if (accessResult != null)
+            {
+                return accessResult;
+            }
+
+            try
+            {
+                (ObjectRecord record, BlobContents? blob) = await _objectService.Get(ns, bucket, key, new string[] {"blobIdentifier", "IsFinalized"});
+                Response.Headers[CommonHeaders.HashHeaderName] = record.BlobIdentifier.ToString();
+
+                if (!record.IsFinalized)
+                {
+                    return NotFound(new ProblemDetails {Title = $"Object {bucket} {key} in namespace {ns} is not finalized."});
+                }
+
+                blob ??= await _blobStore.GetObject(ns, record.BlobIdentifier);
+
+                // we have to verify the blobs are available locally, as the record of the key is replicated a head of the content
+                // TODO: Once we support inline replication this step is not needed as at least one region as this blob, just maybe not this current one
+                byte[] blobContents = await blob.Stream.ToByteArray();
+                CbObject compactBinaryObject = new CbObject(blobContents);
+                // the reference resolver will throw if any blob is missing, so no need to do anything other then process each reference
+                IAsyncEnumerable<BlobIdentifier> references = _referenceResolver.GetReferencedBlobs(ns, compactBinaryObject);
+                List<BlobIdentifier>? _ = await references.ToListAsync();
+
+                // we have to verify the blobs are available locally, as the record of the key is replicated a head of the content
+                // TODO: Once we support inline replication this step is not needed as at least one region as this blob, just maybe not this current one
+                BlobIdentifier[] unknownBlobs = await _blobStore.FilterOutKnownBlobs(ns, new BlobIdentifier[] { record.BlobIdentifier });
+                if (unknownBlobs.Length != 0)
+                {
+                    return NotFound(new ProblemDetails {Title = $"Object {bucket} {key} in namespace {ns} had at least one missing blob."});
+                }
+            }
+            catch (NamespaceNotFoundException e)
+            {
+                return NotFound(new ProblemDetails {Title = $"Namespace {e.Namespace} did not exist"});
+            }
+            catch (BlobNotFoundException e)
+            {
+                return NotFound(new ProblemDetails { Title = $"Blob {e.Blob} in namespace {ns} did not exist" });
+            }
+            catch (ObjectNotFoundException e)
+            {
+                return NotFound(new ProblemDetails {Title = $"Object {e.Bucket} {e.Key} in namespace {e.Namespace} did not exist"});
+            }
+            catch (PartialReferenceResolveException)
+            {
+                return NotFound(new ProblemDetails {Title = $"Object {bucket} {key} in namespace {ns} was missing some content ids"});
+            }
+            catch (ReferenceIsMissingBlobsException)
+            {
+                return NotFound(new ProblemDetails {Title = $"Object {bucket} {key} in namespace {ns} was missing some blobs"});
+            }
+
+            return Ok();
+        }
+
+        [HttpGet("{ns}/exists")]
+        [ProducesDefaultResponseType]
+        public async Task<IActionResult> ExistsMultiple(
+            [FromRoute] [Required] NamespaceId ns,
+            [FromQuery] [Required] List<string> names)
+        {
+            ActionResult? accessResult = await _requestHelper.HasAccessToNamespace(User, Request, ns, new [] { JupiterAclAction.ReadObject });
+            if (accessResult != null)
+            {
+                return accessResult;
+            }
+
+            ConcurrentBag<(BucketId, IoHashKey)> missingObject = new ();
+
+            List<(BucketId, IoHashKey)> requestedNames = new List<(BucketId, IoHashKey)>();
+            foreach (string name in names)
+            {
+                int separatorIndex = name.IndexOf(".", StringComparison.Ordinal);
+                if (separatorIndex == -1)
+                {
+                    return BadRequest(new ProblemDetails() { Title = $"Key {name} did not contain a '.' separator" });
+                }
+
+                BucketId bucket = new BucketId(name.Substring(0, separatorIndex));
+                IoHashKey key = new IoHashKey(name.Substring(separatorIndex + 1));
+                requestedNames.Add((bucket, key));
+            }
+
+            IEnumerable<Task> tasks = requestedNames.Select(async pair =>
+            {
+                (BucketId bucket, IoHashKey key) = pair;
+                try
+                {
+                    (ObjectRecord record, BlobContents? blob) =
+                        await _objectService.Get(ns, bucket, key, new string[] { "blobIdentifier" });
+
+                    blob ??= await _blobStore.GetObject(ns, record.BlobIdentifier);
+
+                    // we have to verify the blobs are available locally, as the record of the key is replicated a head of the content
+                    // TODO: Once we support inline replication this step is not needed as at least one region as this blob, just maybe not this current one
+                    byte[] blobContents = await blob.Stream.ToByteArray();
+                    CbObject cb = new CbObject(blobContents);
+                    // the reference resolver will throw if any blob is missing, so no need to do anything other then process each reference
+                    IAsyncEnumerable<BlobIdentifier> references = _referenceResolver.GetReferencedBlobs(ns, cb);
+                    List<BlobIdentifier>? _ = await references.ToListAsync();
+                }
+                catch (ObjectNotFoundException)
+                {
+                    missingObject.Add((bucket, key));
+                }
+                catch (PartialReferenceResolveException)
+                {
+                    missingObject.Add((bucket, key));
+                }
+                catch (ReferenceIsMissingBlobsException)
+                {
+                    missingObject.Add((bucket, key));
+                }
+                catch (BlobNotFoundException)
+                {
+                    missingObject.Add((bucket, key));
+                }
+            });
+            await Task.WhenAll(tasks);
+
+            return Ok(new ExistCheckMultipleRefsResponse(missingObject.ToList()));
+        }
+
+        [HttpPut("{ns}/{bucket}/{key}.{format?}", Order = 500)]
+        [DisableRequestSizeLimit]
+        public async Task<IActionResult> PutObject(
+            [FromRoute] [Required] NamespaceId ns,
+            [FromRoute] [Required] BucketId bucket,
+            [FromRoute] [Required] IoHashKey key)
+        {
+            ActionResult? accessResult = await _requestHelper.HasAccessToNamespace(User, Request, ns, new [] { JupiterAclAction.WriteObject });
+            if (accessResult != null)
+            {
+                return accessResult;
+            }
+
+            _diagnosticContext.Set("Content-Length", Request.ContentLength ?? -1);
+
+            CbObject payloadObject;
+            BlobIdentifier blobHeader;
+
+            try
+            {
+                using IBufferedPayload payload = await _bufferedPayloadFactory.CreateFromRequest(Request);
+
+                BlobIdentifier headerHash;
+                if (Request.Headers.ContainsKey(CommonHeaders.HashHeaderName))
+                {
+                    headerHash = new BlobIdentifier(Request.Headers[CommonHeaders.HashHeaderName]);
+                }
+                else
+                {
+                    return BadRequest(new ProblemDetails
+                    {
+                        Title = $"Missing expected header {CommonHeaders.HashHeaderName}"
+                    });
+                }
+
+                blobHeader = headerHash;
+
+                switch (Request.ContentType)
+                {
+                    case MediaTypeNames.Application.Json:
+                    {
+                        // TODO: define a scheme for how a json object specifies references
+
+                        blobHeader = await _blobStore.PutObject(ns, payload, headerHash);
+
+                        // TODO: convert the json object into a compact binary instead
+                        CbWriter writer = new CbWriter();
+                        writer.BeginObject();
+                        writer.WriteBinaryAttachmentValue(blobHeader.AsIoHash());
+                        writer.EndObject();
+
+                        byte[] blob = writer.ToByteArray();
+                        payloadObject = new CbObject(blob);
+                        blobHeader = BlobIdentifier.FromBlob(blob);
+                        break;
+                    }
+                    case CustomMediaTypeNames.UnrealCompactBinary:
+                    {
+                        await using MemoryStream ms = new MemoryStream();
+                        await using Stream payloadStream = payload.GetStream();
+                        await payloadStream.CopyToAsync(ms);
+                        payloadObject = new CbObject(ms.ToArray());
+                        break;
+                    }
+                    case MediaTypeNames.Application.Octet:
+                    {
+                        blobHeader = await _blobStore.PutObject(ns, payload, headerHash);
+
+                        CbWriter writer = new CbWriter();
+                        writer.BeginObject();
+                        writer.WriteBinaryAttachment("RawHash", blobHeader.AsIoHash());
+                        writer.WriteInteger("RawSize", payload.Length);
+                        writer.EndObject();
+
+                        byte[] blob = writer.ToByteArray();
+                        payloadObject = new CbObject(blob);
+                        blobHeader = BlobIdentifier.FromBlob(blob);
+                        break;
+                    }
+                    default:
+                        throw new Exception($"Unknown request type {Request.ContentType}, if submitting a blob please use {MediaTypeNames.Application.Octet}");
+                }
+            }
+            catch (HashMismatchException e)
+            {
+                return BadRequest(new ProblemDetails
+                {
+                    Title = $"Incorrect hash, got hash \"{e.SuppliedHash}\" but hash of content was determined to be \"{e.ContentHash}\""
+                });
+            }
+            catch (ClientSendSlowException e)
+            {
+                return Problem(e.Message, null, (int)HttpStatusCode.RequestTimeout);
+            }
+
+            (ContentId[] missingReferences, BlobIdentifier[] missingBlobs) = await _objectService.Put(ns, bucket, key, blobHeader, payloadObject);
+
+            List<JupiterContentHash> missingHashes = new List<JupiterContentHash>(missingReferences);
+            missingHashes.AddRange(missingBlobs);
+            return Ok(new PutObjectResponse(missingHashes.ToArray()));
+        }
+
+        [HttpPut("{ns}/{bucket}/{key}", Order = 300)]
+        [DisableRequestSizeLimit]
+        [RequiredContentType(CustomMediaTypeNames.UnrealCompactBinaryPackage)]
+        public async Task<IActionResult> PutPackage(
+            [FromRoute][Required] NamespaceId ns,
+            [FromRoute][Required] BucketId bucket,
+            [FromRoute][Required] IoHashKey key)
+        {
+            ActionResult? accessResult = await _requestHelper.HasAccessToNamespace(User, Request, ns, new [] { JupiterAclAction.WriteObject });
+            if (accessResult != null)
+            {
+                return accessResult;
+            }
+
+            _diagnosticContext.Set("Content-Length", Request.ContentLength ?? -1);
+
+            byte[] b = await RequestUtil.ReadRawBody(Request);
+            CbPackageReader packageReader = await CbPackageReader.Create(new MemoryStream(b));
+
+            try
+            {
+                await foreach ((CbPackageAttachmentEntry entry, byte[] blob) in packageReader.IterateAttachments())
+                {
+                    if (entry.Flags.HasFlag(CbPackageAttachmentFlags.IsError))
+                    {
+                        return BadRequest(new ProblemDetails
+                        {
+                            Title = $"Package contained attachment with error {entry.AttachmentHash}\""
+                        });
+                    }
+                    if (entry.Flags.HasFlag(CbPackageAttachmentFlags.IsCompressed))
+                    {
+                        using MemoryBufferedPayload payload = new MemoryBufferedPayload(blob);
+                        await _blobStore.PutCompressedObject(ns, payload, ContentId.FromIoHash(entry.AttachmentHash), HttpContext.RequestServices);
+                    }
+                    else
+                    {
+                        await _blobStore.PutObject(ns, blob, BlobIdentifier.FromIoHash(entry.AttachmentHash));
+                    }
+                }
+            }
+            catch (HashMismatchException e)
+            {
+                return BadRequest(new ProblemDetails
+                {
+                    Title = $"Incorrect hash, got hash \"{e.SuppliedHash}\" but hash of content was determined to be \"{e.ContentHash}\""
+                });
+            }
+            
+            CbObject rootObject = packageReader.RootObject;
+            BlobIdentifier rootObjectHash = BlobIdentifier.FromIoHash(packageReader.RootHash);
+
+            (ContentId[] missingReferences, BlobIdentifier[] missingBlobs) = await _objectService.Put(ns, bucket, key, rootObjectHash, rootObject);
+
+            List<JupiterContentHash> missingHashes = new List<JupiterContentHash>(missingReferences);
+            missingHashes.AddRange(missingBlobs);
+            return Ok(new PutObjectResponse(missingHashes.ToArray()));
+        }
+
+        [HttpPost("{ns}/{bucket}/{key}/finalize/{hash}.{format?}")]
+        public async Task<IActionResult> FinalizeObject(
+            [FromRoute] [Required] NamespaceId ns,
+            [FromRoute] [Required] BucketId bucket,
+            [FromRoute] [Required] IoHashKey key,
+            [FromRoute] [Required] BlobIdentifier hash)
+        {
+            ActionResult? accessResult = await _requestHelper.HasAccessToNamespace(User, Request, ns, new [] { JupiterAclAction.WriteObject });
+            if (accessResult != null)
+            {
+                return accessResult;
+            }
+
+            try
+            {
+                (ContentId[] missingReferences, BlobIdentifier[] missingBlobs) = await _objectService.Finalize(ns, bucket, key, hash);
+                List<JupiterContentHash> missingHashes = new List<JupiterContentHash>(missingReferences);
+                missingHashes.AddRange(missingBlobs);
+
+                return Ok(new PutObjectResponse(missingHashes.ToArray()));
+            }
+            catch (ObjectHashMismatchException e)
+            {
+                return BadRequest(e.Message);
+            }
+            catch (ObjectNotFoundException e)
+            {
+                return NotFound(e.Message);
+            }
+        }
+
+        [HttpPost("{ns}")]
+        [Consumes(CustomMediaTypeNames.UnrealCompactBinary)]
+        [Produces(CustomMediaTypeNames.UnrealCompactBinary)]
+        public async Task<IActionResult> Batch(
+            [FromRoute] [Required] NamespaceId ns,
+            [FromBody] [Required] RefBatchOps ops)
+        {
+            JupiterAclAction ActionForOp(RefBatchOps.RefBatchOp.RefOperation op)
+            {
+                switch (op)
+                {
+                    case RefBatchOps.RefBatchOp.RefOperation.GET:
+                        return JupiterAclAction.ReadObject;
+                    case RefBatchOps.RefBatchOp.RefOperation.PUT:
+                        return JupiterAclAction.WriteObject;
+                    case RefBatchOps.RefBatchOp.RefOperation.HEAD:
+                        return JupiterAclAction.ReadObject;
+                    default:
+                        throw new ArgumentOutOfRangeException(nameof(op), op, null);
+                }
+            }
+
+            JupiterAclAction[] requiredActions = ops.Ops.Select(op => ActionForOp(op.Op)).ToArray();
+
+            ActionResult? accessResult = await _requestHelper.HasAccessToNamespace(User, Request, ns, requiredActions);
+            if (accessResult != null)
+            {
+                return accessResult;
+            }
+
+            HashSet<uint> usedOpIds = new HashSet<uint>();
+            foreach (RefBatchOps.RefBatchOp batchOp in ops.Ops)
+            {
+                bool added = usedOpIds.Add(batchOp.OpId);
+                if (!added)
+                {
+                    return BadRequest(new ProblemDetails { Title = $"Duplicate op ids used for id: {batchOp.OpId}" });
+                }
+            }
+            ConcurrentDictionary<uint, (CbObject, HttpStatusCode)> results = new();
+
+            async Task<(CbObject, HttpStatusCode)> BatchGetOp(RefBatchOps.RefBatchOp op)
+            {
+                try
+                {
+                    (ObjectRecord objectRecord, BlobContents? blob) = await _objectService.Get(ns, op.Bucket, op.Key, Array.Empty<string>());
+
+                    if (!objectRecord.IsFinalized)
+                    {
+                        throw new Exception("Object is not finalized");
+                    }
+
+                    if (blob == null)
+                    {
+                        throw new Exception();
+                    }
+
+                    CbObject cb = new CbObject(await blob.Stream.ToByteArray());
+
+                    if (op.ResolveAttachments ?? false)
+                    {
+                        IAsyncEnumerable<BlobIdentifier> references = _referenceResolver.GetReferencedBlobs(ns, cb);
+                        List<BlobIdentifier>? _ = await references.ToListAsync();
+                    }
+
+                    return (cb, HttpStatusCode.OK);
+                }
+                catch (Exception ex) when( ex is ObjectNotFoundException or PartialReferenceResolveException or ReferenceIsMissingBlobsException)
+                {
+                    return ToErrorResult(ex, HttpStatusCode.NotFound);
+                }
+                catch (Exception e)
+                {
+                    return ToErrorResult(e);
+                }
+            }
+
+            async Task<(CbObject, HttpStatusCode)> BatchHeadOp(RefBatchOps.RefBatchOp op)
+            {
+                try
+                {
+                    (ObjectRecord record, BlobContents? blob) = await _objectService.Get(ns, op.Bucket, op.Key, new string[] { "blobIdentifier" });
+
+                    if (!record.IsFinalized)
+                    {
+                        return (CbObject.Build(writer => writer.WriteBool("exists", false)), HttpStatusCode.NotFound);
+                    }
+
+                    blob ??= await _blobStore.GetObject(ns, record.BlobIdentifier);
+
+                    if (op.ResolveAttachments ?? false)
+                    {
+                        byte[] blobContents = await blob.Stream.ToByteArray();
+                        CbObject cb = new CbObject(blobContents);
+                        // the reference resolver will throw if any blob is missing, so no need to do anything other then process each reference
+                        IAsyncEnumerable<BlobIdentifier> references = _referenceResolver.GetReferencedBlobs(ns, cb);
+                        List<BlobIdentifier>? _ = await references.ToListAsync();
+                    }
+
+                    if (blob == null)
+                    {
+                        throw new Exception();
+                    }
+
+                    return (CbObject.Build(writer => writer.WriteBool("exists", true)), HttpStatusCode.OK);
+                }
+                catch (Exception ex) when( ex is ObjectNotFoundException or PartialReferenceResolveException or ReferenceIsMissingBlobsException)
+                {
+                    return (CbObject.Build(writer => writer.WriteBool("exists", false)), HttpStatusCode.NotFound);
+                }
+                catch (Exception e)
+                {
+                    return ToErrorResult(e);
+                }
+            }
+
+            async Task<(CbObject, HttpStatusCode)> BatchPutOp(RefBatchOps.RefBatchOp op)
+            {
+                try
+                {
+                    if (op.Payload == null || op.Payload.Equals(CbObject.Empty))
+                    {
+                        throw new Exception($"Missing payload for operation: {op.OpId}");
+                    }
+
+                    if (op.PayloadHash == null)
+                    {
+                        throw new Exception($"Missing payload hash for operation: {op.OpId}");
+                    }
+                    BlobIdentifier headerHash = BlobIdentifier.FromContentHash(op.PayloadHash);
+                    BlobIdentifier objectHash = BlobIdentifier.FromBlob(op.Payload.GetView().ToArray());
+
+                    if (!headerHash.Equals(objectHash))
+                    {
+                        throw new HashMismatchException(headerHash, objectHash);
+                    }
+
+                    (ContentId[] missingReferences, BlobIdentifier[] missingBlobs) = await _objectService.Put(ns, op.Bucket, op.Key, objectHash, op.Payload);
+                    List<JupiterContentHash> missingHashes = new List<JupiterContentHash>(missingReferences);
+
+                    return (CbSerializer.Serialize(new PutObjectResponse(missingHashes.ToArray())), HttpStatusCode.OK);
+                }
+                catch (Exception e)
+                {
+                    return ToErrorResult(e);
+                }
+            }
+
+            await Parallel.ForEachAsync(ops.Ops, CancellationToken.None, async (op, token) =>
+            {
+                switch (op.Op)
+                {
+                    case RefBatchOps.RefBatchOp.RefOperation.GET:
+                        results.TryAdd(op.OpId, await BatchGetOp(op));
+                        break;
+                    case RefBatchOps.RefBatchOp.RefOperation.PUT:
+                        results.TryAdd(op.OpId, await BatchPutOp(op));
+                        break;
+                    case RefBatchOps.RefBatchOp.RefOperation.HEAD:
+                        results.TryAdd(op.OpId, await BatchHeadOp(op));
+                        break;
+                    case RefBatchOps.RefBatchOp.RefOperation.INVALID:
+                    default:
+                        throw new NotImplementedException($"Unknown op type {op.Op}");
+                }
+
+                await Task.CompletedTask;
+            });
+
+            return Ok(new BatchOpsResponse()
+            {
+                Results = results.Select(result =>
+                {
+                    return new BatchOpsResponse.OpResponses()
+                    {
+                        OpId = result.Key,
+                        Response = result.Value.Item1,
+                        StatusCode = (int)result.Value.Item2
+                    };
+                }).ToList()
+            });
+        }
+
+        private static (CbObject, HttpStatusCode) ToErrorResult(Exception exception, HttpStatusCode statusCode = HttpStatusCode.InternalServerError)
+        {
+            Exception e = exception;
+            CbWriter writer = new CbWriter();
+            writer.BeginObject();
+            writer.WriteString("title", e.Message);
+            writer.WriteInteger("status", (int)statusCode);
+            if (e.StackTrace != null)
+            {
+                writer.WriteString("stackTrace", e.StackTrace);
+            }
+            writer.EndObject();
+            return (writer.ToObject(), statusCode);
+        }
+
+        /// <summary>
+        /// Drop all refs records in the namespace
+        /// </summary>
+        /// <param name="ns">Namespace. Each namespace is completely separated from each other. Use for different types of data that is never expected to be similar (between two different games for instance)</param>
+        [HttpDelete("{ns}", Order = 500)]
+        [ProducesResponseType(204)]
+        public async Task<IActionResult> DeleteNamespace(
+            [FromRoute] [Required] NamespaceId ns
+        )
+        {
+            ActionResult? accessResult = await _requestHelper.HasAccessToNamespace(User, Request, ns, new [] { JupiterAclAction.DeleteNamespace });
+            if (accessResult != null)
+            {
+                return accessResult;
+            }
+
+            try
+            {
+                await _objectService.DropNamespace(ns);
+            }
+            catch (NamespaceNotFoundException e)
+            {
+                return NotFound(new ProblemDetails {Title = $"Namespace {e.Namespace} did not exist"});
+            }
+
+            return NoContent();
+        }
+
+        /// <summary>
+        /// Drop all refs records in the bucket
+        /// </summary>
+        /// <param name="ns">Namespace. Each namespace is completely separated from each other. Use for different types of data that is never expected to be similar (between two different games for instance)</param>
+        /// <param name="bucket">The category/type of record you are caching. Is a clustered key together with the actual key, but all records in the same bucket can be dropped easily.</param>
+        [HttpDelete("{ns}/{bucket}", Order = 500)]
+        [ProducesResponseType(200)]
+        public async Task<IActionResult> DeleteBucket(
+            [FromRoute] [Required] NamespaceId ns,
+            [FromRoute] [Required] BucketId bucket)
+        {
+            ActionResult? accessResult = await _requestHelper.HasAccessToNamespace(User, Request, ns, new [] { JupiterAclAction.DeleteBucket });
+            if (accessResult != null)
+            {
+                return accessResult;
+            }
+
+            long countOfDeletedRecords;
+            try
+            {
+                countOfDeletedRecords = await _objectService.DeleteBucket(ns, bucket);
+            }
+            catch (NamespaceNotFoundException e)
+            {
+                return NotFound(new ProblemDetails {Title = $"Namespace {e.Namespace} did not exist"});
+            }
+
+            return Ok(new BucketDeletedResponse(countOfDeletedRecords));
+        }
+
+        /// <summary>
+        /// Delete a individual refs key
+        /// </summary>
+        /// <param name="ns">Namespace. Each namespace is completely separated from each other. Use for different types of data that is never expected to be similar (between two different games for instance)</param>
+        /// <param name="bucket">The category/type of record you are caching. Is a clustered key together with the actual key, but all records in the same bucket can be dropped easily.</param>
+        /// <param name="key">The unique name of this particular key</param>
+        [HttpDelete("{ns}/{bucket}/{key}", Order = 500)]
+        [ProducesResponseType(200)]
+        [ProducesResponseType(404)]
+        public async Task<IActionResult> Delete(
+            [FromRoute] [Required] NamespaceId ns,
+            [FromRoute] [Required] BucketId bucket,
+            [FromRoute] [Required] IoHashKey key)
+        {
+            ActionResult? accessResult = await _requestHelper.HasAccessToNamespace(User, Request, ns, new [] { JupiterAclAction.DeleteObject });
+            if (accessResult != null)
+            {
+                return accessResult;
+            }
+
+            try
+            {
+                bool deleted = await _objectService.Delete(ns, bucket, key);
+                if (!deleted)
+                {
+                    return NotFound(new ProblemDetails { Title = $"Object {key} in bucket {bucket} and namespace {ns} did not exist" });
+                }
+                return Ok(new RefDeletedResponse(deleted ? 1 : 0));
+            }
+            catch (NamespaceNotFoundException e)
+            {
+                return NotFound(new ProblemDetails { Title = $"Namespace {e.Namespace} did not exist" });
+            }
+            catch (ObjectNotFoundException)
+            {
+                return NotFound(new ProblemDetails { Title = $"Object {key} in bucket {bucket} and namespace {ns} did not exist" });
+            }
+        }
+    }
+
+    public class RefDeletedResponse
+    {
+        public RefDeletedResponse()
+        {
+
+        }
+
+        public RefDeletedResponse(int deletedCount)
+        {
+            DeletedCount = deletedCount;
+        }
+
+        [CbField("deletedCount")]
+        public int DeletedCount { get; set; }
+    }
+
+    public class BucketDeletedResponse
+    {
+        public BucketDeletedResponse()
+        {
+
+        }
+
+        public BucketDeletedResponse(long countOfDeletedRecords)
+        {
+            CountOfDeletedRecords = countOfDeletedRecords;
+        }
+
+        [CbField("countOfDeletedRecords")]
+        public long CountOfDeletedRecords { get; set; }
+    }
+
+    public class RefBatchOps
+    {
+        public RefBatchOps()
+        {
+            Ops = Array.Empty<RefBatchOp>();
+        }
+
+        [System.Diagnostics.CodeAnalysis.SuppressMessage("Design", "CA1034:Nested types should not be visible", Justification = "For serialization only")]
+        public class RefBatchOp
+        {
+            public RefBatchOp()
+            {
+                Payload = null;
+                PayloadHash = null;
+            }
+
+            public enum RefOperation
+            {
+                INVALID,
+                GET,
+                PUT,
+                HEAD,
+            }
+
+            [Required]
+            [CbField("opId")]
+            public uint OpId { get; set; }
+
+            [CbField("op")]
+            [JsonIgnore]
+            [UsedImplicitly]
+            public string OpString
+            {
+                get => Op.ToString();
+                set => Op = Enum.Parse<RefOperation>(value);
+            }
+
+            [Required]
+            public RefOperation Op { get; set; } = RefOperation.INVALID;
+
+            [Required]
+            [CbField("bucket")]
+            public BucketId Bucket { get; set; }
+
+            [Required]
+            [CbField("key")]
+            public IoHashKey Key { get; set; }
+
+            [CbField("resolveAttachments")]
+            public bool? ResolveAttachments { get; set; } = null;
+
+            [CbField("payload")]
+            public CbObject? Payload { get; set; } = null;
+
+            [CbField("payloadHash")] 
+            public JupiterContentHash? PayloadHash { get; set; } = null;
+        }
+
+        [CbField("ops")]
+        public RefBatchOp[] Ops { get; set; }
+    }
+
+    public class BatchOpsResponse
+    {
+        public BatchOpsResponse()
+        {
+
+        }
+
+        [System.Diagnostics.CodeAnalysis.SuppressMessage("Design", "CA1034:Nested types should not be visible", Justification = "For serialization only")]
+        public class OpResponses
+        {
+            public OpResponses()
+            {
+
+            }
+            [CbField("opId")]
+            public uint OpId { get; set; }
+
+            [CbField("response")]
+            public CbObject Response { get; set; } = null!;
+
+            [CbField("statusCode")]
+            public int StatusCode { get; set; }
+        }
+
+        [CbField("results")]
+        [System.Diagnostics.CodeAnalysis.SuppressMessage("Usage", "CA2227:Collection properties should be read only", Justification = "Used by serialization")]
+        public List<OpResponses> Results { get; set; } = new List<OpResponses>();
+    }
+
+    public class RefMetadataResponse
+    {
+        public RefMetadataResponse()
+        {
+            PayloadIdentifier = null!;
+            InlinePayload = null!;
+        }
+
+        [JsonConstructor]
+        public RefMetadataResponse(NamespaceId ns, BucketId bucket, IoHashKey name, BlobIdentifier payloadIdentifier, DateTime lastAccess, bool isFinalized, byte[]? inlinePayload)
+        {
+            Ns = ns;
+            Bucket = bucket;
+            Name = name;
+            PayloadIdentifier = payloadIdentifier;
+            LastAccess = lastAccess;
+            IsFinalized = isFinalized;
+            InlinePayload = inlinePayload;
+        }
+
+        public RefMetadataResponse(ObjectRecord objectRecord)
+        {
+            Ns = objectRecord.Namespace;
+            Bucket = objectRecord.Bucket;
+            Name = objectRecord.Name;
+            PayloadIdentifier = objectRecord.BlobIdentifier;
+            LastAccess = objectRecord.LastAccess;
+            IsFinalized = objectRecord.IsFinalized;
+            InlinePayload = objectRecord.InlinePayload;
+        }
+
+        [CbField("ns")]
+        public NamespaceId Ns { get; set; }
+        
+        [CbField("bucket")]
+        public BucketId Bucket { get; set; }
+
+        [CbField("name")]
+        public IoHashKey Name { get; set; }
+
+        [CbField("payloadIdentifier")]
+        public BlobIdentifier PayloadIdentifier { get; set; }
+
+        [CbField("lastAccess")]
+        public DateTime LastAccess { get; set; }
+
+        [CbField("isFinalized")]
+        public bool IsFinalized { get; set; }
+
+        [CbField("inlinePayload")]
+        public byte[]? InlinePayload { get; set; }
+    }
+
+    public class PutObjectResponse
+    {
+        public PutObjectResponse()
+        {
+            Needs = null!;
+        }
+
+        public PutObjectResponse(JupiterContentHash[] missingReferences)
+        {
+            Needs = missingReferences;
+        }
+
+        [CbField("needs")]
+        public JupiterContentHash[] Needs { get; set; }
+    }
+
+    public class ExistCheckMultipleRefsResponse
+    {
+        public ExistCheckMultipleRefsResponse(List<(BucketId,IoHashKey)> missing)
+        {
+            Missing = missing.Select(pair =>
+            {
+                (BucketId bucketId, IoHashKey ioHashKey) = pair;
+                return new MissingReference()
+                {
+                    Bucket = bucketId,
+                    Key = ioHashKey,
+                };
+            }).ToList();
+        }
+
+        [JsonConstructor]
+        public ExistCheckMultipleRefsResponse(List<MissingReference> missing)
+        {
+            Missing = missing;
+        }
+
+        [CbField("missing")]
+        [System.Diagnostics.CodeAnalysis.SuppressMessage("Usage", "CA2227:Collection properties should be read only", Justification = "Used by serialization")]
+        public List<MissingReference> Missing { get; set; }
+
+        [System.Diagnostics.CodeAnalysis.SuppressMessage("Design", "CA1034:Nested types should not be visible", Justification = "For serialization only")]
+        public class MissingReference
+        {
+            [CbField("bucket")]
+            public BucketId Bucket { get; set; }
+
+            [CbField("key")]
+            public IoHashKey Key { get; set; }
+        }
+    }
+
+    public class GetNamespacesResponse
+    {
+        [JsonConstructor]
+        public GetNamespacesResponse(NamespaceId[] namespaces)
+        {
+            Namespaces = namespaces;
+        }
+
+        public NamespaceId[] Namespaces { get; set; }
+    }
+
+    public class HashMismatchException : Exception
+    {
+        public JupiterContentHash SuppliedHash { get; }
+        public JupiterContentHash ContentHash { get; }
+
+        public HashMismatchException(JupiterContentHash suppliedHash, JupiterContentHash contentHash) : base($"ID was not a hash of the content uploaded. Supplied hash was: {suppliedHash} but hash of content was {contentHash}")
+        {
+            SuppliedHash = suppliedHash;
+            ContentHash = contentHash;
+        }
+    }
+}
diff --git a/Engine/Source/Programs/Horde/Horde.Server/Ddc/ReferencesStore.cs b/Engine/Source/Programs/Horde/Horde.Server/Ddc/ReferencesStore.cs
new file mode 100644
--- /dev/null
+++ b/Engine/Source/Programs/Horde/Horde.Server/Ddc/ReferencesStore.cs
@@ -0,0 +1,58 @@
+// Copyright Epic Games, Inc. All Rights Reserved.
+
+using System;
+using System.Collections.Generic;
+using System.Threading.Tasks;
+using EpicGames.Horde.Storage;
+using Jupiter.Implementation;
+
+namespace Horde.Server.Ddc
+{
+	class ReferencesStore : IReferencesStore
+	{
+		public Task<bool> Delete(NamespaceId ns, BucketId bucket, IoHashKey key)
+		{
+			throw new NotImplementedException();
+		}
+
+		public Task<long> DeleteBucket(NamespaceId ns, BucketId bucket)
+		{
+			throw new NotImplementedException();
+		}
+
+		public Task<long> DropNamespace(NamespaceId ns)
+		{
+			throw new NotImplementedException();
+		}
+
+		public Task Finalize(NamespaceId ns, BucketId bucket, IoHashKey key, BlobIdentifier blobIdentifier)
+		{
+			throw new NotImplementedException();
+		}
+
+		public Task<ObjectRecord> Get(NamespaceId ns, BucketId bucket, IoHashKey key, IReferencesStore.FieldFlags flags)
+		{
+			throw new NotImplementedException();
+		}
+
+		public IAsyncEnumerable<NamespaceId> GetNamespaces()
+		{
+			throw new NotImplementedException();
+		}
+
+		public IAsyncEnumerable<(NamespaceId, BucketId, IoHashKey, DateTime)> GetRecords()
+		{
+			throw new NotImplementedException();
+		}
+
+		public Task Put(NamespaceId ns, BucketId bucket, IoHashKey key, BlobIdentifier blobHash, byte[] blob, bool isFinalized)
+		{
+			throw new NotImplementedException();
+		}
+
+		public Task UpdateLastAccessTime(NamespaceId ns, BucketId bucket, IoHashKey key, DateTime newLastAccessTime)
+		{
+			throw new NotImplementedException();
+		}
+	}
+}
diff --git a/Engine/Source/Programs/Horde/Horde.Server/Ddc/RequestHelper.cs b/Engine/Source/Programs/Horde/Horde.Server/Ddc/RequestHelper.cs
new file mode 100644
--- /dev/null
+++ b/Engine/Source/Programs/Horde/Horde.Server/Ddc/RequestHelper.cs
@@ -0,0 +1,26 @@
+// Copyright Epic Games, Inc. All Rights Reserved.
+
+using System;
+using System.Security.Claims;
+using System.Threading.Tasks;
+using EpicGames.Horde.Storage;
+using Jupiter;
+using Jupiter.Controllers;
+using Microsoft.AspNetCore.Http;
+using Microsoft.AspNetCore.Mvc;
+
+namespace Horde.Server.Ddc
+{
+	class RequestHelper : IRequestHelper
+	{
+		public Task<ActionResult?> HasAccessForGlobalOperations(ClaimsPrincipal user, JupiterAclAction[] aclActions)
+		{
+			throw new NotImplementedException();
+		}
+
+		public Task<ActionResult?> HasAccessToNamespace(ClaimsPrincipal user, HttpRequest request, NamespaceId ns, JupiterAclAction[] aclActions)
+		{
+			throw new NotImplementedException();
+		}
+	}
+}
diff --git a/Engine/Source/Programs/Horde/Horde.Server/Ddc/RequestUtil.cs b/Engine/Source/Programs/Horde/Horde.Server/Ddc/RequestUtil.cs
new file mode 100644
--- /dev/null
+++ b/Engine/Source/Programs/Horde/Horde.Server/Ddc/RequestUtil.cs
@@ -0,0 +1,42 @@
+﻿// Copyright Epic Games, Inc. All Rights Reserved.
+
+using System;
+using System.IO;
+using System.Threading.Tasks;
+using Microsoft.AspNetCore.Http;
+using Microsoft.Extensions.DependencyInjection;
+using OpenTelemetry.Trace;
+
+#pragma warning disable CS1591
+
+namespace Jupiter
+{
+    public static class RequestUtil
+    {
+        public static async Task<byte[]> ReadRawBody(HttpRequest request)
+        {
+            long? contentLength = request.GetTypedHeaders().ContentLength;
+
+            if (contentLength == null)
+            {
+                throw new Exception("Expected content-length on all raw body requests");
+            }
+
+            Tracer? tracer = request.HttpContext.RequestServices.GetService<Tracer>();
+            using TelemetrySpan? scope = tracer?.StartActiveSpan("readbody")
+                .SetAttribute("operation.name", "readbody")
+                .SetAttribute("Content-Length", contentLength.Value);
+
+            await using MemoryStream ms = new MemoryStream((int)contentLength);
+            DateTime readStart = DateTime.Now;
+            await request.Body.CopyToAsync(ms);
+            TimeSpan duration = DateTime.Now - readStart;
+
+#if WITH_DOGSTATSD
+			double rate = contentLength.Value / duration.TotalSeconds;
+            StatsdClient.DogStatsd.Histogram("jupiter.stream_throughput", rate, tags: new string[] {"sourceIdentifier:" + "readbody"});
+#endif
+            return ms.GetBuffer();
+        }
+    }
+}
diff --git a/Engine/Source/Programs/Horde/Horde.Server/Ddc/RequiredContentTypeAttribute.cs b/Engine/Source/Programs/Horde/Horde.Server/Ddc/RequiredContentTypeAttribute.cs
new file mode 100644
--- /dev/null
+++ b/Engine/Source/Programs/Horde/Horde.Server/Ddc/RequiredContentTypeAttribute.cs
@@ -0,0 +1,34 @@
+﻿// Copyright Epic Games, Inc. All Rights Reserved.
+
+using System;
+using System.Linq;
+using Microsoft.AspNetCore.Mvc.ActionConstraints;
+using Microsoft.Extensions.Primitives;
+
+#pragma warning disable CS1591
+
+namespace Jupiter
+{
+    [AttributeUsage(AttributeTargets.Method)]
+    public sealed class RequiredContentTypeAttribute : Attribute, IActionConstraint
+    {
+        private readonly string _mediaTypeName;
+
+        public RequiredContentTypeAttribute(string mediaTypeName)
+        {
+            _mediaTypeName = mediaTypeName;
+        }
+
+        public string MediaTypeName => _mediaTypeName;
+
+        public int Order => 0;
+
+        public bool Accept(ActionConstraintContext context)
+        {
+            StringValues contentTypeHeader = context.RouteContext.HttpContext.Request.Headers["Content-Type"];
+
+            bool valid = contentTypeHeader.ToList().Contains(_mediaTypeName, StringComparer.InvariantCultureIgnoreCase);
+            return valid;
+        }
+    }
+}
diff --git a/Engine/Source/Programs/Horde/Horde.Server/Ddc/Sha1Utils.cs b/Engine/Source/Programs/Horde/Horde.Server/Ddc/Sha1Utils.cs
new file mode 100644
--- /dev/null
+++ b/Engine/Source/Programs/Horde/Horde.Server/Ddc/Sha1Utils.cs
@@ -0,0 +1,77 @@
+﻿// Copyright Epic Games, Inc. All Rights Reserved.
+
+using System.Collections.Generic;
+using System.Security.Cryptography;
+using System.Threading.Tasks;
+
+#pragma warning disable CS1591
+
+namespace Jupiter
+{
+    public static class Sha1Utils
+    {
+        public static async Task<string> GetSHA1AsStringAsync(byte[] payload)
+        {
+            return await Task.Run(() => GetSHA1AsString(payload));
+        }
+
+        public static byte[] GetSHA1(byte[] payload)
+        {
+            using SHA1 sha1 = SHA1.Create();
+            return sha1.ComputeHash(payload);
+        }
+
+        public static byte[] GetSHA1(byte[] payload, int offset, int count)
+        {
+            using SHA1 sha1 = SHA1.Create();
+            return sha1.ComputeHash(payload, offset, count);
+        }
+
+        public static string GetSHA1AsString(byte[] payload)
+        {
+            byte[] sha1 = GetSHA1(payload);
+
+            return StringUtils.FormatAsHexString(sha1);
+        }
+    }
+
+    public class ByteArrayComparer : EqualityComparer<byte[]>
+    {
+        public override bool Equals(byte[]? a, byte[]? b)
+        {
+            if (a == null)
+            {
+                return b == null;
+            }
+
+            if (b == null)
+            {
+                return false;
+            }
+
+            if (a.Length != b.Length)
+            {
+                return false;
+            }
+
+            for (int i = 0; i < a.Length; i++)
+            {
+                if (a[i] != b[i])
+                {
+                    return false;
+                }
+            }
+            return true;
+        }
+        public override int GetHashCode(byte[] a)
+        {
+            uint b = 0;
+            foreach (byte t in a)
+            {
+	            b = ((b << 23) | (b >> 9)) ^ t;
+            }
+
+            return unchecked((int)b);
+        }
+    }
+}
diff --git a/Engine/Source/Programs/Horde/Horde.Server/Ddc/StreamUtils.cs b/Engine/Source/Programs/Horde/Horde.Server/Ddc/StreamUtils.cs
new file mode 100644
--- /dev/null
+++ b/Engine/Source/Programs/Horde/Horde.Server/Ddc/StreamUtils.cs
@@ -0,0 +1,29 @@
+﻿// Copyright Epic Games, Inc. All Rights Reserved.
+
+using System.IO;
+using System.Threading.Tasks;
+using EpicGames.AspNet;
+using Microsoft.AspNetCore.Http;
+
+#pragma warning disable CS1591
+
+namespace Jupiter.Utils
+{
+    public static class StreamUtils
+    {
+        public static async Task<byte[]> ToByteArray(this Stream s)
+        {
+            try
+            {
+                await using MemoryStream ms = new MemoryStream();
+                await s.CopyToAsync(ms);
+                return ms.ToArray();
+            }
+            catch (BadHttpRequestException e)
+            {
+                ClientSendSlowExceptionUtil.MaybeThrowSlowSendException(e);
+                throw;
+            }
+        }
+    }
+}
diff --git a/Engine/Source/Programs/Horde/Horde.Server/Ddc/StreamingBodyPayload.cs b/Engine/Source/Programs/Horde/Horde.Server/Ddc/StreamingBodyPayload.cs
new file mode 100644
--- /dev/null
+++ b/Engine/Source/Programs/Horde/Horde.Server/Ddc/StreamingBodyPayload.cs
@@ -0,0 +1,150 @@
+// Copyright Epic Games, Inc. All Rights Reserved.
+
+using System;
+using System.IO;
+using System.Threading.Tasks;
+using Jupiter.Utils;
+using Microsoft.AspNetCore.Http;
+using Microsoft.Extensions.Options;
+using OpenTelemetry.Trace;
+
+#pragma warning disable CS1591
+
+namespace Jupiter.Common.Implementation
+{
+    public interface IBufferedPayload : IDisposable
+    {
+        Stream GetStream();
+        long Length { get; }
+    }
+
+    /// <summary>
+    /// Streaming request that is streamed into memory
+    /// </summary>
+    public sealed class MemoryBufferedPayload : IBufferedPayload
+    {
+        private readonly byte[] _buffer;
+
+        public MemoryBufferedPayload(byte[] source)
+        {
+            _buffer = source;
+        }
+
+        public static async Task<MemoryBufferedPayload> Create(Tracer tracer, Stream s)
+        {
+            using TelemetrySpan scope = tracer.StartActiveSpan("payload.buffer")
+                .SetAttribute("operation.name", "payload.buffer")
+                .SetAttribute("bufferType", "Memory");
+            MemoryBufferedPayload payload = new MemoryBufferedPayload(await s.ToByteArray());
+
+            return payload;
+        }
+
+        public void Dispose()
+        {
+            
+        }
+
+        public Stream GetStream()
+        {
+            return new MemoryStream(_buffer);
+        }
+
+        public long Length => _buffer.LongLength;
+    }
+
+    /// <summary>
+    /// A streaming request backed by a temporary file on disk
+    /// </summary>
+    public sealed class FilesystemBufferedPayload : IBufferedPayload
+    {
+        private readonly FileInfo _tempFile;
+        private long _length;
+
+        private FilesystemBufferedPayload()
+        {
+            _tempFile = new FileInfo(Path.GetTempFileName());
+        }
+
+        public static async Task<FilesystemBufferedPayload> Create(Tracer tracer, Stream s)
+        {
+            FilesystemBufferedPayload payload = new FilesystemBufferedPayload();
+
+            {
+                using TelemetrySpan? scope = tracer.StartActiveSpan("payload.buffer")
+                    .SetAttribute("operation.name", "payload.buffer")
+                    .SetAttribute("bufferType", "Filesystem");
+                await using FileStream fs = payload._tempFile.OpenWrite();
+                await s.CopyToAsync(fs);
+            }
+            
+            payload._tempFile.Refresh();
+            payload._length = payload._tempFile.Length;
+
+            return payload;
+        }
+
+        public void Dispose()
+        {
+            if (_tempFile.Exists)
+            {
+                _tempFile.Delete();
+            }
+        }
+
+        public Stream GetStream()
+        {
+            return _tempFile.OpenRead();
+        }
+
+        public long Length => _length;
+    }
+
+    public class BufferedPayloadOptions
+    {
+        /// <summary>
+        /// If the request is smaller then MemoryBufferSize we buffer it in memory rather then as a file
+        /// </summary>
+        public long MemoryBufferSize { get; set; } = int.MaxValue;
+    }
+
+    public class BufferedPayloadFactory
+    {
+        private readonly IOptionsMonitor<BufferedPayloadOptions> _options;
+        private readonly Tracer _tracer;
+
+        public BufferedPayloadFactory(IOptionsMonitor<BufferedPayloadOptions> options, Tracer tracer)
+        {
+            _options = options;
+            _tracer = tracer;
+        }
+
+        public Task<IBufferedPayload> CreateFromRequest(HttpRequest request)
+        {
+            long? contentLength = request.ContentLength;
+
+            if (contentLength == null)
+            {
+                throw new Exception("Expected content-length on all requests");
+            }
+
+            return CreateFromStream(request.Body, contentLength.Value);
+        }
+
+        public async Task<IBufferedPayload> CreateFromStream(Stream s, long contentLength)
+        {
+            // blob is small enough to fit into memory we just read it as is
+            if (contentLength < _options.CurrentValue.MemoryBufferSize)
+            {
+                return await MemoryBufferedPayload.Create(_tracer, s);
+            }
+
+            return await FilesystemBufferedPayload.Create(_tracer, s);
+        }
+
+        public async Task<IBufferedPayload> CreateFilesystemBufferedPayload(Stream s)
+        {
+            return await FilesystemBufferedPayload.Create(_tracer, s);
+        }
+    }
+}
diff --git a/Engine/Source/Programs/Horde/Horde.Server/Ddc/StringUtils.cs b/Engine/Source/Programs/Horde/Horde.Server/Ddc/StringUtils.cs
new file mode 100644
--- /dev/null
+++ b/Engine/Source/Programs/Horde/Horde.Server/Ddc/StringUtils.cs
@@ -0,0 +1,46 @@
+// Copyright Epic Games, Inc. All Rights Reserved.
+
+using System;
+using System.Text;
+
+#pragma warning disable CS1591
+
+namespace Jupiter
+{
+    public static class StringUtils
+    {
+        public static int GetHexVal(char hex)
+        {
+            return hex - (hex < 58 ? 48 : hex < 97 ? 55 : 87);
+        }
+
+        public static string FormatAsHexString(byte[] data)
+        {
+            StringBuilder sb = new StringBuilder();
+            foreach (byte b in data)
+            {
+                string hex = b.ToString("X2");
+                sb.Append(hex);
+            }
+
+            return sb.ToString();
+        }
+
+        public static byte[] ToHashFromHexString(string hex)
+        {
+            if (hex.Length % 2 == 1)
+            {
+                throw new Exception("The binary key cannot have an odd number of digits");
+            }
+
+            byte[] arr = new byte[hex.Length >> 1];
+
+            for (int i = 0; i < hex.Length >> 1; ++i)
+            {
+                arr[i] = (byte) ((StringUtils.GetHexVal(hex[i << 1]) << 4) + StringUtils.GetHexVal(hex[(i << 1) + 1]));
+            }
+
+            return arr;
+        }
+    }
+}
diff --git a/Engine/Source/Programs/Horde/Horde.Server/Startup.cs b/Engine/Source/Programs/Horde/Horde.Server/Startup.cs
--- a/Engine/Source/Programs/Horde/Horde.Server/Startup.cs
+++ b/Engine/Source/Programs/Horde/Horde.Server/Startup.cs
@@ -104,9 +104,15 @@ using Microsoft.AspNetCore.Diagnostics;
 using EpicGames.Horde;
 using EpicGames.Horde.Api;
 using EpicGames.Horde.Storage.Bundles;
+using Jupiter.Implementation;
+using Horde.Server.Ddc;
 
 namespace Horde.Server
 {
+	using ReferenceResolver = Jupiter.Implementation.ReferenceResolver;
+	using IStorageBackend = Horde.Server.Storage.IStorageBackend;
+	using ContentHash = EpicGames.Core.ContentHash;
+
 	class Startup
 	{
 		static Startup()
@@ -515,6 +521,14 @@ namespace Horde.Server
 				services.AddSingleton<IExternalIssueService, DefaultExternalIssueService>();
 			}
 
+			// DDC
+			services.AddSingleton<IBlobStore, BlobStore>();
+			services.AddSingleton<IBlobService, BlobService>();
+			services.AddSingleton<IObjectService, ObjectService>();
+			services.AddSingleton<IReferencesStore, ReferencesStore>();
+			services.AddSingleton<IReferenceResolver, ReferenceResolver>();
+			services.AddSingleton<IContentIdStore, ContentIdStore>();
+
 			// Secret providers
 			services.AddSingleton<SecretService>();
 
