Horde: Move all DDC classes into the Horde.Server.Ddc namespace.

diff --git a/Engine/Source/Programs/Horde/Horde.Server/Ddc/Acl.cs b/Engine/Source/Programs/Horde/Horde.Server/Ddc/Acl.cs
--- a/Engine/Source/Programs/Horde/Horde.Server/Ddc/Acl.cs
+++ b/Engine/Source/Programs/Horde/Horde.Server/Ddc/Acl.cs
@@ -7,106 +7,107 @@ using Microsoft.AspNetCore.Authorization;
 
 #pragma warning disable CS1591
 
-namespace Jupiter;
-
-public enum JupiterAclAction
+namespace Horde.Server.Ddc
 {
-    /// <summary>
-    /// General read access to refs / blobs and so on
-    /// </summary>
-    ReadObject,
-    /// <summary>
-    /// General write access to upload refs / blobs etc
-    /// </summary>
-    WriteObject,
-    /// <summary>
-    /// Access to delete blobs / refs etc
-    /// </summary>
-    DeleteObject,
+	public enum JupiterAclAction
+	{
+		/// <summary>
+		/// General read access to refs / blobs and so on
+		/// </summary>
+		ReadObject,
+		/// <summary>
+		/// General write access to upload refs / blobs etc
+		/// </summary>
+		WriteObject,
+		/// <summary>
+		/// Access to delete blobs / refs etc
+		/// </summary>
+		DeleteObject,
 
-    /// <summary>
-    /// Access to delete a particular bucket
-    /// </summary>
-    DeleteBucket,
-    /// <summary>
-    /// Access to delete a whole namespace
-    /// </summary>
-    DeleteNamespace,
+		/// <summary>
+		/// Access to delete a particular bucket
+		/// </summary>
+		DeleteBucket,
+		/// <summary>
+		/// Access to delete a whole namespace
+		/// </summary>
+		DeleteNamespace,
 
-    /// <summary>
-    /// Access to read the transaction log
-    /// </summary>
-    ReadTransactionLog,
+		/// <summary>
+		/// Access to read the transaction log
+		/// </summary>
+		ReadTransactionLog,
 
-    /// <summary>
-    /// Access to write the transaction log
-    /// </summary>
-    WriteTransactionLog,
+		/// <summary>
+		/// Access to write the transaction log
+		/// </summary>
+		WriteTransactionLog,
 
-    /// <summary>
-    /// Access to perform administrative task
-    /// </summary>
-    AdminAction
-}
+		/// <summary>
+		/// Access to perform administrative task
+		/// </summary>
+		AdminAction
+	}
 
-public class AclEntry
-{
-    /// <summary>
-    /// Claims required to be present to be allowed to do the actions - if multiple claims are present *all* of them are required
-    /// </summary>
-    [System.Diagnostics.CodeAnalysis.SuppressMessage("Usage", "CA2227:Collection properties should be read only", Justification = "Used by the configuration system")]
-    public List<string> Claims { get; set; } = new List<string>();
+	public class AclEntry
+	{
+		/// <summary>
+		/// Claims required to be present to be allowed to do the actions - if multiple claims are present *all* of them are required
+		/// </summary>
+		[System.Diagnostics.CodeAnalysis.SuppressMessage("Usage", "CA2227:Collection properties should be read only", Justification = "Used by the configuration system")]
+		public List<string> Claims { get; set; } = new List<string>();
 
-    /// <summary>
-    /// The actions granted if the claims match
-    /// </summary>
-    [System.Diagnostics.CodeAnalysis.SuppressMessage("Usage", "CA2227:Collection properties should be read only", Justification = "Used by the configuration system")]
-    public List<JupiterAclAction> Actions { get; set; } = new List<JupiterAclAction>();
+		/// <summary>
+		/// The actions granted if the claims match
+		/// </summary>
+		[System.Diagnostics.CodeAnalysis.SuppressMessage("Usage", "CA2227:Collection properties should be read only", Justification = "Used by the configuration system")]
+		public List<JupiterAclAction> Actions { get; set; } = new List<JupiterAclAction>();
 
-    public IEnumerable<JupiterAclAction> Resolve(ClaimsPrincipal user)
-    {
-        List<JupiterAclAction> allowedActions = new List<JupiterAclAction>();
+		public IEnumerable<JupiterAclAction> Resolve(ClaimsPrincipal user)
+		{
+			List<JupiterAclAction> allowedActions = new List<JupiterAclAction>();
 
-        bool allClaimsFound = true;
-        // These are ANDed, e.g. all claims needs to be present
-        foreach (string expectedClaim in Claims)
-        {
-            bool claimFound = false;
-            // if expected claim is * then everyone has the associated actions
-            if (expectedClaim == "*")
-            {
-                claimFound = true;
-            }
-            else if (expectedClaim.Contains('=', StringComparison.InvariantCultureIgnoreCase))
-            {
-                int separatorIndex = expectedClaim.IndexOf('=', StringComparison.InvariantCultureIgnoreCase);
-                string claimName = expectedClaim.Substring(0, separatorIndex);
-                string claimValue = expectedClaim.Substring(separatorIndex + 1);
-                if (user.HasClaim(claim => string.Equals(claim.Type, claimName, StringComparison.OrdinalIgnoreCase) && string.Equals(claim.Value, claimValue, StringComparison.OrdinalIgnoreCase)))
-                {
-                    claimFound = true;
-                }
-            }
-            else if (user.HasClaim(claim => string.Equals(claim.Type, expectedClaim, StringComparison.OrdinalIgnoreCase)))
-            {
-                claimFound = true;
-            }
+			bool allClaimsFound = true;
+			// These are ANDed, e.g. all claims needs to be present
+			foreach (string expectedClaim in Claims)
+			{
+				bool claimFound = false;
+				// if expected claim is * then everyone has the associated actions
+				if (expectedClaim == "*")
+				{
+					claimFound = true;
+				}
+				else if (expectedClaim.Contains('=', StringComparison.InvariantCultureIgnoreCase))
+				{
+					int separatorIndex = expectedClaim.IndexOf('=', StringComparison.InvariantCultureIgnoreCase);
+					string claimName = expectedClaim.Substring(0, separatorIndex);
+					string claimValue = expectedClaim.Substring(separatorIndex + 1);
+					if (user.HasClaim(claim => string.Equals(claim.Type, claimName, StringComparison.OrdinalIgnoreCase) && string.Equals(claim.Value, claimValue, StringComparison.OrdinalIgnoreCase)))
+					{
+						claimFound = true;
+					}
+				}
+				else if (user.HasClaim(claim => string.Equals(claim.Type, expectedClaim, StringComparison.OrdinalIgnoreCase)))
+				{
+					claimFound = true;
+				}
 
-            if (!claimFound)
-            {
-                allClaimsFound = false;
-            }
-        }
+				if (!claimFound)
+				{
+					allClaimsFound = false;
+				}
+			}
 
-        if (allClaimsFound)
-        {
-            allowedActions.AddRange(Actions);
-        }
-        return allowedActions;
-    }
+			if (allClaimsFound)
+			{
+				allowedActions.AddRange(Actions);
+			}
+			return allowedActions;
+		}
 
-    public IEnumerable<JupiterAclAction> Resolve(AuthorizationHandlerContext context)
-    {
-        return Resolve(context.User);
-    }
-}
\ No newline at end of file
+		public IEnumerable<JupiterAclAction> Resolve(AuthorizationHandlerContext context)
+		{
+			return Resolve(context.User);
+		}
+	}
+}
diff --git a/Engine/Source/Programs/Horde/Horde.Server/Ddc/BlobContents.cs b/Engine/Source/Programs/Horde/Horde.Server/Ddc/BlobContents.cs
--- a/Engine/Source/Programs/Horde/Horde.Server/Ddc/BlobContents.cs
+++ b/Engine/Source/Programs/Horde/Horde.Server/Ddc/BlobContents.cs
@@ -6,9 +6,9 @@ using System.Threading.Tasks;
 
 #pragma warning disable CS1591
 
-namespace Jupiter.Implementation
+namespace Horde.Server.Ddc
 {
-    public sealed class BlobContents : IDisposable, IAsyncDisposable
+	public sealed class BlobContents : IDisposable, IAsyncDisposable
     {
         private readonly Stream? _stream;
 
diff --git a/Engine/Source/Programs/Horde/Horde.Server/Ddc/BlobController.cs b/Engine/Source/Programs/Horde/Horde.Server/Ddc/BlobController.cs
--- a/Engine/Source/Programs/Horde/Horde.Server/Ddc/BlobController.cs
+++ b/Engine/Source/Programs/Horde/Horde.Server/Ddc/BlobController.cs
@@ -13,18 +13,15 @@ using System.Threading.Tasks;
 using EpicGames.AspNet;
 using EpicGames.Horde.Storage;
 using EpicGames.Serialization;
-using Jupiter.Implementation;
-using Jupiter.Common.Implementation;
 using Microsoft.AspNetCore.Authorization;
 using Microsoft.AspNetCore.Http;
 using Microsoft.AspNetCore.Mvc;
 
 #pragma warning disable CS1591
 
-namespace Jupiter.Controllers
+namespace Horde.Server.Ddc
 {
-    using IDiagnosticContext = Serilog.IDiagnosticContext;
-    using BlobNotFoundException = Jupiter.Implementation.BlobNotFoundException;
+	using IDiagnosticContext = Serilog.IDiagnosticContext;
 
     [ApiController]
     [Route("api/v1/s", Order = 1)]
diff --git a/Engine/Source/Programs/Horde/Horde.Server/Ddc/BlobIdentifier.cs b/Engine/Source/Programs/Horde/Horde.Server/Ddc/BlobIdentifier.cs
--- a/Engine/Source/Programs/Horde/Horde.Server/Ddc/BlobIdentifier.cs
+++ b/Engine/Source/Programs/Horde/Horde.Server/Ddc/BlobIdentifier.cs
@@ -16,9 +16,9 @@ using EpicGames.Serialization;
 
 #pragma warning disable CS1591
 
-namespace Jupiter.Implementation
+namespace Horde.Server.Ddc
 {
-    [TypeConverter(typeof(BlobIdentifierTypeConverter))]
+	[TypeConverter(typeof(BlobIdentifierTypeConverter))]
     [JsonConverter(typeof(BlobIdentifierJsonConverter))]
     [CbConverter(typeof(BlobIdentifierCbConverter))]
     public class BlobIdentifier : JupiterContentHash,  IEquatable<BlobIdentifier>
diff --git a/Engine/Source/Programs/Horde/Horde.Server/Ddc/BlobService.cs b/Engine/Source/Programs/Horde/Horde.Server/Ddc/BlobService.cs
--- a/Engine/Source/Programs/Horde/Horde.Server/Ddc/BlobService.cs
+++ b/Engine/Source/Programs/Horde/Horde.Server/Ddc/BlobService.cs
@@ -5,8 +5,6 @@ using System.Collections.Generic;
 using System.IO;
 using System.Threading.Tasks;
 using EpicGames.Horde.Storage;
-using Jupiter.Common.Implementation;
-using Jupiter.Implementation;
 
 namespace Horde.Server.Ddc
 {
diff --git a/Engine/Source/Programs/Horde/Horde.Server/Ddc/BlobStore.cs b/Engine/Source/Programs/Horde/Horde.Server/Ddc/BlobStore.cs
--- a/Engine/Source/Programs/Horde/Horde.Server/Ddc/BlobStore.cs
+++ b/Engine/Source/Programs/Horde/Horde.Server/Ddc/BlobStore.cs
@@ -5,7 +5,6 @@ using System.Collections.Generic;
 using System.IO;
 using System.Threading.Tasks;
 using EpicGames.Horde.Storage;
-using Jupiter.Implementation;
 
 namespace Horde.Server.Ddc
 {
diff --git a/Engine/Source/Programs/Horde/Horde.Server/Ddc/BucketId.cs b/Engine/Source/Programs/Horde/Horde.Server/Ddc/BucketId.cs
--- a/Engine/Source/Programs/Horde/Horde.Server/Ddc/BucketId.cs
+++ b/Engine/Source/Programs/Horde/Horde.Server/Ddc/BucketId.cs
@@ -6,9 +6,10 @@ using System.Globalization;
 using System.Text.Json;
 using System.Text.Json.Serialization;
 using EpicGames.Core;
+using EpicGames.Horde;
 using EpicGames.Serialization;
 
-namespace EpicGames.Horde.Storage
+namespace Horde.Server.Ddc
 {
 	/// <summary>
 	/// Identifier for a storage bucket
diff --git a/Engine/Source/Programs/Horde/Horde.Server/Ddc/CommonHeaders.cs b/Engine/Source/Programs/Horde/Horde.Server/Ddc/CommonHeaders.cs
--- a/Engine/Source/Programs/Horde/Horde.Server/Ddc/CommonHeaders.cs
+++ b/Engine/Source/Programs/Horde/Horde.Server/Ddc/CommonHeaders.cs
@@ -2,9 +2,9 @@
 
 #pragma warning disable CS1591
 
-namespace Jupiter.Implementation
+namespace Horde.Server.Ddc
 {
-    public static class CommonHeaders
+	public static class CommonHeaders
     {
         public const string HashHeaderSHA1Name = "X-Jupiter-Sha1";
         public const string HashHeaderName = "X-Jupiter-IoHash";
diff --git a/Engine/Source/Programs/Horde/Horde.Server/Ddc/CompressedBufferUtils.cs b/Engine/Source/Programs/Horde/Horde.Server/Ddc/CompressedBufferUtils.cs
--- a/Engine/Source/Programs/Horde/Horde.Server/Ddc/CompressedBufferUtils.cs
+++ b/Engine/Source/Programs/Horde/Horde.Server/Ddc/CompressedBufferUtils.cs
@@ -13,9 +13,9 @@ using OpenTelemetry.Trace;
 
 #pragma warning disable CS1591
 
-namespace Jupiter.Implementation
+namespace Horde.Server.Ddc
 {
-    public class CompressedBufferUtils
+	public class CompressedBufferUtils
     {
         private readonly Tracer _tracer;
 
diff --git a/Engine/Source/Programs/Horde/Horde.Server/Ddc/ContentId.cs b/Engine/Source/Programs/Horde/Horde.Server/Ddc/ContentId.cs
--- a/Engine/Source/Programs/Horde/Horde.Server/Ddc/ContentId.cs
+++ b/Engine/Source/Programs/Horde/Horde.Server/Ddc/ContentId.cs
@@ -8,9 +8,9 @@ using EpicGames.Core;
 
 #pragma warning disable CS1591
 
-namespace Jupiter.Implementation
+namespace Horde.Server.Ddc
 {
-    [TypeConverter(typeof(ContentIdTypeConverter))]
+	[TypeConverter(typeof(ContentIdTypeConverter))]
     public class ContentId : JupiterContentHash, IEquatable<ContentId>
     {
         public ContentId(byte[] identifier) : base(identifier)
diff --git a/Engine/Source/Programs/Horde/Horde.Server/Ddc/ContentIdController.cs b/Engine/Source/Programs/Horde/Horde.Server/Ddc/ContentIdController.cs
--- a/Engine/Source/Programs/Horde/Horde.Server/Ddc/ContentIdController.cs
+++ b/Engine/Source/Programs/Horde/Horde.Server/Ddc/ContentIdController.cs
@@ -4,17 +4,14 @@ using System.Net.Mime;
 using System.Threading.Tasks;
 using EpicGames.AspNet;
 using EpicGames.Horde.Storage;
-using Jupiter.Implementation;
 using Microsoft.AspNetCore.Authorization;
 using Microsoft.AspNetCore.Mvc;
 
 #pragma warning disable CS1591
 
-using ContentId = Jupiter.Implementation.ContentId;
-
-namespace Jupiter.Controllers
+namespace Horde.Server.Ddc
 {
-    [ApiController]
+	[ApiController]
     [FormatFilter]
     [Produces(MediaTypeNames.Application.Json, MediaTypeNames.Application.Octet, CustomMediaTypeNames.UnrealCompactBinary)]
     [Route("api/v1/content-id")]
diff --git a/Engine/Source/Programs/Horde/Horde.Server/Ddc/ContentIdStore.cs b/Engine/Source/Programs/Horde/Horde.Server/Ddc/ContentIdStore.cs
--- a/Engine/Source/Programs/Horde/Horde.Server/Ddc/ContentIdStore.cs
+++ b/Engine/Source/Programs/Horde/Horde.Server/Ddc/ContentIdStore.cs
@@ -3,7 +3,6 @@
 using EpicGames.Horde.Storage;
 using System;
 using System.Threading.Tasks;
-using Jupiter.Implementation;
 
 namespace Horde.Server.Ddc
 {
diff --git a/Engine/Source/Programs/Horde/Horde.Server/Ddc/FormatResolver.cs b/Engine/Source/Programs/Horde/Horde.Server/Ddc/FormatResolver.cs
--- a/Engine/Source/Programs/Horde/Horde.Server/Ddc/FormatResolver.cs
+++ b/Engine/Source/Programs/Horde/Horde.Server/Ddc/FormatResolver.cs
@@ -11,9 +11,9 @@ using Microsoft.Extensions.Primitives;
 
 #pragma warning disable CS1591
 
-namespace Jupiter.Controllers
+namespace Horde.Server.Ddc
 {
-    public class FormatResolver
+	public class FormatResolver
     {
         private readonly IOptionsMonitor<MvcOptions> _mvcOptions;
 
diff --git a/Engine/Source/Programs/Horde/Horde.Server/Ddc/IBlobService.cs b/Engine/Source/Programs/Horde/Horde.Server/Ddc/IBlobService.cs
--- a/Engine/Source/Programs/Horde/Horde.Server/Ddc/IBlobService.cs
+++ b/Engine/Source/Programs/Horde/Horde.Server/Ddc/IBlobService.cs
@@ -7,150 +7,149 @@ using System.Net.Mime;
 using System.Threading.Tasks;
 using EpicGames.AspNet;
 using EpicGames.Horde.Storage;
-using Jupiter.Common.Implementation;
-using Jupiter.Utils;
 using Microsoft.Extensions.DependencyInjection;
 using OpenTelemetry.Trace;
 
 #pragma warning disable CS1591
 
-namespace Jupiter.Implementation;
-
-public interface IBlobService
-{
-    Task<JupiterContentHash> VerifyContentMatchesHash(Stream content, JupiterContentHash identifier);
-    Task<BlobIdentifier> PutObjectKnownHash(NamespaceId ns, IBufferedPayload content, BlobIdentifier identifier);
-    Task<BlobIdentifier> PutObject(NamespaceId ns, IBufferedPayload payload, BlobIdentifier identifier);
-    Task<BlobIdentifier> PutObject(NamespaceId ns, byte[] payload, BlobIdentifier identifier);
-    Task<Uri?> MaybePutObjectWithRedirect(NamespaceId ns, BlobIdentifier identifier);
-
-    Task<BlobContents> GetObject(NamespaceId ns, BlobIdentifier blob, List<string>? storageLayers = null, bool supportsRedirectUri = false);
-    
-    Task<Uri?> GetObjectWithRedirect(NamespaceId ns, BlobIdentifier blobIdentifier, List<string>? storageLayers = null);
-
-    Task<BlobContents> ReplicateObject(NamespaceId ns, BlobIdentifier blob, bool force = false);
-
-    Task<bool> Exists(NamespaceId ns, BlobIdentifier blob, List<string>? storageLayers = null);
-
-    /// <summary>
-    /// Checks that the blob exists in the root store, the store which is last in the list and thus is intended to have every blob in it
-    /// </summary>
-    /// <param name="ns">The namespace</param>
-    /// <param name="blob">The identifier of the blob</param>
-    /// <returns></returns>
-    Task<bool> ExistsInRootStore(NamespaceId ns, BlobIdentifier blob);
-
-    // Delete a object
-    Task DeleteObject(NamespaceId ns, BlobIdentifier blob);
-
-    // delete the whole namespace
-    Task DeleteNamespace(NamespaceId ns);
-
-    IAsyncEnumerable<(BlobIdentifier,DateTime)> ListObjects(NamespaceId ns);
-    Task<BlobIdentifier[]> FilterOutKnownBlobs(NamespaceId ns, IEnumerable<BlobIdentifier> blobs);
-    Task<BlobIdentifier[]> FilterOutKnownBlobs(NamespaceId ns, IAsyncEnumerable<BlobIdentifier> blobs);
-    Task<BlobContents> GetObjects(NamespaceId ns, BlobIdentifier[] refRequestBlobReferences);
-
-    bool ShouldFetchBlobOnDemand(NamespaceId ns);
-}
-
-public static class BlobServiceExtensions
+namespace Horde.Server.Ddc
 {
-    public static async Task<ContentId> PutCompressedObject(this IBlobService blobService, NamespaceId ns, IBufferedPayload payload, ContentId? id, IServiceProvider provider)
-    {
-        IContentIdStore contentIdStore = provider.GetService<IContentIdStore>()!;
-        CompressedBufferUtils compressedBufferUtils = provider.GetService<CompressedBufferUtils>()!;
-        Tracer tracer = provider.GetService<Tracer>()!;
-
-        // decompress the content and generate a identifier from it to verify the identifier we got
-        await using Stream decompressStream = payload.GetStream();
-        // TODO: we should add a overload for decompress content that can work on streams, otherwise we are still limited to 2GB compressed blobs
-        byte[] decompressedContent = compressedBufferUtils.DecompressContent(await decompressStream.ToByteArray());
-
-        await using MemoryStream decompressedStream = new MemoryStream(decompressedContent);
-        ContentId identifierDecompressedPayload;
-        if (id != null)
-        {
-            identifierDecompressedPayload = ContentId.FromContentHash(await blobService.VerifyContentMatchesHash(decompressedStream, id));
-        }
-        else
-        {
-            JupiterContentHash blobHash;
-            {
-                using TelemetrySpan _ = tracer.StartActiveSpan("web.hash").SetAttribute("operation.name", "web.hash");
-                blobHash = await BlobIdentifier.FromStream(decompressedStream);
-            }
-
-            identifierDecompressedPayload = ContentId.FromContentHash(blobHash);
-        }
-
-        BlobIdentifier identifierCompressedPayload;
-        {
-            using TelemetrySpan _ = tracer.StartActiveSpan("web.hash").SetAttribute("operation.name", "web.hash");
-            await using Stream hashStream = payload.GetStream();
-            identifierCompressedPayload = await BlobIdentifier.FromStream(hashStream);
-        }
-
-        // commit the mapping from the decompressed hash to the compressed hash, we run this in parallel with the blob store submit
-        // TODO: let users specify weight of the blob compared to previously submitted content ids
-        int contentIdWeight = (int)payload.Length;
-        Task contentIdStoreTask = contentIdStore.Put(ns, identifierDecompressedPayload, identifierCompressedPayload, contentIdWeight);
-
-        // we still commit the compressed buffer to the object store using the hash of the compressed content
-        {
-            await blobService.PutObjectKnownHash(ns, payload, identifierCompressedPayload);
-        }
-
-        await contentIdStoreTask;
-
-        return identifierDecompressedPayload;
-    }
-
-    public static async Task<(BlobContents, string)> GetCompressedObject(this IBlobService blobService, NamespaceId ns, ContentId contentId, IServiceProvider provider, bool supportsRedirectUri = false)
-    {
-        IContentIdStore contentIdStore = provider.GetService<IContentIdStore>()!;
-        Tracer tracer = provider.GetService<Tracer>()!;
-
-        BlobIdentifier[]? chunks = await contentIdStore.Resolve(ns, contentId, mustBeContentId: false);
-        if (chunks == null || chunks.Length == 0)
-        {
-            throw new ContentIdResolveException(contentId);
-        }
-
-        // single chunk, we just return that chunk
-        if (chunks.Length == 1)
-        {
-            BlobIdentifier blobToReturn = chunks[0];
-            string mimeType = CustomMediaTypeNames.UnrealCompressedBuffer;
-            if (contentId.Equals(blobToReturn))
-            {
-                // this was actually the unmapped blob, meaning its not a compressed buffer
-                mimeType = MediaTypeNames.Application.Octet;
-            }
-
-            return (await blobService.GetObject(ns, blobToReturn, supportsRedirectUri: supportsRedirectUri), mimeType);
-        }
-
-        // chunked content, combine the chunks into a single stream
-        using TelemetrySpan _ = tracer.StartActiveSpan("blob.combine").SetAttribute("operation.name", "blob.combine");
-        Task<BlobContents>[] tasks = new Task<BlobContents>[chunks.Length];
-        for (int i = 0; i < chunks.Length; i++)
-        {
-            // even if it was requested to support redirect, since we need to combine the chunks using redirects is not possible
-            tasks[i] = blobService.GetObject(ns, chunks[i], supportsRedirectUri: false);
-        }
-
-        MemoryStream ms = new MemoryStream();
-        foreach (Task<BlobContents> task in tasks)
-        {
-            BlobContents blob = await task;
-            await using Stream s = blob.Stream;
-            await s.CopyToAsync(ms);
-        }
-
-        ms.Seek(0, SeekOrigin.Begin);
-
-        // chunking could not have happened for a non compressed buffer so assume it is compressed
-        return (new BlobContents(ms, ms.Length), CustomMediaTypeNames.UnrealCompressedBuffer);
-    }
+	public interface IBlobService
+	{
+		Task<JupiterContentHash> VerifyContentMatchesHash(Stream content, JupiterContentHash identifier);
+		Task<BlobIdentifier> PutObjectKnownHash(NamespaceId ns, IBufferedPayload content, BlobIdentifier identifier);
+		Task<BlobIdentifier> PutObject(NamespaceId ns, IBufferedPayload payload, BlobIdentifier identifier);
+		Task<BlobIdentifier> PutObject(NamespaceId ns, byte[] payload, BlobIdentifier identifier);
+		Task<Uri?> MaybePutObjectWithRedirect(NamespaceId ns, BlobIdentifier identifier);
+
+		Task<BlobContents> GetObject(NamespaceId ns, BlobIdentifier blob, List<string>? storageLayers = null, bool supportsRedirectUri = false);
+
+		Task<Uri?> GetObjectWithRedirect(NamespaceId ns, BlobIdentifier blobIdentifier, List<string>? storageLayers = null);
+
+		Task<BlobContents> ReplicateObject(NamespaceId ns, BlobIdentifier blob, bool force = false);
+
+		Task<bool> Exists(NamespaceId ns, BlobIdentifier blob, List<string>? storageLayers = null);
+
+		/// <summary>
+		/// Checks that the blob exists in the root store, the store which is last in the list and thus is intended to have every blob in it
+		/// </summary>
+		/// <param name="ns">The namespace</param>
+		/// <param name="blob">The identifier of the blob</param>
+		/// <returns></returns>
+		Task<bool> ExistsInRootStore(NamespaceId ns, BlobIdentifier blob);
+
+		// Delete a object
+		Task DeleteObject(NamespaceId ns, BlobIdentifier blob);
+
+		// delete the whole namespace
+		Task DeleteNamespace(NamespaceId ns);
+
+		IAsyncEnumerable<(BlobIdentifier, DateTime)> ListObjects(NamespaceId ns);
+		Task<BlobIdentifier[]> FilterOutKnownBlobs(NamespaceId ns, IEnumerable<BlobIdentifier> blobs);
+		Task<BlobIdentifier[]> FilterOutKnownBlobs(NamespaceId ns, IAsyncEnumerable<BlobIdentifier> blobs);
+		Task<BlobContents> GetObjects(NamespaceId ns, BlobIdentifier[] refRequestBlobReferences);
+
+		bool ShouldFetchBlobOnDemand(NamespaceId ns);
+	}
+
+	public static class BlobServiceExtensions
+	{
+		public static async Task<ContentId> PutCompressedObject(this IBlobService blobService, NamespaceId ns, IBufferedPayload payload, ContentId? id, IServiceProvider provider)
+		{
+			IContentIdStore contentIdStore = provider.GetService<IContentIdStore>()!;
+			CompressedBufferUtils compressedBufferUtils = provider.GetService<CompressedBufferUtils>()!;
+			Tracer tracer = provider.GetService<Tracer>()!;
+
+			// decompress the content and generate a identifier from it to verify the identifier we got
+			await using Stream decompressStream = payload.GetStream();
+			// TODO: we should add a overload for decompress content that can work on streams, otherwise we are still limited to 2GB compressed blobs
+			byte[] decompressedContent = compressedBufferUtils.DecompressContent(await decompressStream.ToByteArray());
+
+			await using MemoryStream decompressedStream = new MemoryStream(decompressedContent);
+			ContentId identifierDecompressedPayload;
+			if (id != null)
+			{
+				identifierDecompressedPayload = ContentId.FromContentHash(await blobService.VerifyContentMatchesHash(decompressedStream, id));
+			}
+			else
+			{
+				JupiterContentHash blobHash;
+				{
+					using TelemetrySpan _ = tracer.StartActiveSpan("web.hash").SetAttribute("operation.name", "web.hash");
+					blobHash = await BlobIdentifier.FromStream(decompressedStream);
+				}
+
+				identifierDecompressedPayload = ContentId.FromContentHash(blobHash);
+			}
+
+			BlobIdentifier identifierCompressedPayload;
+			{
+				using TelemetrySpan _ = tracer.StartActiveSpan("web.hash").SetAttribute("operation.name", "web.hash");
+				await using Stream hashStream = payload.GetStream();
+				identifierCompressedPayload = await BlobIdentifier.FromStream(hashStream);
+			}
+
+			// commit the mapping from the decompressed hash to the compressed hash, we run this in parallel with the blob store submit
+			// TODO: let users specify weight of the blob compared to previously submitted content ids
+			int contentIdWeight = (int)payload.Length;
+			Task contentIdStoreTask = contentIdStore.Put(ns, identifierDecompressedPayload, identifierCompressedPayload, contentIdWeight);
+
+			// we still commit the compressed buffer to the object store using the hash of the compressed content
+			{
+				await blobService.PutObjectKnownHash(ns, payload, identifierCompressedPayload);
+			}
+
+			await contentIdStoreTask;
+
+			return identifierDecompressedPayload;
+		}
+
+		public static async Task<(BlobContents, string)> GetCompressedObject(this IBlobService blobService, NamespaceId ns, ContentId contentId, IServiceProvider provider, bool supportsRedirectUri = false)
+		{
+			IContentIdStore contentIdStore = provider.GetService<IContentIdStore>()!;
+			Tracer tracer = provider.GetService<Tracer>()!;
+
+			BlobIdentifier[]? chunks = await contentIdStore.Resolve(ns, contentId, mustBeContentId: false);
+			if (chunks == null || chunks.Length == 0)
+			{
+				throw new ContentIdResolveException(contentId);
+			}
+
+			// single chunk, we just return that chunk
+			if (chunks.Length == 1)
+			{
+				BlobIdentifier blobToReturn = chunks[0];
+				string mimeType = CustomMediaTypeNames.UnrealCompressedBuffer;
+				if (contentId.Equals(blobToReturn))
+				{
+					// this was actually the unmapped blob, meaning its not a compressed buffer
+					mimeType = MediaTypeNames.Application.Octet;
+				}
+
+				return (await blobService.GetObject(ns, blobToReturn, supportsRedirectUri: supportsRedirectUri), mimeType);
+			}
+
+			// chunked content, combine the chunks into a single stream
+			using TelemetrySpan _ = tracer.StartActiveSpan("blob.combine").SetAttribute("operation.name", "blob.combine");
+			Task<BlobContents>[] tasks = new Task<BlobContents>[chunks.Length];
+			for (int i = 0; i < chunks.Length; i++)
+			{
+				// even if it was requested to support redirect, since we need to combine the chunks using redirects is not possible
+				tasks[i] = blobService.GetObject(ns, chunks[i], supportsRedirectUri: false);
+			}
+
+			MemoryStream ms = new MemoryStream();
+			foreach (Task<BlobContents> task in tasks)
+			{
+				BlobContents blob = await task;
+				await using Stream s = blob.Stream;
+				await s.CopyToAsync(ms);
+			}
+
+			ms.Seek(0, SeekOrigin.Begin);
+
+			// chunking could not have happened for a non compressed buffer so assume it is compressed
+			return (new BlobContents(ms, ms.Length), CustomMediaTypeNames.UnrealCompressedBuffer);
+		}
+	}
 }
diff --git a/Engine/Source/Programs/Horde/Horde.Server/Ddc/IBlobStore.cs b/Engine/Source/Programs/Horde/Horde.Server/Ddc/IBlobStore.cs
--- a/Engine/Source/Programs/Horde/Horde.Server/Ddc/IBlobStore.cs
+++ b/Engine/Source/Programs/Horde/Horde.Server/Ddc/IBlobStore.cs
@@ -9,9 +9,9 @@ using EpicGames.Horde.Storage;
 
 #pragma warning disable CS1591
 
-namespace Jupiter.Implementation
+namespace Horde.Server.Ddc
 {
-    public enum LastAccessTrackingFlags
+	public enum LastAccessTrackingFlags
     {
         DoTracking = 0,
         SkipTracking = 1,
diff --git a/Engine/Source/Programs/Horde/Horde.Server/Ddc/IContentIdStore.cs b/Engine/Source/Programs/Horde/Horde.Server/Ddc/IContentIdStore.cs
--- a/Engine/Source/Programs/Horde/Horde.Server/Ddc/IContentIdStore.cs
+++ b/Engine/Source/Programs/Horde/Horde.Server/Ddc/IContentIdStore.cs
@@ -6,9 +6,9 @@ using System.Threading.Tasks;
 
 #pragma warning disable CS1591
 
-namespace Jupiter.Implementation
+namespace Horde.Server.Ddc
 {
-    public interface IContentIdStore
+	public interface IContentIdStore
     {
         /// <summary>
         /// Resolve a content id from its hash into the actual blob (that can in turn be chunked into a set of blobs)
diff --git a/Engine/Source/Programs/Horde/Horde.Server/Ddc/IObjectService.cs b/Engine/Source/Programs/Horde/Horde.Server/Ddc/IObjectService.cs
--- a/Engine/Source/Programs/Horde/Horde.Server/Ddc/IObjectService.cs
+++ b/Engine/Source/Programs/Horde/Horde.Server/Ddc/IObjectService.cs
@@ -8,9 +8,9 @@ using EpicGames.Serialization;
 
 #pragma warning disable CS1591
 
-namespace Jupiter.Implementation
+namespace Horde.Server.Ddc
 {
-    public interface IObjectService
+	public interface IObjectService
     {
         Task<(ObjectRecord, BlobContents?)> Get(NamespaceId ns, BucketId bucket, IoHashKey key, string[] fields, bool doLastAccessTracking = true);
         Task<(ContentId[], BlobIdentifier[])> Put(NamespaceId ns, BucketId bucket, IoHashKey key, BlobIdentifier blobHash, CbObject payload);
diff --git a/Engine/Source/Programs/Horde/Horde.Server/Ddc/IReferenceResolver.cs b/Engine/Source/Programs/Horde/Horde.Server/Ddc/IReferenceResolver.cs
--- a/Engine/Source/Programs/Horde/Horde.Server/Ddc/IReferenceResolver.cs
+++ b/Engine/Source/Programs/Horde/Horde.Server/Ddc/IReferenceResolver.cs
@@ -7,14 +7,13 @@ using System.Threading.Tasks;
 using EpicGames.Core;
 using EpicGames.Horde.Storage;
 using EpicGames.Serialization;
-using Jupiter.Utils;
 using OpenTelemetry.Trace;
 
 #pragma warning disable CS1591
 
-namespace Jupiter.Implementation
+namespace Horde.Server.Ddc
 {
-    public abstract class Attachment
+	public abstract class Attachment
     {
         public abstract IoHash AsIoHash();
     }
diff --git a/Engine/Source/Programs/Horde/Horde.Server/Ddc/IReferencesStore.cs b/Engine/Source/Programs/Horde/Horde.Server/Ddc/IReferencesStore.cs
--- a/Engine/Source/Programs/Horde/Horde.Server/Ddc/IReferencesStore.cs
+++ b/Engine/Source/Programs/Horde/Horde.Server/Ddc/IReferencesStore.cs
@@ -7,9 +7,9 @@ using EpicGames.Horde.Storage;
 
 #pragma warning disable CS1591
 
-namespace Jupiter.Implementation
+namespace Horde.Server.Ddc
 {
-    public interface IReferencesStore
+	public interface IReferencesStore
     {
         Task<ObjectRecord> Get(NamespaceId ns, BucketId bucket, IoHashKey key, FieldFlags flags);
 
diff --git a/Engine/Source/Programs/Horde/Horde.Server/Ddc/IRequestHelper.cs b/Engine/Source/Programs/Horde/Horde.Server/Ddc/IRequestHelper.cs
--- a/Engine/Source/Programs/Horde/Horde.Server/Ddc/IRequestHelper.cs
+++ b/Engine/Source/Programs/Horde/Horde.Server/Ddc/IRequestHelper.cs
@@ -9,9 +9,9 @@ using Microsoft.AspNetCore.Mvc;
 
 #pragma warning disable CS1591
 
-namespace Jupiter.Controllers
+namespace Horde.Server.Ddc
 {
-    public interface IRequestHelper
+	public interface IRequestHelper
     {
         public Task<ActionResult?> HasAccessToNamespace(ClaimsPrincipal user, HttpRequest request, NamespaceId ns, JupiterAclAction[] aclActions);
         public Task<ActionResult?> HasAccessForGlobalOperations(ClaimsPrincipal user, JupiterAclAction[] aclActions);
diff --git a/Engine/Source/Programs/Horde/Horde.Server/Ddc/IoHashKey.cs b/Engine/Source/Programs/Horde/Horde.Server/Ddc/IoHashKey.cs
--- a/Engine/Source/Programs/Horde/Horde.Server/Ddc/IoHashKey.cs
+++ b/Engine/Source/Programs/Horde/Horde.Server/Ddc/IoHashKey.cs
@@ -12,9 +12,9 @@ using EpicGames.Serialization;
 
 #pragma warning disable CS1591
 
-namespace Jupiter.Implementation
+namespace Horde.Server.Ddc
 {
-    [TypeConverter(typeof(IoHashKeyTypeConverter))]
+	[TypeConverter(typeof(IoHashKeyTypeConverter))]
     [JsonConverter(typeof(IoHashKeyJsonConverter))]
     [CbConverter(typeof(IoHashKeyCbConverter))]
 
diff --git a/Engine/Source/Programs/Horde/Horde.Server/Ddc/JupiterContentHash.cs b/Engine/Source/Programs/Horde/Horde.Server/Ddc/JupiterContentHash.cs
--- a/Engine/Source/Programs/Horde/Horde.Server/Ddc/JupiterContentHash.cs
+++ b/Engine/Source/Programs/Horde/Horde.Server/Ddc/JupiterContentHash.cs
@@ -12,9 +12,9 @@ using EpicGames.Serialization;
 
 #pragma warning disable CS1591
 
-namespace Jupiter.Implementation
+namespace Horde.Server.Ddc
 {
-    [TypeConverter(typeof(ContentHashTypeConverter))]
+	[TypeConverter(typeof(ContentHashTypeConverter))]
     [JsonConverter(typeof(ContentHashJsonConverter))]
     [CbConverter(typeof(ContentHashCbConverter))]
     public class JupiterContentHash : IEquatable<JupiterContentHash>, IEquatable<byte[]>
diff --git a/Engine/Source/Programs/Horde/Horde.Server/Ddc/NamespaceNotFoundException.cs b/Engine/Source/Programs/Horde/Horde.Server/Ddc/NamespaceNotFoundException.cs
--- a/Engine/Source/Programs/Horde/Horde.Server/Ddc/NamespaceNotFoundException.cs
+++ b/Engine/Source/Programs/Horde/Horde.Server/Ddc/NamespaceNotFoundException.cs
@@ -5,9 +5,9 @@ using EpicGames.Horde.Storage;
 
 #pragma warning disable CS1591
 
-namespace Jupiter.Common
+namespace Horde.Server.Ddc
 {
-    public class NamespaceNotFoundException : Exception
+	public class NamespaceNotFoundException : Exception
     {
         public NamespaceId Namespace { get; }
 
diff --git a/Engine/Source/Programs/Horde/Horde.Server/Ddc/NginxRedirectHelper.cs b/Engine/Source/Programs/Horde/Horde.Server/Ddc/NginxRedirectHelper.cs
--- a/Engine/Source/Programs/Horde/Horde.Server/Ddc/NginxRedirectHelper.cs
+++ b/Engine/Source/Programs/Horde/Horde.Server/Ddc/NginxRedirectHelper.cs
@@ -8,9 +8,9 @@ using Microsoft.Extensions.Options;
 
 #pragma warning disable CS1591
 
-namespace Jupiter.Implementation
+namespace Horde.Server.Ddc
 {
-    public class NginxRedirectHelper
+	public class NginxRedirectHelper
     {
         private readonly IOptionsMonitor<NginxSettings> _nginxSettings;
 
diff --git a/Engine/Source/Programs/Horde/Horde.Server/Ddc/ObjectController.cs b/Engine/Source/Programs/Horde/Horde.Server/Ddc/ObjectController.cs
--- a/Engine/Source/Programs/Horde/Horde.Server/Ddc/ObjectController.cs
+++ b/Engine/Source/Programs/Horde/Horde.Server/Ddc/ObjectController.cs
@@ -11,19 +11,15 @@ using System.Threading.Tasks;
 using EpicGames.AspNet;
 using EpicGames.Horde.Storage;
 using EpicGames.Serialization;
-using Jupiter.Implementation;
-using Jupiter.Common.Implementation;
-using Jupiter.Utils;
 using Microsoft.AspNetCore.Authorization;
 using Microsoft.AspNetCore.Mvc;
 using Microsoft.Extensions.Logging;
 
 #pragma warning disable CS1591
 
-namespace Jupiter.Controllers
+namespace Horde.Server.Ddc
 {
-    using IDiagnosticContext = Serilog.IDiagnosticContext;
-    using BlobNotFoundException = Jupiter.Implementation.BlobNotFoundException;
+	using IDiagnosticContext = Serilog.IDiagnosticContext;
 
     [ApiController]
     [Route("api/v1/objects", Order = 0)]
diff --git a/Engine/Source/Programs/Horde/Horde.Server/Ddc/ObjectService.cs b/Engine/Source/Programs/Horde/Horde.Server/Ddc/ObjectService.cs
--- a/Engine/Source/Programs/Horde/Horde.Server/Ddc/ObjectService.cs
+++ b/Engine/Source/Programs/Horde/Horde.Server/Ddc/ObjectService.cs
@@ -4,7 +4,6 @@ using System.Collections.Generic;
 using System.Threading.Tasks;
 using EpicGames.Horde.Storage;
 using EpicGames.Serialization;
-using Jupiter.Implementation;
 
 namespace Horde.Server.Ddc
 {
diff --git a/Engine/Source/Programs/Horde/Horde.Server/Ddc/ReferencesController.cs b/Engine/Source/Programs/Horde/Horde.Server/Ddc/ReferencesController.cs
--- a/Engine/Source/Programs/Horde/Horde.Server/Ddc/ReferencesController.cs
+++ b/Engine/Source/Programs/Horde/Horde.Server/Ddc/ReferencesController.cs
@@ -17,11 +17,7 @@ using EpicGames.AspNet;
 using EpicGames.Core;
 using EpicGames.Horde.Storage;
 using EpicGames.Serialization;
-using Jupiter.Implementation;
 using JetBrains.Annotations;
-using Jupiter.Common;
-using Jupiter.Common.Implementation;
-using Jupiter.Utils;
 using Microsoft.AspNetCore.Authorization;
 using Microsoft.AspNetCore.Http;
 using Microsoft.AspNetCore.Http.Extensions;
@@ -32,13 +28,9 @@ using OpenTelemetry.Trace;
 
 #pragma warning disable CS1591
 
-using JupiterContentHash = Jupiter.Implementation.JupiterContentHash;
-using ContentId = Jupiter.Implementation.ContentId;
-
-namespace Jupiter.Controllers
+namespace Horde.Server.Ddc
 {
 	using IDiagnosticContext = Serilog.IDiagnosticContext;
-    using BlobNotFoundException = Jupiter.Implementation.BlobNotFoundException;
 
     [ApiController]
     [FormatFilter]
diff --git a/Engine/Source/Programs/Horde/Horde.Server/Ddc/ReferencesStore.cs b/Engine/Source/Programs/Horde/Horde.Server/Ddc/ReferencesStore.cs
--- a/Engine/Source/Programs/Horde/Horde.Server/Ddc/ReferencesStore.cs
+++ b/Engine/Source/Programs/Horde/Horde.Server/Ddc/ReferencesStore.cs
@@ -4,7 +4,6 @@ using System;
 using System.Collections.Generic;
 using System.Threading.Tasks;
 using EpicGames.Horde.Storage;
-using Jupiter.Implementation;
 
 namespace Horde.Server.Ddc
 {
diff --git a/Engine/Source/Programs/Horde/Horde.Server/Ddc/RequestHelper.cs b/Engine/Source/Programs/Horde/Horde.Server/Ddc/RequestHelper.cs
--- a/Engine/Source/Programs/Horde/Horde.Server/Ddc/RequestHelper.cs
+++ b/Engine/Source/Programs/Horde/Horde.Server/Ddc/RequestHelper.cs
@@ -4,8 +4,6 @@ using System;
 using System.Security.Claims;
 using System.Threading.Tasks;
 using EpicGames.Horde.Storage;
-using Jupiter;
-using Jupiter.Controllers;
 using Microsoft.AspNetCore.Http;
 using Microsoft.AspNetCore.Mvc;
 
diff --git a/Engine/Source/Programs/Horde/Horde.Server/Ddc/RequestUtil.cs b/Engine/Source/Programs/Horde/Horde.Server/Ddc/RequestUtil.cs
--- a/Engine/Source/Programs/Horde/Horde.Server/Ddc/RequestUtil.cs
+++ b/Engine/Source/Programs/Horde/Horde.Server/Ddc/RequestUtil.cs
@@ -9,9 +9,9 @@ using OpenTelemetry.Trace;
 
 #pragma warning disable CS1591
 
-namespace Jupiter
+namespace Horde.Server.Ddc
 {
-    public static class RequestUtil
+	public static class RequestUtil
     {
         public static async Task<byte[]> ReadRawBody(HttpRequest request)
         {
diff --git a/Engine/Source/Programs/Horde/Horde.Server/Ddc/RequiredContentTypeAttribute.cs b/Engine/Source/Programs/Horde/Horde.Server/Ddc/RequiredContentTypeAttribute.cs
--- a/Engine/Source/Programs/Horde/Horde.Server/Ddc/RequiredContentTypeAttribute.cs
+++ b/Engine/Source/Programs/Horde/Horde.Server/Ddc/RequiredContentTypeAttribute.cs
@@ -7,9 +7,9 @@ using Microsoft.Extensions.Primitives;
 
 #pragma warning disable CS1591
 
-namespace Jupiter
+namespace Horde.Server.Ddc
 {
-    [AttributeUsage(AttributeTargets.Method)]
+	[AttributeUsage(AttributeTargets.Method)]
     public sealed class RequiredContentTypeAttribute : Attribute, IActionConstraint
     {
         private readonly string _mediaTypeName;
diff --git a/Engine/Source/Programs/Horde/Horde.Server/Ddc/Sha1Utils.cs b/Engine/Source/Programs/Horde/Horde.Server/Ddc/Sha1Utils.cs
--- a/Engine/Source/Programs/Horde/Horde.Server/Ddc/Sha1Utils.cs
+++ b/Engine/Source/Programs/Horde/Horde.Server/Ddc/Sha1Utils.cs
@@ -6,9 +6,9 @@ using System.Threading.Tasks;
 
 #pragma warning disable CS1591
 
-namespace Jupiter
+namespace Horde.Server.Ddc
 {
-    public static class Sha1Utils
+	public static class Sha1Utils
     {
         public static async Task<string> GetSHA1AsStringAsync(byte[] payload)
         {
diff --git a/Engine/Source/Programs/Horde/Horde.Server/Ddc/StreamUtils.cs b/Engine/Source/Programs/Horde/Horde.Server/Ddc/StreamUtils.cs
--- a/Engine/Source/Programs/Horde/Horde.Server/Ddc/StreamUtils.cs
+++ b/Engine/Source/Programs/Horde/Horde.Server/Ddc/StreamUtils.cs
@@ -7,9 +7,9 @@ using Microsoft.AspNetCore.Http;
 
 #pragma warning disable CS1591
 
-namespace Jupiter.Utils
+namespace Horde.Server.Ddc
 {
-    public static class StreamUtils
+	public static class StreamUtils
     {
         public static async Task<byte[]> ToByteArray(this Stream s)
         {
diff --git a/Engine/Source/Programs/Horde/Horde.Server/Ddc/StreamingBodyPayload.cs b/Engine/Source/Programs/Horde/Horde.Server/Ddc/StreamingBodyPayload.cs
--- a/Engine/Source/Programs/Horde/Horde.Server/Ddc/StreamingBodyPayload.cs
+++ b/Engine/Source/Programs/Horde/Horde.Server/Ddc/StreamingBodyPayload.cs
@@ -3,16 +3,15 @@
 using System;
 using System.IO;
 using System.Threading.Tasks;
-using Jupiter.Utils;
 using Microsoft.AspNetCore.Http;
 using Microsoft.Extensions.Options;
 using OpenTelemetry.Trace;
 
 #pragma warning disable CS1591
 
-namespace Jupiter.Common.Implementation
+namespace Horde.Server.Ddc
 {
-    public interface IBufferedPayload : IDisposable
+	public interface IBufferedPayload : IDisposable
     {
         Stream GetStream();
         long Length { get; }
diff --git a/Engine/Source/Programs/Horde/Horde.Server/Ddc/StringUtils.cs b/Engine/Source/Programs/Horde/Horde.Server/Ddc/StringUtils.cs
--- a/Engine/Source/Programs/Horde/Horde.Server/Ddc/StringUtils.cs
+++ b/Engine/Source/Programs/Horde/Horde.Server/Ddc/StringUtils.cs
@@ -5,9 +5,9 @@ using System.Text;
 
 #pragma warning disable CS1591
 
-namespace Jupiter
+namespace Horde.Server.Ddc
 {
-    public static class StringUtils
+	public static class StringUtils
     {
         public static int GetHexVal(char hex)
         {
diff --git a/Engine/Source/Programs/Horde/Horde.Server/Startup.cs b/Engine/Source/Programs/Horde/Horde.Server/Startup.cs
--- a/Engine/Source/Programs/Horde/Horde.Server/Startup.cs
+++ b/Engine/Source/Programs/Horde/Horde.Server/Startup.cs
@@ -104,14 +104,13 @@ using Microsoft.AspNetCore.Diagnostics;
 using EpicGames.Horde;
 using EpicGames.Horde.Api;
 using EpicGames.Horde.Storage.Bundles;
-using Jupiter.Implementation;
 using Horde.Server.Ddc;
 
 namespace Horde.Server
 {
-	using ReferenceResolver = Jupiter.Implementation.ReferenceResolver;
 	using IStorageBackend = Horde.Server.Storage.IStorageBackend;
 	using ContentHash = EpicGames.Core.ContentHash;
+	using ReferenceResolver = Horde.Server.Ddc.ReferenceResolver;
 
 	class Startup
 	{
