Horde: Use the regular Horde AclAction type to represent DDC actions.

diff --git a/Engine/Source/Programs/Horde/Horde.Server/Ddc/Acl.cs b/Engine/Source/Programs/Horde/Horde.Server/Ddc/Acl.cs
deleted file mode 100644
diff --git a/Engine/Source/Programs/Horde/Horde.Server/Ddc/BlobController.cs b/Engine/Source/Programs/Horde/Horde.Server/Ddc/BlobController.cs
--- a/Engine/Source/Programs/Horde/Horde.Server/Ddc/BlobController.cs
+++ b/Engine/Source/Programs/Horde/Horde.Server/Ddc/BlobController.cs
@@ -13,6 +13,7 @@ using System.Threading.Tasks;
 using EpicGames.AspNet;
 using EpicGames.Horde.Storage;
 using EpicGames.Serialization;
+using Horde.Server.Acls;
 using Microsoft.AspNetCore.Authorization;
 using Microsoft.AspNetCore.Http;
 using Microsoft.AspNetCore.Mvc;
@@ -51,7 +52,7 @@ namespace Horde.Server.Ddc
             [Required] BlobIdentifier id,
             [FromQuery] List<string>? storageLayers = null)
         {
-            ActionResult? result = await _requestHelper.HasAccessToNamespace(User, Request, ns, new [] { JupiterAclAction.ReadObject });
+            ActionResult? result = await _requestHelper.HasAccessToNamespace(User, Request, ns, new [] { DdcAclAction.ReadObject });
             if (result != null)
             {
                 return result;
@@ -88,7 +89,7 @@ namespace Horde.Server.Ddc
             [Required] BlobIdentifier id,
             [FromQuery] List<string>? storageLayers = null)
         {
-            ActionResult? result = await _requestHelper.HasAccessToNamespace(User, Request, ns, new [] { JupiterAclAction.ReadObject });
+            ActionResult? result = await _requestHelper.HasAccessToNamespace(User, Request, ns, new [] { DdcAclAction.ReadObject });
             if (result != null)
             {
                 return result;
@@ -109,7 +110,7 @@ namespace Horde.Server.Ddc
             [Required] NamespaceId ns,
             [Required] [FromQuery] List<BlobIdentifier> id)
         {
-            ActionResult? result = await _requestHelper.HasAccessToNamespace(User, Request, ns, new [] { JupiterAclAction.ReadObject });
+            ActionResult? result = await _requestHelper.HasAccessToNamespace(User, Request, ns, new [] { DdcAclAction.ReadObject });
             if (result != null)
             {
                 return result;
@@ -135,7 +136,7 @@ namespace Horde.Server.Ddc
             [Required] NamespaceId ns,
             [FromBody] BlobIdentifier[] bodyIds)
         {
-            ActionResult? result = await _requestHelper.HasAccessToNamespace(User, Request, ns, new [] { JupiterAclAction.ReadObject });
+            ActionResult? result = await _requestHelper.HasAccessToNamespace(User, Request, ns, new [] { DdcAclAction.ReadObject });
             if (result != null)
             {
                 return result;
@@ -179,7 +180,7 @@ namespace Horde.Server.Ddc
             [Required] NamespaceId ns,
             [Required] BlobIdentifier id)
         {
-            ActionResult? result = await _requestHelper.HasAccessToNamespace(User, Request, ns, new [] { JupiterAclAction.WriteObject });
+            ActionResult? result = await _requestHelper.HasAccessToNamespace(User, Request, ns, new [] { DdcAclAction.WriteObject });
             if (result != null)
             {
                 return result;
@@ -222,7 +223,7 @@ namespace Horde.Server.Ddc
         public async Task<IActionResult> Post(
             [Required] NamespaceId ns)
         {
-            ActionResult? result = await _requestHelper.HasAccessToNamespace(User, Request, ns, new [] { JupiterAclAction.WriteObject });
+            ActionResult? result = await _requestHelper.HasAccessToNamespace(User, Request, ns, new [] { DdcAclAction.WriteObject });
             if (result != null)
             {
                 return result;
@@ -254,7 +255,7 @@ namespace Horde.Server.Ddc
             [Required] NamespaceId ns,
             [Required] BlobIdentifier id)
         {
-            ActionResult? result = await _requestHelper.HasAccessToNamespace(User, Request, ns, new [] { JupiterAclAction.DeleteObject });
+            ActionResult? result = await _requestHelper.HasAccessToNamespace(User, Request, ns, new [] { DdcAclAction.DeleteObject });
             if (result != null)
             {
                 return result;
@@ -320,17 +321,17 @@ namespace Horde.Server.Ddc
         [HttpPost("")]
         public async Task<IActionResult> Post([FromBody] BatchCall batch)
         {
-            JupiterAclAction MapToAclAction(BatchOp.Operation op)
+            static AclAction MapToAclAction(BatchOp.Operation op)
             {
                 switch (op)
                 {
                     case BatchOp.Operation.GET:
                     case BatchOp.Operation.HEAD:
-                        return JupiterAclAction.ReadObject;
+                        return DdcAclAction.ReadObject;
                     case BatchOp.Operation.PUT:
-                        return JupiterAclAction.WriteObject;
+                        return DdcAclAction.WriteObject;
                     case BatchOp.Operation.DELETE:
-                        return JupiterAclAction.DeleteObject;
+                        return DdcAclAction.DeleteObject;
                     default:
                         throw new ArgumentOutOfRangeException(nameof(op), op, null);
                 }
diff --git a/Engine/Source/Programs/Horde/Horde.Server/Ddc/ContentIdController.cs b/Engine/Source/Programs/Horde/Horde.Server/Ddc/ContentIdController.cs
--- a/Engine/Source/Programs/Horde/Horde.Server/Ddc/ContentIdController.cs
+++ b/Engine/Source/Programs/Horde/Horde.Server/Ddc/ContentIdController.cs
@@ -37,7 +37,7 @@ namespace Horde.Server.Ddc
         [ProducesResponseType(type: typeof(ProblemDetails), 400)]
         public async Task<IActionResult> Resolve(NamespaceId ns, ContentId contentId)
         {
-            ActionResult? result = await _requestHelper.HasAccessToNamespace(User, Request, ns, new [] { JupiterAclAction.ReadObject });
+            ActionResult? result = await _requestHelper.HasAccessToNamespace(User, Request, ns, new [] { DdcAclAction.ReadObject });
             if (result != null)
             {
                 return result;
@@ -64,7 +64,7 @@ namespace Horde.Server.Ddc
         [ProducesResponseType(type: typeof(ProblemDetails), 400)]
         public async Task<IActionResult> UpdateContentIdMapping(NamespaceId ns, ContentId contentId, BlobIdentifier blobIdentifier, int contentWeight)
         {
-            ActionResult? result = await _requestHelper.HasAccessToNamespace(User, Request, ns, new [] { JupiterAclAction.WriteObject });
+            ActionResult? result = await _requestHelper.HasAccessToNamespace(User, Request, ns, new [] { DdcAclAction.WriteObject });
             if (result != null)
             {
                 return result;
diff --git a/Engine/Source/Programs/Horde/Horde.Server/Ddc/DdcAclAction.cs b/Engine/Source/Programs/Horde/Horde.Server/Ddc/DdcAclAction.cs
new file mode 100644
--- /dev/null
+++ b/Engine/Source/Programs/Horde/Horde.Server/Ddc/DdcAclAction.cs
@@ -0,0 +1,42 @@
+// Copyright Epic Games, Inc. All Rights Reserved.
+
+using Horde.Server.Acls;
+
+namespace Horde.Server.Ddc
+{
+	/// <summary>
+	/// Actions for manipulating objects in DDC
+	/// </summary>
+	public static class DdcAclAction
+	{
+		/// <summary>
+		/// General read access to refs / blobs and so on
+		/// </summary>
+		public static AclAction ReadObject { get; } = new AclAction("DdcReadObject");
+
+		/// <summary>
+		/// General write access to upload refs / blobs etc
+		/// </summary>
+		public static AclAction WriteObject { get; } = new AclAction("DdcWriteObject");
+
+		/// <summary>
+		/// Access to delete blobs / refs etc
+		/// </summary>
+		public static AclAction DeleteObject { get; } = new AclAction("DdcDeleteObject");
+
+		/// <summary>
+		/// Access to perform administrative task
+		/// </summary>
+		public static AclAction AdminAction { get; } = new AclAction("DdcAdmin");
+
+		/// <summary>
+		/// Access to delete a particular bucket
+		/// </summary>
+		public static AclAction DeleteBucket { get; } = new AclAction("DdcDeleteBucket");
+
+		/// <summary>
+		/// Access to delete a whole namespace
+		/// </summary>
+		public static AclAction DeleteNamespace { get; } = new AclAction("DdcDeleteNamespace");
+	}
+}
diff --git a/Engine/Source/Programs/Horde/Horde.Server/Ddc/IRequestHelper.cs b/Engine/Source/Programs/Horde/Horde.Server/Ddc/IRequestHelper.cs
--- a/Engine/Source/Programs/Horde/Horde.Server/Ddc/IRequestHelper.cs
+++ b/Engine/Source/Programs/Horde/Horde.Server/Ddc/IRequestHelper.cs
@@ -4,6 +4,7 @@ using System;
 using System.Security.Claims;
 using System.Threading.Tasks;
 using EpicGames.Horde.Storage;
+using Horde.Server.Acls;
 using Microsoft.AspNetCore.Http;
 using Microsoft.AspNetCore.Mvc;
 
@@ -13,8 +14,8 @@ namespace Horde.Server.Ddc
 {
 	public interface IRequestHelper
     {
-        public Task<ActionResult?> HasAccessToNamespace(ClaimsPrincipal user, HttpRequest request, NamespaceId ns, JupiterAclAction[] aclActions);
-        public Task<ActionResult?> HasAccessForGlobalOperations(ClaimsPrincipal user, JupiterAclAction[] aclActions);
+        public Task<ActionResult?> HasAccessToNamespace(ClaimsPrincipal user, HttpRequest request, NamespaceId ns, AclAction[] aclActions);
+        public Task<ActionResult?> HasAccessForGlobalOperations(ClaimsPrincipal user, AclAction[] aclActions);
     }
 
     public class AuthorizationException : Exception
diff --git a/Engine/Source/Programs/Horde/Horde.Server/Ddc/ObjectController.cs b/Engine/Source/Programs/Horde/Horde.Server/Ddc/ObjectController.cs
--- a/Engine/Source/Programs/Horde/Horde.Server/Ddc/ObjectController.cs
+++ b/Engine/Source/Programs/Horde/Horde.Server/Ddc/ObjectController.cs
@@ -51,7 +51,7 @@ namespace Horde.Server.Ddc
             [Required] NamespaceId ns,
             [Required] BlobIdentifier id)
         {
-            ActionResult? result = await _requestHelper.HasAccessToNamespace(User, Request, ns, new [] { JupiterAclAction.ReadObject });
+            ActionResult? result = await _requestHelper.HasAccessToNamespace(User, Request, ns, new [] { DdcAclAction.ReadObject });
             if (result != null)
             {
                 return result;
@@ -75,7 +75,7 @@ namespace Horde.Server.Ddc
             [Required] NamespaceId ns,
             [Required] BlobIdentifier id)
         {
-            ActionResult? result = await _requestHelper.HasAccessToNamespace(User, Request, ns, new [] { JupiterAclAction.ReadObject });
+            ActionResult? result = await _requestHelper.HasAccessToNamespace(User, Request, ns, new [] { DdcAclAction.ReadObject });
             if (result != null)
             {
                 return result;
@@ -97,7 +97,7 @@ namespace Horde.Server.Ddc
             [Required] NamespaceId ns,
             [Required] [FromQuery] List<BlobIdentifier> id)
         {
-            ActionResult? result = await _requestHelper.HasAccessToNamespace(User, Request, ns, new [] { JupiterAclAction.ReadObject });
+            ActionResult? result = await _requestHelper.HasAccessToNamespace(User, Request, ns, new [] { DdcAclAction.ReadObject });
             if (result != null)
             {
                 return result;
@@ -123,7 +123,7 @@ namespace Horde.Server.Ddc
             [Required] NamespaceId ns,
             [FromBody] BlobIdentifier[] bodyIds)
         {
-            ActionResult? result = await _requestHelper.HasAccessToNamespace(User, Request, ns, new [] { JupiterAclAction.ReadObject });
+            ActionResult? result = await _requestHelper.HasAccessToNamespace(User, Request, ns, new [] { DdcAclAction.ReadObject });
             if (result != null)
             {
                 return result;
@@ -149,7 +149,7 @@ namespace Horde.Server.Ddc
             [Required] NamespaceId ns,
             [Required] BlobIdentifier id)
         {
-            ActionResult? result = await _requestHelper.HasAccessToNamespace(User, Request, ns, new [] { JupiterAclAction.WriteObject });
+            ActionResult? result = await _requestHelper.HasAccessToNamespace(User, Request, ns, new [] { DdcAclAction.WriteObject });
             if (result != null)
             {
                 return result;
@@ -174,7 +174,7 @@ namespace Horde.Server.Ddc
             [Required] NamespaceId ns,
             [Required] BlobIdentifier id)
         {
-            ActionResult? result = await _requestHelper.HasAccessToNamespace(User, Request, ns, new [] { JupiterAclAction.ReadObject });
+            ActionResult? result = await _requestHelper.HasAccessToNamespace(User, Request, ns, new [] { DdcAclAction.ReadObject });
             if (result != null)
             {
                 return result;
@@ -226,7 +226,7 @@ namespace Horde.Server.Ddc
             [Required] NamespaceId ns,
             [Required] BlobIdentifier id)
         {
-            ActionResult? result = await _requestHelper.HasAccessToNamespace(User, Request, ns, new [] { JupiterAclAction.DeleteObject });
+            ActionResult? result = await _requestHelper.HasAccessToNamespace(User, Request, ns, new [] { DdcAclAction.DeleteObject });
             if (result != null)
             {
                 return result;
@@ -244,7 +244,7 @@ namespace Horde.Server.Ddc
         public async Task<IActionResult> DeleteNamespace(
             [Required] NamespaceId ns)
         {
-            ActionResult? result = await _requestHelper.HasAccessToNamespace(User, Request, ns, new [] { JupiterAclAction.DeleteNamespace });
+            ActionResult? result = await _requestHelper.HasAccessToNamespace(User, Request, ns, new [] { DdcAclAction.DeleteNamespace });
             if (result != null)
             {
                 return result;
diff --git a/Engine/Source/Programs/Horde/Horde.Server/Ddc/ReferencesController.cs b/Engine/Source/Programs/Horde/Horde.Server/Ddc/ReferencesController.cs
--- a/Engine/Source/Programs/Horde/Horde.Server/Ddc/ReferencesController.cs
+++ b/Engine/Source/Programs/Horde/Horde.Server/Ddc/ReferencesController.cs
@@ -17,6 +17,7 @@ using EpicGames.AspNet;
 using EpicGames.Core;
 using EpicGames.Horde.Storage;
 using EpicGames.Serialization;
+using Horde.Server.Acls;
 using JetBrains.Annotations;
 using Microsoft.AspNetCore.Authorization;
 using Microsoft.AspNetCore.Http;
@@ -80,7 +81,7 @@ namespace Horde.Server.Ddc
             List<NamespaceId> namespacesWithAccess = new();
             foreach (NamespaceId ns in namespaces)
             {
-                ActionResult? accessResult = await _requestHelper.HasAccessToNamespace(User, Request, ns, new [] { JupiterAclAction.ReadObject });
+                ActionResult? accessResult = await _requestHelper.HasAccessToNamespace(User, Request, ns, new [] { DdcAclAction.ReadObject });
                 if (accessResult == null)
                 {
                     namespacesWithAccess.Add(ns);
@@ -105,7 +106,7 @@ namespace Horde.Server.Ddc
             [FromRoute] [Required] IoHashKey key,
             [FromRoute] string? format = null)
         {
-            ActionResult? accessResult = await _requestHelper.HasAccessToNamespace(User, Request, ns, new [] { JupiterAclAction.ReadObject });
+            ActionResult? accessResult = await _requestHelper.HasAccessToNamespace(User, Request, ns, new [] { DdcAclAction.ReadObject });
             if (accessResult != null)
             {
                 return accessResult;
@@ -436,7 +437,7 @@ namespace Horde.Server.Ddc
         [FromRoute] [Required] IoHashKey key,
         [FromQuery] string[] fields)
     {
-        ActionResult? accessResult = await _requestHelper.HasAccessToNamespace(User, Request, ns, new [] { JupiterAclAction.ReadObject });
+        ActionResult? accessResult = await _requestHelper.HasAccessToNamespace(User, Request, ns, new [] { DdcAclAction.ReadObject });
         if (accessResult != null)
         {
             return accessResult;
@@ -477,7 +478,7 @@ namespace Horde.Server.Ddc
             [FromRoute] [Required] BucketId bucket,
             [FromRoute] [Required] IoHashKey key)
         {
-            ActionResult? accessResult = await _requestHelper.HasAccessToNamespace(User, Request, ns, new [] { JupiterAclAction.ReadObject });
+            ActionResult? accessResult = await _requestHelper.HasAccessToNamespace(User, Request, ns, new [] { DdcAclAction.ReadObject });
             if (accessResult != null)
             {
                 return accessResult;
@@ -541,7 +542,7 @@ namespace Horde.Server.Ddc
             [FromRoute] [Required] NamespaceId ns,
             [FromQuery] [Required] List<string> names)
         {
-            ActionResult? accessResult = await _requestHelper.HasAccessToNamespace(User, Request, ns, new [] { JupiterAclAction.ReadObject });
+            ActionResult? accessResult = await _requestHelper.HasAccessToNamespace(User, Request, ns, new [] { DdcAclAction.ReadObject });
             if (accessResult != null)
             {
                 return accessResult;
@@ -610,7 +611,7 @@ namespace Horde.Server.Ddc
             [FromRoute] [Required] BucketId bucket,
             [FromRoute] [Required] IoHashKey key)
         {
-            ActionResult? accessResult = await _requestHelper.HasAccessToNamespace(User, Request, ns, new [] { JupiterAclAction.WriteObject });
+            ActionResult? accessResult = await _requestHelper.HasAccessToNamespace(User, Request, ns, new [] { DdcAclAction.WriteObject });
             if (accessResult != null)
             {
                 return accessResult;
@@ -713,7 +714,7 @@ namespace Horde.Server.Ddc
             [FromRoute][Required] BucketId bucket,
             [FromRoute][Required] IoHashKey key)
         {
-            ActionResult? accessResult = await _requestHelper.HasAccessToNamespace(User, Request, ns, new [] { JupiterAclAction.WriteObject });
+            ActionResult? accessResult = await _requestHelper.HasAccessToNamespace(User, Request, ns, new [] { DdcAclAction.WriteObject });
             if (accessResult != null)
             {
                 return accessResult;
@@ -771,7 +772,7 @@ namespace Horde.Server.Ddc
             [FromRoute] [Required] IoHashKey key,
             [FromRoute] [Required] BlobIdentifier hash)
         {
-            ActionResult? accessResult = await _requestHelper.HasAccessToNamespace(User, Request, ns, new [] { JupiterAclAction.WriteObject });
+            ActionResult? accessResult = await _requestHelper.HasAccessToNamespace(User, Request, ns, new [] { DdcAclAction.WriteObject });
             if (accessResult != null)
             {
                 return accessResult;
@@ -802,22 +803,22 @@ namespace Horde.Server.Ddc
             [FromRoute] [Required] NamespaceId ns,
             [FromBody] [Required] RefBatchOps ops)
         {
-            JupiterAclAction ActionForOp(RefBatchOps.RefBatchOp.RefOperation op)
+            AclAction ActionForOp(RefBatchOps.RefBatchOp.RefOperation op)
             {
                 switch (op)
                 {
                     case RefBatchOps.RefBatchOp.RefOperation.GET:
-                        return JupiterAclAction.ReadObject;
+                        return DdcAclAction.ReadObject;
                     case RefBatchOps.RefBatchOp.RefOperation.PUT:
-                        return JupiterAclAction.WriteObject;
+                        return DdcAclAction.WriteObject;
                     case RefBatchOps.RefBatchOp.RefOperation.HEAD:
-                        return JupiterAclAction.ReadObject;
+                        return DdcAclAction.ReadObject;
                     default:
                         throw new ArgumentOutOfRangeException(nameof(op), op, null);
                 }
             }
 
-            JupiterAclAction[] requiredActions = ops.Ops.Select(op => ActionForOp(op.Op)).ToArray();
+            AclAction[] requiredActions = ops.Ops.Select(op => ActionForOp(op.Op)).ToArray();
 
             ActionResult? accessResult = await _requestHelper.HasAccessToNamespace(User, Request, ns, requiredActions);
             if (accessResult != null)
@@ -1003,7 +1004,7 @@ namespace Horde.Server.Ddc
             [FromRoute] [Required] NamespaceId ns
         )
         {
-            ActionResult? accessResult = await _requestHelper.HasAccessToNamespace(User, Request, ns, new [] { JupiterAclAction.DeleteNamespace });
+            ActionResult? accessResult = await _requestHelper.HasAccessToNamespace(User, Request, ns, new [] { DdcAclAction.DeleteNamespace });
             if (accessResult != null)
             {
                 return accessResult;
@@ -1032,7 +1033,7 @@ namespace Horde.Server.Ddc
             [FromRoute] [Required] NamespaceId ns,
             [FromRoute] [Required] BucketId bucket)
         {
-            ActionResult? accessResult = await _requestHelper.HasAccessToNamespace(User, Request, ns, new [] { JupiterAclAction.DeleteBucket });
+            ActionResult? accessResult = await _requestHelper.HasAccessToNamespace(User, Request, ns, new [] { DdcAclAction.DeleteBucket });
             if (accessResult != null)
             {
                 return accessResult;
@@ -1065,7 +1066,7 @@ namespace Horde.Server.Ddc
             [FromRoute] [Required] BucketId bucket,
             [FromRoute] [Required] IoHashKey key)
         {
-            ActionResult? accessResult = await _requestHelper.HasAccessToNamespace(User, Request, ns, new [] { JupiterAclAction.DeleteObject });
+            ActionResult? accessResult = await _requestHelper.HasAccessToNamespace(User, Request, ns, new [] { DdcAclAction.DeleteObject });
             if (accessResult != null)
             {
                 return accessResult;
diff --git a/Engine/Source/Programs/Horde/Horde.Server/Ddc/RequestHelper.cs b/Engine/Source/Programs/Horde/Horde.Server/Ddc/RequestHelper.cs
--- a/Engine/Source/Programs/Horde/Horde.Server/Ddc/RequestHelper.cs
+++ b/Engine/Source/Programs/Horde/Horde.Server/Ddc/RequestHelper.cs
@@ -4,6 +4,7 @@ using System;
 using System.Security.Claims;
 using System.Threading.Tasks;
 using EpicGames.Horde.Storage;
+using Horde.Server.Acls;
 using Microsoft.AspNetCore.Http;
 using Microsoft.AspNetCore.Mvc;
 
@@ -11,12 +12,12 @@ namespace Horde.Server.Ddc
 {
 	class RequestHelper : IRequestHelper
 	{
-		public Task<ActionResult?> HasAccessForGlobalOperations(ClaimsPrincipal user, JupiterAclAction[] aclActions)
+		public Task<ActionResult?> HasAccessForGlobalOperations(ClaimsPrincipal user, AclAction[] aclActions)
 		{
 			throw new NotImplementedException();
 		}
 
-		public Task<ActionResult?> HasAccessToNamespace(ClaimsPrincipal user, HttpRequest request, NamespaceId ns, JupiterAclAction[] aclActions)
+		public Task<ActionResult?> HasAccessToNamespace(ClaimsPrincipal user, HttpRequest request, NamespaceId ns, AclAction[] aclActions)
 		{
 			throw new NotImplementedException();
 		}
