Horde: DDC refactoring.

* Rename controllers to Ddc*
* Remove IBlobStore and IReferencesStore interfaces, which are used by Jupiter specific implementations.

diff --git a/Engine/Source/Programs/Horde/Horde.Server/Ddc/BlobController.cs b/Engine/Source/Programs/Horde/Horde.Server/Ddc/BlobController.cs
deleted file mode 100644
diff --git a/Engine/Source/Programs/Horde/Horde.Server/Ddc/BlobService.cs b/Engine/Source/Programs/Horde/Horde.Server/Ddc/BlobService.cs
--- a/Engine/Source/Programs/Horde/Horde.Server/Ddc/BlobService.cs
+++ b/Engine/Source/Programs/Horde/Horde.Server/Ddc/BlobService.cs
@@ -8,7 +8,7 @@ using EpicGames.Horde.Storage;
 
 namespace Horde.Server.Ddc
 {
-	class BlobService : IBlobService
+	class BlobService : IDdcBlobService
 	{
 		public Task DeleteNamespace(NamespaceId ns)
 		{
diff --git a/Engine/Source/Programs/Horde/Horde.Server/Ddc/BlobStore.cs b/Engine/Source/Programs/Horde/Horde.Server/Ddc/BlobStore.cs
deleted file mode 100644
diff --git a/Engine/Source/Programs/Horde/Horde.Server/Ddc/ContentIdController.cs b/Engine/Source/Programs/Horde/Horde.Server/Ddc/ContentIdController.cs
deleted file mode 100644
diff --git a/Engine/Source/Programs/Horde/Horde.Server/Ddc/ContentIdStore.cs b/Engine/Source/Programs/Horde/Horde.Server/Ddc/ContentIdStore.cs
--- a/Engine/Source/Programs/Horde/Horde.Server/Ddc/ContentIdStore.cs
+++ b/Engine/Source/Programs/Horde/Horde.Server/Ddc/ContentIdStore.cs
@@ -6,7 +6,7 @@ using System.Threading.Tasks;
 
 namespace Horde.Server.Ddc
 {
-	class ContentIdStore : IContentIdStore
+	class ContentIdStore : IDdcContentIdStore
 	{
 		public Task<BlobIdentifier[]?> Resolve(NamespaceId ns, ContentId contentId, bool mustBeContentId = false)
 		{
diff --git a/Engine/Source/Programs/Horde/Horde.Server/Ddc/DdcBlobsController.cs b/Engine/Source/Programs/Horde/Horde.Server/Ddc/DdcBlobsController.cs
new file mode 100644
--- /dev/null
+++ b/Engine/Source/Programs/Horde/Horde.Server/Ddc/DdcBlobsController.cs
@@ -0,0 +1,434 @@
+// Copyright Epic Games, Inc. All Rights Reserved.
+
+using System;
+using System.Collections.Concurrent;
+using System.Collections.Generic;
+using System.ComponentModel.DataAnnotations;
+using System.IO;
+using System.Linq;
+using System.Net;
+using System.Net.Mime;
+using System.Text.Json.Serialization;
+using System.Threading.Tasks;
+using EpicGames.AspNet;
+using EpicGames.Horde.Storage;
+using EpicGames.Serialization;
+using Horde.Server.Acls;
+using Microsoft.AspNetCore.Authorization;
+using Microsoft.AspNetCore.Http;
+using Microsoft.AspNetCore.Mvc;
+
+#pragma warning disable CS1591
+
+namespace Horde.Server.Ddc
+{
+	using IDiagnosticContext = Serilog.IDiagnosticContext;
+
+    [ApiController]
+    [Route("api/v1/s", Order = 1)]
+    [Route("api/v1/blobs", Order = 0)]
+    [Authorize]
+    public class DdcBlobsController : ControllerBase
+    {
+        private readonly IDdcBlobService _storage;
+        private readonly IDiagnosticContext _diagnosticContext;
+        private readonly IRequestHelper _requestHelper;
+        private readonly BufferedPayloadFactory _bufferedPayloadFactory;
+        private readonly NginxRedirectHelper _nginxRedirectHelper;
+
+        public DdcBlobsController(IDdcBlobService storage, IDiagnosticContext diagnosticContext, IRequestHelper requestHelper, BufferedPayloadFactory bufferedPayloadFactory, NginxRedirectHelper nginxRedirectHelper)
+        {
+            _storage = storage;
+            _diagnosticContext = diagnosticContext;
+            _requestHelper = requestHelper;
+            _bufferedPayloadFactory = bufferedPayloadFactory;
+            _nginxRedirectHelper = nginxRedirectHelper;
+        }
+
+        [HttpGet("{ns}/{id}")]
+        [ProducesDefaultResponseType]
+        public async Task<IActionResult> Get(
+            [Required] NamespaceId ns,
+            [Required] BlobIdentifier id,
+            [FromQuery] List<string>? storageLayers = null)
+        {
+            ActionResult? result = await _requestHelper.HasAccessToNamespace(User, Request, ns, new [] { DdcAclAction.ReadObject });
+            if (result != null)
+            {
+                return result;
+            }
+
+            try
+            {
+                BlobContents blobContents = await GetImpl(ns, id, storageLayers, supportsRedirectUri: true);
+
+                if (blobContents.RedirectUri != null)
+                {
+                    return Redirect(blobContents.RedirectUri.ToString());
+                }
+                if (_nginxRedirectHelper.CanRedirect(Request, blobContents))
+                {
+                    return _nginxRedirectHelper.CreateActionResult(blobContents, MediaTypeNames.Application.Octet);
+                }
+                return File(blobContents.Stream, MediaTypeNames.Application.Octet, enableRangeProcessing: true);
+            }
+            catch (BlobNotFoundException e)
+            {
+                return NotFound(new ValidationProblemDetails { Title = $"Blob {e.Blob} not found"});
+            }
+            catch (AuthorizationException e)
+            {
+                return e.Result;
+            }
+        }
+        
+        [HttpHead("{ns}/{id}")]
+        [ProducesDefaultResponseType]
+        public async Task<IActionResult> Head(
+            [Required] NamespaceId ns,
+            [Required] BlobIdentifier id,
+            [FromQuery] List<string>? storageLayers = null)
+        {
+            ActionResult? result = await _requestHelper.HasAccessToNamespace(User, Request, ns, new [] { DdcAclAction.ReadObject });
+            if (result != null)
+            {
+                return result;
+            }
+            bool exists = await _storage.Exists(ns, id, storageLayers);
+
+            if (!exists)
+            {
+                return NotFound(new ValidationProblemDetails { Title = $"Blob {id} not found"});
+            }
+
+            return Ok();
+        }
+
+        [HttpPost("{ns}/exists")]
+        [ProducesDefaultResponseType]
+        public async Task<IActionResult> ExistsMultiple(
+            [Required] NamespaceId ns,
+            [Required] [FromQuery] List<BlobIdentifier> id)
+        {
+            ActionResult? result = await _requestHelper.HasAccessToNamespace(User, Request, ns, new [] { DdcAclAction.ReadObject });
+            if (result != null)
+            {
+                return result;
+            }
+
+            ConcurrentBag<BlobIdentifier> missingBlobs = new ConcurrentBag<BlobIdentifier>();
+
+            IEnumerable<Task> tasks = id.Select(async blob =>
+            {
+                if (!await _storage.Exists(ns, blob))
+                {
+                    missingBlobs.Add(blob);
+                }
+            });
+            await Task.WhenAll(tasks);
+
+            return Ok(new HeadMultipleResponse { Needs = missingBlobs.ToArray()});
+        }
+
+        [HttpPost("{ns}/exist")]
+        [ProducesDefaultResponseType]
+        public async Task<IActionResult> ExistsBody(
+            [Required] NamespaceId ns,
+            [FromBody] BlobIdentifier[] bodyIds)
+        {
+            ActionResult? result = await _requestHelper.HasAccessToNamespace(User, Request, ns, new [] { DdcAclAction.ReadObject });
+            if (result != null)
+            {
+                return result;
+            }
+
+            ConcurrentBag<BlobIdentifier> missingBlobs = new ConcurrentBag<BlobIdentifier>();
+
+            IEnumerable<Task> tasks = bodyIds.Select(async blob =>
+            {
+                if (!await _storage.Exists(ns, blob))
+                {
+                    missingBlobs.Add(blob);
+                }
+            });
+            await Task.WhenAll(tasks);
+
+            return Ok(new HeadMultipleResponse { Needs = missingBlobs.ToArray()});
+        }
+
+        private async Task<BlobContents> GetImpl(NamespaceId ns, BlobIdentifier blob, List<string>? storageLayers = null, bool supportsRedirectUri = false)
+        {
+            try
+            {
+                return await _storage.GetObject(ns, blob, storageLayers, supportsRedirectUri);
+            }
+            catch (BlobNotFoundException)
+            {
+                if (!_storage.ShouldFetchBlobOnDemand(ns))
+                {
+                    throw;
+                }
+
+                return await _storage.ReplicateObject(ns, blob);
+            }
+        }
+
+        [HttpPut("{ns}/{id}")]
+        [RequiredContentType(MediaTypeNames.Application.Octet)]
+        [DisableRequestSizeLimit]
+        public async Task<IActionResult> Put(
+            [Required] NamespaceId ns,
+            [Required] BlobIdentifier id)
+        {
+            ActionResult? result = await _requestHelper.HasAccessToNamespace(User, Request, ns, new [] { DdcAclAction.WriteObject });
+            if (result != null)
+            {
+                return result;
+            }
+
+            _diagnosticContext.Set("Content-Length", Request.ContentLength ?? -1);
+
+            try
+            {
+                Uri? uri = await _storage.MaybePutObjectWithRedirect(ns, id);
+                if (uri != null)
+                {
+                    return Ok(new
+                    {
+                        Identifier = id.ToString(),
+                        RedirectUri = uri,
+                    });
+                }
+                using IBufferedPayload payload = await _bufferedPayloadFactory.CreateFromRequest(Request);
+
+                BlobIdentifier identifier = await _storage.PutObject(ns, payload, id);
+                return Ok(new
+                {
+                    Identifier = identifier.ToString()
+                });
+            }
+            catch (ResourceHasToManyRequestsException)
+            {
+                return StatusCode(StatusCodes.Status429TooManyRequests);
+            }
+            catch (ClientSendSlowException e)
+            {
+                return Problem(e.Message, null, (int)HttpStatusCode.RequestTimeout);
+            }
+        }
+
+        [HttpPost("{ns}")]
+        [RequiredContentType(MediaTypeNames.Application.Octet)]
+        [DisableRequestSizeLimit]
+        public async Task<IActionResult> Post(
+            [Required] NamespaceId ns)
+        {
+            ActionResult? result = await _requestHelper.HasAccessToNamespace(User, Request, ns, new [] { DdcAclAction.WriteObject });
+            if (result != null)
+            {
+                return result;
+            }
+
+            _diagnosticContext.Set("Content-Length", Request.ContentLength ?? -1);
+            try
+            {
+                using IBufferedPayload payload = await _bufferedPayloadFactory.CreateFromRequest(Request);
+
+                await using Stream stream = payload.GetStream();
+
+                BlobIdentifier id = await BlobIdentifier.FromStream(stream);
+                await _storage.PutObjectKnownHash(ns, payload, id);
+                
+                return Ok(new
+                {
+                    Identifier = id.ToString()
+                });
+            }
+            catch (ClientSendSlowException e)
+            {
+                return Problem(e.Message, null, (int)HttpStatusCode.RequestTimeout);
+            }
+        }
+
+        [HttpDelete("{ns}/{id}")]
+        public async Task<IActionResult> Delete(
+            [Required] NamespaceId ns,
+            [Required] BlobIdentifier id)
+        {
+            ActionResult? result = await _requestHelper.HasAccessToNamespace(User, Request, ns, new [] { DdcAclAction.DeleteObject });
+            if (result != null)
+            {
+                return result;
+            }
+
+            await DeleteImpl(ns, id);
+
+            return NoContent();
+        }
+
+        
+        [HttpDelete("{ns}")]
+        public async Task<IActionResult> DeleteNamespace(
+            [Required] NamespaceId ns)
+        {
+            ActionResult? result = await _requestHelper.HasAccessToNamespace(User, Request, ns, new [] { DdcAclAction.DeleteNamespace });
+            if (result != null)
+            {
+                return result;
+            }
+
+            await  _storage.DeleteNamespace(ns);
+
+            return NoContent();
+        }
+
+        private async Task DeleteImpl(NamespaceId ns, BlobIdentifier id)
+        {
+            await _storage.DeleteObject(ns, id);
+        }
+
+        // ReSharper disable UnusedAutoPropertyAccessor.Global
+        // ReSharper disable once ClassNeverInstantiated.Global
+        [System.Diagnostics.CodeAnalysis.SuppressMessage("Design", "CA1034:Nested types should not be visible", Justification = "Used by serialization")]
+        public class BatchOp
+        {
+            // ReSharper disable once InconsistentNaming
+            public enum Operation
+            {
+                INVALID,
+                GET,
+                PUT,
+                DELETE,
+                HEAD
+            }
+
+            [Required] public NamespaceId? Namespace { get; set; }
+
+            public BlobIdentifier? Id { get; set; }
+
+            [Required] public Operation Op { get; set; }
+
+            public byte[]? Content { get; set; }
+        }
+
+        [System.Diagnostics.CodeAnalysis.SuppressMessage("Design", "CA1034:Nested types should not be visible", Justification = "Used by serialization")]
+        public class BatchCall
+        {
+            public BatchOp[]? Operations { get; set; }
+        }
+        // ReSharper restore UnusedAutoPropertyAccessor.Global
+
+        [HttpPost("")]
+        public async Task<IActionResult> Post([FromBody] BatchCall batch)
+        {
+            static AclAction MapToAclAction(BatchOp.Operation op)
+            {
+                switch (op)
+                {
+                    case BatchOp.Operation.GET:
+                    case BatchOp.Operation.HEAD:
+                        return DdcAclAction.ReadObject;
+                    case BatchOp.Operation.PUT:
+                        return DdcAclAction.WriteObject;
+                    case BatchOp.Operation.DELETE:
+                        return DdcAclAction.DeleteObject;
+                    default:
+                        throw new ArgumentOutOfRangeException(nameof(op), op, null);
+                }
+            }
+
+            if (batch?.Operations == null)
+            {
+                throw new ();
+            }
+
+            Task<object?>[] tasks = new Task<object?>[batch.Operations.Length];
+            for (int index = 0; index < batch.Operations.Length; index++)
+            {
+                BatchOp op = batch.Operations[index];
+                if (op.Namespace == null)
+                {
+                    throw new Exception(nameof(op.Namespace));
+                }
+
+                ActionResult? result = await _requestHelper.HasAccessToNamespace(User, Request, op.Namespace!.Value, new[] { MapToAclAction(op.Op) });
+
+                if (result != null)
+                {
+                    return result;
+                }
+
+                switch (op.Op)
+                {
+                    case BatchOp.Operation.INVALID:
+                        throw new InvalidOperationException($"Op type {BatchOp.Operation.INVALID} is not a valid op type for operation id {op.Id}");
+                    case BatchOp.Operation.GET:
+                        if (op.Id == null)
+                        {
+                            return BadRequest();
+                        }
+
+                        tasks[index] = GetImpl(op.Namespace.Value, op.Id).ContinueWith((t, _) =>
+                        {
+                            // TODO: This is very allocation heavy but given that the end result is a json object we can not really stream this anyway
+                            using BlobContents blobContents = t.Result;
+
+                            using MemoryStream ms = new MemoryStream();
+                            blobContents.Stream.CopyTo(ms);
+                            ms.Seek(0, SeekOrigin.Begin);
+                            string str = Convert.ToBase64String(ms.ToArray());
+                            return (object?) str;
+                        }, null, TaskScheduler.Current);
+                        break;
+                    case BatchOp.Operation.HEAD:
+                        if (op.Id == null)
+                        {
+                            return BadRequest();
+                        }
+
+                        tasks[index] = _storage.Exists(op.Namespace.Value, op.Id)
+                            .ContinueWith((t,_) => t.Result ? (object?) null : op.Id, null, TaskScheduler.Current);
+                        break;
+                    case BatchOp.Operation.PUT:
+                        {
+                            if (op.Content == null)
+                            {
+                                return BadRequest();
+                            }
+
+                            if (op.Id == null)
+                            {
+                                return BadRequest();
+                            }
+
+                            using MemoryBufferedPayload payload = new MemoryBufferedPayload(op.Content);
+                            tasks[index] = _storage.PutObject(op.Namespace.Value, payload, op.Id).ContinueWith((t, _) => (object?) t.Result, null, TaskScheduler.Current);
+                            break;
+                        }
+                    case BatchOp.Operation.DELETE:
+                        if (op.Id == null)
+                        {
+                            return BadRequest();
+                        }
+
+                        tasks[index] = DeleteImpl(op.Namespace.Value, op.Id).ContinueWith((t, _) => (object?) null, null, TaskScheduler.Current);
+                        break;
+                    default:
+                        throw new NotImplementedException($"{op.Op} is not a support op type");
+                }
+            }
+
+            await Task.WhenAll(tasks);
+
+            object?[] results = tasks.Select(t => t.Result).ToArray();
+
+            return Ok(results);
+        }
+    }
+
+    public class HeadMultipleResponse
+    {
+        [CbField("needs")]
+        public BlobIdentifier[] Needs { get; set; } = null!;
+    }
+}
diff --git a/Engine/Source/Programs/Horde/Horde.Server/Ddc/DdcContentIdController.cs b/Engine/Source/Programs/Horde/Horde.Server/Ddc/DdcContentIdController.cs
new file mode 100644
--- /dev/null
+++ b/Engine/Source/Programs/Horde/Horde.Server/Ddc/DdcContentIdController.cs
@@ -0,0 +1,87 @@
+// Copyright Epic Games, Inc. All Rights Reserved.
+
+using System.Net.Mime;
+using System.Threading.Tasks;
+using EpicGames.AspNet;
+using EpicGames.Horde.Storage;
+using Microsoft.AspNetCore.Authorization;
+using Microsoft.AspNetCore.Mvc;
+
+#pragma warning disable CS1591
+
+namespace Horde.Server.Ddc
+{
+	[ApiController]
+    [FormatFilter]
+    [Produces(MediaTypeNames.Application.Json, MediaTypeNames.Application.Octet, CustomMediaTypeNames.UnrealCompactBinary)]
+    [Route("api/v1/content-id")]
+    [Authorize]
+    public class DdcContentIdController : ControllerBase
+    {
+        private readonly IRequestHelper _requestHelper;
+        private readonly IDdcContentIdStore _contentIdStore;
+
+        public DdcContentIdController(IRequestHelper requestHelper, IDdcContentIdStore contentIdStore)
+        {
+            _requestHelper = requestHelper;
+            _contentIdStore = contentIdStore;
+        }
+
+        /// <summary>
+        /// Returns which blobs a content id maps to
+        /// </summary>
+        /// <param name="ns">Namespace. Each namespace is completely separated from each other. Use for different types of data that is never expected to be similar (between two different games for instance). Example: `uc4.ddc`</param>
+        /// <param name="contentId">The content id to resolve </param>
+        [HttpGet("{ns}/{contentId}.{format?}", Order = 500)]
+        [ProducesDefaultResponseType]
+        [ProducesResponseType(type: typeof(ProblemDetails), 400)]
+        public async Task<IActionResult> Resolve(NamespaceId ns, ContentId contentId)
+        {
+            ActionResult? result = await _requestHelper.HasAccessToNamespace(User, Request, ns, new [] { DdcAclAction.ReadObject });
+            if (result != null)
+            {
+                return result;
+            }
+            BlobIdentifier[]? blobs = await _contentIdStore.Resolve(ns, contentId, mustBeContentId: true);
+
+            if (blobs == null)
+            {
+                return NotFound(new ProblemDetails { Title = $"Unable to resolve content id {contentId} ({ns})." });
+            }
+
+            return Ok(new ResolvedContentIdResponse(blobs));
+        }
+
+        /// <summary>
+        /// Update a content id mapping
+        /// </summary>
+        /// <param name="ns">Namespace. Each namespace is completely separated from each other. Use for different types of data that is never expected to be similar (between two different games for instance). Example: `uc4.ddc`</param>
+        /// <param name="contentId">The content id to resolve </param>
+        /// <param name="blobIdentifier">The blob identifier to map the content id to</param>
+        /// <param name="contentWeight">The weight of this mapping, higher means more important</param>
+        [HttpPut("{ns}/{contentId}/update/{blobIdentifier}/{contentWeight}", Order = 500)]
+        [ProducesDefaultResponseType]
+        [ProducesResponseType(type: typeof(ProblemDetails), 400)]
+        public async Task<IActionResult> UpdateContentIdMapping(NamespaceId ns, ContentId contentId, BlobIdentifier blobIdentifier, int contentWeight)
+        {
+            ActionResult? result = await _requestHelper.HasAccessToNamespace(User, Request, ns, new [] { DdcAclAction.WriteObject });
+            if (result != null)
+            {
+                return result;
+            }
+            await _contentIdStore.Put(ns, contentId, blobIdentifier, contentWeight);
+
+            return Ok();
+        }
+    }
+
+    public class ResolvedContentIdResponse
+    {
+        public BlobIdentifier[] Blobs { get; set; }
+
+        public ResolvedContentIdResponse(BlobIdentifier[] blobs)
+        {
+            Blobs = blobs;
+        }
+    }
+}
diff --git a/Engine/Source/Programs/Horde/Horde.Server/Ddc/DdcObjectsController.cs b/Engine/Source/Programs/Horde/Horde.Server/Ddc/DdcObjectsController.cs
new file mode 100644
--- /dev/null
+++ b/Engine/Source/Programs/Horde/Horde.Server/Ddc/DdcObjectsController.cs
@@ -0,0 +1,295 @@
+﻿// Copyright Epic Games, Inc. All Rights Reserved.
+
+using System;
+using System.Collections.Concurrent;
+using System.Collections.Generic;
+using System.ComponentModel.DataAnnotations;
+using System.Linq;
+using System.Net;
+using System.Net.Mime;
+using System.Threading.Tasks;
+using EpicGames.AspNet;
+using EpicGames.Horde.Storage;
+using EpicGames.Serialization;
+using Microsoft.AspNetCore.Authorization;
+using Microsoft.AspNetCore.Mvc;
+using Microsoft.Extensions.Logging;
+
+#pragma warning disable CS1591
+
+namespace Horde.Server.Ddc
+{
+	using IDiagnosticContext = Serilog.IDiagnosticContext;
+
+    [ApiController]
+    [Route("api/v1/objects", Order = 0)]
+    [Authorize]
+    [Produces(CustomMediaTypeNames.UnrealCompactBinary, MediaTypeNames.Application.Json)]
+    public class DdcObjectsController : ControllerBase
+    {
+        private readonly IDdcBlobService _storage;
+        private readonly IDiagnosticContext _diagnosticContext;
+        private readonly IRequestHelper _requestHelper;
+        private readonly IReferenceResolver _referenceResolver;
+        private readonly BufferedPayloadFactory _bufferedPayloadFactory;
+
+        private readonly ILogger _logger;
+
+        public DdcObjectsController(IDdcBlobService storage, IDiagnosticContext diagnosticContext, IRequestHelper requestHelper, IReferenceResolver referenceResolver, BufferedPayloadFactory bufferedPayloadFactory, ILogger<DdcObjectsController> logger)
+        {
+            _storage = storage;
+            _diagnosticContext = diagnosticContext;
+            _requestHelper = requestHelper;
+            _referenceResolver = referenceResolver;
+            _bufferedPayloadFactory = bufferedPayloadFactory;
+            _logger = logger;
+        }
+
+        [HttpGet("{ns}/{id}")]
+        [ProducesDefaultResponseType]
+        public async Task<IActionResult> Get(
+            [Required] NamespaceId ns,
+            [Required] BlobIdentifier id)
+        {
+            ActionResult? result = await _requestHelper.HasAccessToNamespace(User, Request, ns, new [] { DdcAclAction.ReadObject });
+            if (result != null)
+            {
+                return result;
+            }
+
+            try
+            {
+                BlobContents blobContents = await _storage.GetObject(ns, id);
+
+                return File(blobContents.Stream, CustomMediaTypeNames.UnrealCompactBinary);
+            }
+            catch (BlobNotFoundException e)
+            {
+                return NotFound(new ValidationProblemDetails {Title = $"Object {e.Blob} not found"});
+            }
+        }
+
+        [HttpHead("{ns}/{id}")]
+        [ProducesDefaultResponseType]
+        public async Task<IActionResult> Head(
+            [Required] NamespaceId ns,
+            [Required] BlobIdentifier id)
+        {
+            ActionResult? result = await _requestHelper.HasAccessToNamespace(User, Request, ns, new [] { DdcAclAction.ReadObject });
+            if (result != null)
+            {
+                return result;
+            }
+
+            bool exists = await _storage.Exists(ns, id);
+
+            if (!exists)
+            {
+                return NotFound(new ValidationProblemDetails {Title = $"Object {id} not found"});
+            }
+
+            return Ok();
+        }
+
+        [HttpPost("{ns}/exists")]
+        [ProducesDefaultResponseType]
+        public async Task<IActionResult> ExistsMultiple(
+            [Required] NamespaceId ns,
+            [Required] [FromQuery] List<BlobIdentifier> id)
+        {
+            ActionResult? result = await _requestHelper.HasAccessToNamespace(User, Request, ns, new [] { DdcAclAction.ReadObject });
+            if (result != null)
+            {
+                return result;
+            }
+
+            ConcurrentBag<BlobIdentifier> missingBlobs = new ConcurrentBag<BlobIdentifier>();
+
+            IEnumerable<Task> tasks = id.Select(async blob =>
+            {
+                if (!await _storage.Exists(ns, blob))
+                {
+                    missingBlobs.Add(blob);
+                }
+            });
+            await Task.WhenAll(tasks);
+
+            return Ok(new HeadMultipleResponse {Needs = missingBlobs.ToArray()});
+        }
+
+        [HttpPost("{ns}/exist")]
+        [ProducesDefaultResponseType]
+        public async Task<IActionResult> ExistsBody(
+            [Required] NamespaceId ns,
+            [FromBody] BlobIdentifier[] bodyIds)
+        {
+            ActionResult? result = await _requestHelper.HasAccessToNamespace(User, Request, ns, new [] { DdcAclAction.ReadObject });
+            if (result != null)
+            {
+                return result;
+            }
+
+            ConcurrentBag<BlobIdentifier> missingBlobs = new ConcurrentBag<BlobIdentifier>();
+
+            IEnumerable<Task> tasks = bodyIds.Select(async blob =>
+            {
+                if (!await _storage.Exists(ns, blob))
+                {
+                    missingBlobs.Add(blob);
+                }
+            });
+            await Task.WhenAll(tasks);
+
+            return Ok(new HeadMultipleResponse { Needs = missingBlobs.ToArray() });
+        }
+
+        [HttpPut("{ns}/{id}")]
+        [RequiredContentType(CustomMediaTypeNames.UnrealCompactBinary)]
+        public async Task<IActionResult> Put(
+            [Required] NamespaceId ns,
+            [Required] BlobIdentifier id)
+        {
+            ActionResult? result = await _requestHelper.HasAccessToNamespace(User, Request, ns, new [] { DdcAclAction.WriteObject });
+            if (result != null)
+            {
+                return result;
+            }
+
+            _diagnosticContext.Set("Content-Length", Request.ContentLength ?? -1);
+            try
+            {
+                using IBufferedPayload payload = await _bufferedPayloadFactory.CreateFromRequest(Request);
+
+                BlobIdentifier identifier = await _storage.PutObject(ns, payload, id);
+                return Ok(new PutBlobResponse(identifier));
+            }
+            catch (ClientSendSlowException e)
+            {
+                return Problem(e.Message, null, (int)HttpStatusCode.RequestTimeout);
+            }
+        }
+
+        [HttpGet("{ns}/{id}/references")]
+        public async Task<IActionResult> ResolveReferences(
+            [Required] NamespaceId ns,
+            [Required] BlobIdentifier id)
+        {
+            ActionResult? result = await _requestHelper.HasAccessToNamespace(User, Request, ns, new [] { DdcAclAction.ReadObject });
+            if (result != null)
+            {
+                return result;
+            }
+
+            BlobContents blob;
+            try
+            {
+                blob = await _storage.GetObject(ns, id);
+            }
+            catch (BlobNotFoundException e)
+            {
+                return NotFound(new ValidationProblemDetails {Title = $"Object {e.Blob} not found"});
+            }
+           
+            byte[] blobContents = await blob.Stream.ToByteArray();
+            if (blobContents.Length == 0)
+            {
+                _logger.LogWarning("0 byte object found for {Id} {Namespace}", id, ns);
+            }
+
+            CbObject compactBinaryObject;
+            try
+            {
+                compactBinaryObject = new CbObject(blobContents);
+            }
+            catch (IndexOutOfRangeException)
+            {
+                return Problem(title: $"{id} was not a proper compact binary object.", detail: "Index out of range");
+            }
+
+            try
+            {
+                BlobIdentifier[] references = await _referenceResolver.GetReferencedBlobs(ns, compactBinaryObject).ToArrayAsync();
+                return Ok(new ResolvedReferencesResult(references));
+            }
+            catch (PartialReferenceResolveException e)
+            {
+                return BadRequest(new ValidationProblemDetails {Title = $"Object {id} is missing content ids", Detail = $"Following content ids are invalid: {string.Join(",", e.UnresolvedReferences)}"});
+            }
+            catch (ReferenceIsMissingBlobsException e)
+            {
+                return BadRequest(new ValidationProblemDetails {Title = $"Object {id} is missing blobs", Detail = $"Following blobs are missing: {string.Join(",", e.MissingBlobs)}"});
+            }
+        }
+
+        [HttpDelete("{ns}/{id}")]
+        public async Task<IActionResult> Delete(
+            [Required] NamespaceId ns,
+            [Required] BlobIdentifier id)
+        {
+            ActionResult? result = await _requestHelper.HasAccessToNamespace(User, Request, ns, new [] { DdcAclAction.DeleteObject });
+            if (result != null)
+            {
+                return result;
+            }
+
+            await _storage.DeleteObject(ns, id);
+
+            return Ok( new DeletedResponse
+            {
+                DeletedCount = 1
+            });
+        }
+
+        [HttpDelete("{ns}")]
+        public async Task<IActionResult> DeleteNamespace(
+            [Required] NamespaceId ns)
+        {
+            ActionResult? result = await _requestHelper.HasAccessToNamespace(User, Request, ns, new [] { DdcAclAction.DeleteNamespace });
+            if (result != null)
+            {
+                return result;
+            }
+
+            await _storage.DeleteNamespace(ns);
+
+            return Ok();
+        }
+    }
+
+    public class PutBlobResponse
+    {
+        public PutBlobResponse()
+        {
+            Identifier = null!;
+        }
+
+        public PutBlobResponse(BlobIdentifier identifier)
+        {
+            Identifier = identifier;
+        }
+
+        [CbField("identifier")]
+        public BlobIdentifier Identifier { get; set; }
+    }
+
+    public class DeletedResponse
+    {
+        public int DeletedCount { get; set; }
+    }
+
+    public class ResolvedReferencesResult
+    {
+        public ResolvedReferencesResult()
+        {
+            References = null!;
+        }
+
+        public ResolvedReferencesResult(BlobIdentifier[] references)
+        {
+            References = references;
+        }
+
+        [CbField("references")]
+        public BlobIdentifier[] References { get; set; }
+    }
+}
diff --git a/Engine/Source/Programs/Horde/Horde.Server/Ddc/DdcRefsController.cs b/Engine/Source/Programs/Horde/Horde.Server/Ddc/DdcRefsController.cs
new file mode 100644
--- /dev/null
+++ b/Engine/Source/Programs/Horde/Horde.Server/Ddc/DdcRefsController.cs
@@ -0,0 +1,1345 @@
+// Copyright Epic Games, Inc. All Rights Reserved.
+
+using System;
+using System.Collections.Concurrent;
+using System.Collections.Generic;
+using System.ComponentModel.DataAnnotations;
+using System.Globalization;
+using System.IO;
+using System.Linq;
+using System.Net;
+using System.Net.Mime;
+using System.Text;
+using System.Text.Json.Serialization;
+using System.Threading;
+using System.Threading.Tasks;
+using EpicGames.AspNet;
+using EpicGames.Core;
+using EpicGames.Horde.Storage;
+using EpicGames.Serialization;
+using Horde.Server.Acls;
+using JetBrains.Annotations;
+using Microsoft.AspNetCore.Authorization;
+using Microsoft.AspNetCore.Http;
+using Microsoft.AspNetCore.Http.Extensions;
+using Microsoft.AspNetCore.Mvc;
+using Microsoft.Extensions.DependencyInjection;
+using Microsoft.Extensions.Logging;
+using OpenTelemetry.Trace;
+
+#pragma warning disable CS1591
+
+namespace Horde.Server.Ddc
+{
+	using IDiagnosticContext = Serilog.IDiagnosticContext;
+
+    [ApiController]
+    [FormatFilter]
+    [Produces(MediaTypeNames.Application.Json, MediaTypeNames.Application.Octet, CustomMediaTypeNames.UnrealCompactBinary)]
+    [Route("api/v1/refs")]
+    [Authorize]
+    public class DdcRefsController : ControllerBase
+    {
+        private readonly IDiagnosticContext _diagnosticContext;
+        private readonly FormatResolver _formatResolver;
+        private readonly BufferedPayloadFactory _bufferedPayloadFactory;
+        private readonly IReferenceResolver _referenceResolver;
+        private readonly NginxRedirectHelper _nginxRedirectHelper;
+        private readonly IRequestHelper _requestHelper;
+        private readonly Tracer _tracer;
+
+        private readonly ILogger _logger;
+        private readonly IDdcObjectService _objectService;
+        private readonly IDdcBlobService _blobStore;
+
+        public DdcRefsController(IDdcObjectService objectService, IDdcBlobService blobStore, IDiagnosticContext diagnosticContext, FormatResolver formatResolver, BufferedPayloadFactory bufferedPayloadFactory, IReferenceResolver referenceResolver, NginxRedirectHelper nginxRedirectHelper, IRequestHelper requestHelper, Tracer tracer, ILogger<DdcRefsController> logger)
+        {
+            _objectService = objectService;
+            _blobStore = blobStore;
+            _diagnosticContext = diagnosticContext;
+            _formatResolver = formatResolver;
+            _bufferedPayloadFactory = bufferedPayloadFactory;
+            _referenceResolver = referenceResolver;
+            _nginxRedirectHelper = nginxRedirectHelper;
+            _requestHelper = requestHelper;
+            _tracer = tracer;
+            _logger = logger;
+        }
+
+        /// <summary>
+        /// Returns all the known namespace the token has access to
+        /// </summary>
+        /// <returns></returns>
+        [HttpGet("")]
+        [ProducesDefaultResponseType]
+        [ProducesResponseType(type: typeof(ProblemDetails), 400)]
+        public async Task<IActionResult> GetNamespaces()
+        {
+            NamespaceId[] namespaces = await _objectService.GetNamespaces().ToArrayAsync();
+
+            // filter namespaces down to only the namespaces the user has access to
+            List<NamespaceId> namespacesWithAccess = new();
+            foreach (NamespaceId ns in namespaces)
+            {
+                ActionResult? accessResult = await _requestHelper.HasAccessToNamespace(User, Request, ns, new [] { DdcAclAction.ReadObject });
+                if (accessResult == null)
+                {
+                    namespacesWithAccess.Add(ns);
+                }
+            }
+
+            return Ok(new GetNamespacesResponse(namespacesWithAccess.ToArray()));
+        }
+
+        /// <summary>
+        /// Returns a refs key
+        /// </summary>
+        /// <param name="ns">Namespace. Each namespace is completely separated from each other. Use for different types of data that is never expected to be similar (between two different games for instance). Example: `uc4.ddc`</param>
+        /// <param name="bucket">The category/type of record you are caching. Is a clustered key together with the actual key, but all records in the same bucket can be dropped easily. Example: `terrainTexture` </param>
+        /// <param name="key">The unique name of this particular key. `iAmAVeryValidKey`</param>
+        /// <param name="format">Optional specifier to set which output format is used json/raw/cb</param>
+        [HttpGet("{ns}/{bucket}/{key}.{format?}", Order = 500)]
+        [Produces(MediaTypeNames.Application.Json, MediaTypeNames.Application.Octet, CustomMediaTypeNames.UnrealCompactBinary, CustomMediaTypeNames.JupiterInlinedPayload, CustomMediaTypeNames.UnrealCompactBinaryPackage)]
+        public async Task<IActionResult> Get(
+            [FromRoute] [Required] NamespaceId ns,
+            [FromRoute] [Required] BucketId bucket,
+            [FromRoute] [Required] IoHashKey key,
+            [FromRoute] string? format = null)
+        {
+            ActionResult? accessResult = await _requestHelper.HasAccessToNamespace(User, Request, ns, new [] { DdcAclAction.ReadObject });
+            if (accessResult != null)
+            {
+                return accessResult;
+            }
+
+            try
+            {
+                (ObjectRecord objectRecord, BlobContents? blob) = await _objectService.Get(ns, bucket, key, Array.Empty<string>());
+
+                if (blob == null)
+                {
+                    throw new InvalidOperationException($"Blob was null when attempting to fetch {ns} {bucket} {key}");
+                }
+
+                if (!objectRecord.IsFinalized)
+                {
+                    // we do not consider un-finalized objects as valid
+                    return BadRequest(new ProblemDetails { Title = $"Object {objectRecord.Bucket} {objectRecord.Name} is not finalized." });
+                }
+
+                Response.Headers[CommonHeaders.HashHeaderName] = objectRecord.BlobIdentifier.ToString();
+                Response.Headers[CommonHeaders.LastAccessHeaderName] = objectRecord.LastAccess.ToString(CultureInfo.InvariantCulture);
+
+                async Task WriteBody(BlobContents blobContents, string contentType)
+                {
+                    IServerTiming? serverTiming = Request.HttpContext.RequestServices.GetService<IServerTiming>();
+                    using ServerTimingMetricScoped? serverTimingScope = serverTiming?.CreateServerTimingMetricScope("body.write", "Time spent writing body");
+
+                    long contentLength = blobContents.Length;
+                    using TelemetrySpan scope = _tracer.StartActiveSpan("body.write").SetAttribute("operation.name", "body.write");
+                    scope.SetAttribute("content-length", contentLength);
+                    const int BufferSize = 64 * 1024;
+                    Stream outputStream = Response.Body;
+                    Response.ContentLength = contentLength;
+                    Response.ContentType = contentType;
+                    Response.StatusCode = StatusCodes.Status200OK;
+                    try
+                    {
+                        await StreamCopyOperation.CopyToAsync(blobContents.Stream, outputStream, count: null, bufferSize: BufferSize, cancel: Response.HttpContext.RequestAborted);
+                    }
+                    catch (OperationCanceledException)
+                    {
+                        // do not raise exceptions for cancelled writes
+                        // as we have already started writing a response we can not change the status code
+                        // so we just drop a warning and proceed
+                        _logger.LogWarning("The operation was canceled while writing the body");
+                    }
+                }
+
+                string responseType = _formatResolver.GetResponseType(Request, format, CustomMediaTypeNames.UnrealCompactBinary);
+                Tracer.CurrentSpan.SetAttribute("response-type", responseType);
+                switch (responseType)
+                {
+                    case CustomMediaTypeNames.UnrealCompactBinary:
+                    {
+                        // for compact binary we can just serialize our internal object
+                        await WriteBody(blob, CustomMediaTypeNames.UnrealCompactBinary);
+
+                        break;
+                    }
+                    case MediaTypeNames.Application.Octet:
+                    {
+                        byte[] blobMemory = await blob.Stream.ToByteArray();
+                        CbObject cb = new CbObject(blobMemory);
+
+                        (int,CbField?) CountFields(CbObject o)
+                        {
+                            int count = 0;
+                            CbField? foundField = null;
+                            cb.IterateAttachments(field =>
+                            {
+                                ++count;
+                                if (field.IsBinaryAttachment())
+                                {
+                                    foundField = field;
+                                }
+                            });
+
+                            return (count, foundField);
+                        }
+
+                        // breaking lambda call into private method to workaround incorrect triggering of CA1508 - [REDACTED_PROFILE_URL]/roslyn-analyzers/issues/5254
+                        (int countOfAttachmentFields,CbField? binaryAttachmentField) = CountFields(cb);
+
+                        if (countOfAttachmentFields == 1 && binaryAttachmentField != null)
+                        {
+                            // there is a single attachment field and that is of the binary attachment type, fetch that attachment and return it instead of the compact binary
+                            // this is so that we match the uploaded that happened as a octet-stream which generates a small cb object with a single attachment
+
+                            IoHash hash = binaryAttachmentField.AsBinaryAttachment();
+
+                            BlobContents referencedBlobContents = await _blobStore.GetObject(ns, BlobIdentifier.FromIoHash(hash));
+
+                            if (_nginxRedirectHelper.CanRedirect(Request, referencedBlobContents))
+                            {
+                                return _nginxRedirectHelper.CreateActionResult(referencedBlobContents, MediaTypeNames.Application.Octet);
+                            }
+                            await WriteBody(referencedBlobContents, MediaTypeNames.Application.Octet);
+                            break;
+                        }
+
+                        // this doesn't look like the generated compact binary so we just return the payload
+                        await using BlobContents contents = new(blobMemory);
+                        await WriteBody(contents, MediaTypeNames.Application.Octet);
+                        break;
+                    }
+                    case MediaTypeNames.Application.Json:
+                    {
+                        byte[] blobMemory;
+                        {
+                            using TelemetrySpan scope = _tracer.StartActiveSpan("json.readblob").SetAttribute("operation.name", "json.readblob");
+                            blobMemory = await blob.Stream.ToByteArray();
+                        }
+                        CbObject cb = new CbObject(blobMemory);
+                        string s = cb.ToJson();
+                        await using BlobContents contents = new BlobContents(Encoding.UTF8.GetBytes(s));
+                        await WriteBody(contents, MediaTypeNames.Application.Json);
+                        break;
+
+                    }
+                    case CustomMediaTypeNames.UnrealCompactBinaryPackage:
+                    {
+                        using TelemetrySpan packageScope = _tracer.StartActiveSpan("cbpackage.fetch").SetAttribute("operation.name", "cbpackage.fetch");
+                        byte[] blobMemory = await blob.Stream.ToByteArray();
+                        CbObject cb = new CbObject(blobMemory);
+
+                        IAsyncEnumerable<Attachment> attachments = _referenceResolver.GetAttachments(ns, cb);
+
+                        using CbPackageBuilder writer = new CbPackageBuilder();
+                        writer.AddAttachment(objectRecord.BlobIdentifier.AsIoHash(), CbPackageAttachmentFlags.IsObject, blobMemory);
+
+                        await Parallel.ForEachAsync(attachments, async (attachment, token) =>
+                        {
+                            IoHash attachmentHash = attachment.AsIoHash();
+                            CbPackageAttachmentFlags flags = 0;
+
+                            try
+                            {
+                                BlobContents attachmentContents;
+                                if (attachment is BlobAttachment blobAttachment)
+                                {
+                                    BlobIdentifier referencedBlob = blobAttachment.Identifier;
+                                    attachmentContents = await _blobStore.GetObject(ns, referencedBlob);
+                                }
+                                else if (attachment is ObjectAttachment objectAttachment)
+                                {
+                                    flags |= CbPackageAttachmentFlags.IsObject;
+                                    BlobIdentifier referencedBlob = objectAttachment.Identifier;
+                                    attachmentContents = await _blobStore.GetObject(ns, referencedBlob);
+                                }
+                                else if (attachment is ContentIdAttachment contentIdAttachment)
+                                {
+
+                                    ContentId contentId = contentIdAttachment.Identifier;
+                                    (attachmentContents, string mime) = await _blobStore.GetCompressedObject(ns, contentId, HttpContext.RequestServices);
+                                    if (mime == CustomMediaTypeNames.UnrealCompressedBuffer)
+                                    {
+                                        flags |= CbPackageAttachmentFlags.IsCompressed;
+                                    }
+                                    else
+                                    {
+                                        // this resolved to a uncompressed blob, the content id existed the the compressed blob didn't
+                                        // so resetting flags to indicate this.
+                                        flags = 0;
+                                    }
+                                }
+                                else
+                                {
+                                    throw new NotSupportedException($"Unknown attachment type {attachment.GetType()}");
+                                }
+
+                                writer.AddAttachment(attachmentHash, flags, attachmentContents.Stream, (ulong)attachmentContents.Length);
+                            }
+                            catch (Exception e)
+                            {
+                                (CbObject errorObject, HttpStatusCode _) = ToErrorResult(e);
+                                writer.AddAttachment(attachmentHash,
+                                    CbPackageAttachmentFlags.IsError | CbPackageAttachmentFlags.IsObject,
+                                    errorObject.GetView().ToArray());
+                            }
+                        });
+
+                        byte[] packageBytes;
+                        {
+                            using TelemetrySpan _ = _tracer.StartActiveSpan("cbpackage.buffer").SetAttribute("operation.name", "cbpackage.buffer");
+                            packageBytes = await writer.ToByteArray();
+                        }
+                        await using BlobContents contents = new BlobContents(packageBytes);
+                        await WriteBody(contents, CustomMediaTypeNames.UnrealCompactBinaryPackage);
+                        break;
+                    }
+                    case CustomMediaTypeNames.JupiterInlinedPayload:
+                    {
+                        byte[] blobMemory = await blob.Stream.ToByteArray();
+                        CbObject cb = new CbObject(blobMemory);
+
+                        static (int, int) CountFields(CbObject o)
+                        {
+                            int countOfBinaryAttachmentFields = 0;
+                            int countOfAttachmentFields = 0;
+
+                            o.IterateAttachments(field =>
+                            {
+                                if (field.IsBinaryAttachment())
+                                {
+                                    ++countOfBinaryAttachmentFields;
+                                }
+
+                                if (field.IsAttachment())
+                                {
+                                    ++countOfAttachmentFields;
+                                }
+                            });
+
+                            return (countOfBinaryAttachmentFields, countOfAttachmentFields);
+                        }
+                        // breaking lambda call into private method to workaround incorrect triggering of CA1508 - [REDACTED_PROFILE_URL]/roslyn-analyzers/issues/5254
+                        (int countOfAttachmentFields, int countOfBinaryAttachmentFields) = CountFields(cb);
+
+                        // if the object consists of a single attachment field we return this attachment field instead
+                        if (countOfBinaryAttachmentFields == 1 && countOfAttachmentFields == 1)
+                        {
+                            // fetch the blob so we can resolve any content ids in it
+                            List<BlobIdentifier> referencedBlobs;
+                            try
+                            {
+                                IAsyncEnumerable<BlobIdentifier> referencedBlobsEnumerable = _referenceResolver.GetReferencedBlobs(ns, cb);
+                                referencedBlobs = await referencedBlobsEnumerable.ToListAsync();
+                            }
+                            catch (PartialReferenceResolveException)
+                            {
+                                return NotFound(new ProblemDetails {Title = $"Object {bucket} {key} in namespace {ns} was missing some content ids"});
+                            }
+                            catch (ReferenceIsMissingBlobsException)
+                            {
+                                return NotFound(new ProblemDetails {Title = $"Object {bucket} {key} in namespace {ns} was missing some blobs"});
+                            }
+
+                            if (referencedBlobs.Count == 1)
+                            {
+                                BlobIdentifier attachmentToSend = referencedBlobs.First();
+                                try
+                                {
+                                    BlobContents referencedBlobContents = await _blobStore.GetObject(ns, attachmentToSend);
+                                    Response.Headers[CommonHeaders.InlinePayloadHash] = attachmentToSend.ToString();
+
+                                    if (_nginxRedirectHelper.CanRedirect(Request, referencedBlobContents))
+                                    {
+                                        return _nginxRedirectHelper.CreateActionResult(referencedBlobContents, CustomMediaTypeNames.JupiterInlinedPayload);
+                                    }
+                                    await WriteBody(referencedBlobContents, CustomMediaTypeNames.JupiterInlinedPayload);
+                                }
+                                catch (BlobNotFoundException)
+                                {
+                                    return NotFound(new ProblemDetails {Title = $"Object {bucket} {key} in namespace {ns} was missing blob {attachmentToSend}"});
+                                }
+                                catch (Exception ex)
+                                {
+                                    _logger.LogError(ex, "Unknown exception encountered while writing body for jupiter inlined payload.");
+                                    throw;
+                                }
+                                return new EmptyResult();
+                            }
+                            else if (referencedBlobs.Count == 0)
+                            {
+                                return NotFound(new ProblemDetails
+                                {
+                                    Title =
+                                        $"Object {objectRecord.Bucket} {objectRecord.Name} did not resolve into any objects that we could find."
+                                });
+                            }
+
+                            return BadRequest(new ProblemDetails
+                            {
+                                Title =
+                                    $"Object {objectRecord.Bucket} {objectRecord.Name} contained a content id which resolved to more then 1 blob, unable to inline this object. Use compact object response instead."
+                            });
+                        }
+                        else if (countOfBinaryAttachmentFields == 0 && countOfAttachmentFields == 0)
+                        {
+                            // no attachments so we just return the compact object instead
+                            await using BlobContents contents = new BlobContents(blobMemory);
+                            await WriteBody(contents, CustomMediaTypeNames.JupiterInlinedPayload);
+                            return new EmptyResult();
+                        }
+
+                        return BadRequest(new ProblemDetails
+                        {
+                            Title =
+                                $"Object {objectRecord.Bucket} {objectRecord.Name} had more then 1 binary attachment field, unable to inline this object. Use compact object response instead."
+                        });
+                    }
+                    default:
+                        throw new NotImplementedException($"Unknown expected response type {responseType}");
+                }
+                
+                // this result is ignored as we write to the body explicitly
+                return new EmptyResult();
+
+            }
+            catch (NamespaceNotFoundException e)
+            {
+                return NotFound(new ProblemDetails {Title = $"Namespace {e.Namespace} did not exist"});
+            }
+            catch (ObjectNotFoundException e)
+            {
+                return NotFound(new ProblemDetails { Title = $"Object {e.Bucket} {e.Key} did not exist" });
+            }
+            catch (BlobNotFoundException e)
+            {
+                return NotFound(new ProblemDetails { Title = $"Object {e.Blob} in {e.Ns} not found" });
+            }
+        }
+
+       
+
+    /// <summary>
+    /// Returns the metadata about a ref key
+    /// </summary>
+    /// <param name="ns">Namespace. Each namespace is completely separated from each other. Use for different types of data that is never expected to be similar (between two different games for instance). Example: `uc4.ddc`</param>
+    /// <param name="bucket">The category/type of record you are caching. Is a clustered key together with the actual key, but all records in the same bucket can be dropped easily. Example: `terrainTexture` </param>
+    /// <param name="key">The unique name of this particular key. `iAmAVeryValidKey`</param>
+    /// <param name="fields">The fields to include in the response, omit this to include everything.</param>
+    [HttpGet("{ns}/{bucket}/{key}/metadata", Order = 500)]
+    public async Task<IActionResult> GetMetadata(
+        [FromRoute] [Required] NamespaceId ns,
+        [FromRoute] [Required] BucketId bucket,
+        [FromRoute] [Required] IoHashKey key,
+        [FromQuery] string[] fields)
+    {
+        ActionResult? accessResult = await _requestHelper.HasAccessToNamespace(User, Request, ns, new [] { DdcAclAction.ReadObject });
+        if (accessResult != null)
+        {
+            return accessResult;
+        }
+
+        try
+        {
+            (ObjectRecord objectRecord, BlobContents? _) = await _objectService.Get(ns, bucket, key, fields);
+
+            return Ok(new RefMetadataResponse(objectRecord));
+        }
+        catch (NamespaceNotFoundException e)
+        {
+            return NotFound(new ProblemDetails {Title = $"Namespace {e.Namespace} did not exist"});
+        }
+        catch (ObjectNotFoundException e)
+        {
+            return NotFound(new ProblemDetails { Title = $"Object {e.Bucket} {e.Key} did not exist" });
+        }
+        catch (BlobNotFoundException e)
+        {
+            return NotFound(new ProblemDetails { Title = $"Object {e.Blob} in {e.Ns} not found" });
+        }
+    }
+
+        /// <summary>
+        /// Checks if a object exists
+        /// </summary>
+        /// <param name="ns">Namespace. Each namespace is completely separated from each other. Use for different types of data that is never expected to be similar (between two different games for instance). Example: `uc4.ddc`</param>
+        /// <param name="bucket">The category/type of record you are caching. Is a clustered key together with the actual key, but all records in the same bucket can be dropped easily. Example: `terrainTexture` </param>
+        /// <param name="key">The unique name of this particular key. `iAmAVeryValidKey`</param>
+        /// <returns>200 if it existed, 400 otherwise</returns>
+        [HttpHead("{ns}/{bucket}/{key}", Order = 500)]
+        [ProducesResponseType(type: typeof(OkResult), 200)]
+        [ProducesResponseType(type: typeof(ValidationProblemDetails), 400)]
+        public async Task<IActionResult> Head(
+            [FromRoute] [Required] NamespaceId ns,
+            [FromRoute] [Required] BucketId bucket,
+            [FromRoute] [Required] IoHashKey key)
+        {
+            ActionResult? accessResult = await _requestHelper.HasAccessToNamespace(User, Request, ns, new [] { DdcAclAction.ReadObject });
+            if (accessResult != null)
+            {
+                return accessResult;
+            }
+
+            try
+            {
+                (ObjectRecord record, BlobContents? blob) = await _objectService.Get(ns, bucket, key, new string[] {"blobIdentifier", "IsFinalized"});
+                Response.Headers[CommonHeaders.HashHeaderName] = record.BlobIdentifier.ToString();
+
+                if (!record.IsFinalized)
+                {
+                    return NotFound(new ProblemDetails {Title = $"Object {bucket} {key} in namespace {ns} is not finalized."});
+                }
+
+                blob ??= await _blobStore.GetObject(ns, record.BlobIdentifier);
+
+                // we have to verify the blobs are available locally, as the record of the key is replicated a head of the content
+                // TODO: Once we support inline replication this step is not needed as at least one region as this blob, just maybe not this current one
+                byte[] blobContents = await blob.Stream.ToByteArray();
+                CbObject compactBinaryObject = new CbObject(blobContents);
+                // the reference resolver will throw if any blob is missing, so no need to do anything other then process each reference
+                IAsyncEnumerable<BlobIdentifier> references = _referenceResolver.GetReferencedBlobs(ns, compactBinaryObject);
+                List<BlobIdentifier>? _ = await references.ToListAsync();
+
+                // we have to verify the blobs are available locally, as the record of the key is replicated a head of the content
+                // TODO: Once we support inline replication this step is not needed as at least one region as this blob, just maybe not this current one
+                BlobIdentifier[] unknownBlobs = await _blobStore.FilterOutKnownBlobs(ns, new BlobIdentifier[] { record.BlobIdentifier });
+                if (unknownBlobs.Length != 0)
+                {
+                    return NotFound(new ProblemDetails {Title = $"Object {bucket} {key} in namespace {ns} had at least one missing blob."});
+                }
+            }
+            catch (NamespaceNotFoundException e)
+            {
+                return NotFound(new ProblemDetails {Title = $"Namespace {e.Namespace} did not exist"});
+            }
+            catch (BlobNotFoundException e)
+            {
+                return NotFound(new ProblemDetails { Title = $"Blob {e.Blob} in namespace {ns} did not exist" });
+            }
+            catch (ObjectNotFoundException e)
+            {
+                return NotFound(new ProblemDetails {Title = $"Object {e.Bucket} {e.Key} in namespace {e.Namespace} did not exist"});
+            }
+            catch (PartialReferenceResolveException)
+            {
+                return NotFound(new ProblemDetails {Title = $"Object {bucket} {key} in namespace {ns} was missing some content ids"});
+            }
+            catch (ReferenceIsMissingBlobsException)
+            {
+                return NotFound(new ProblemDetails {Title = $"Object {bucket} {key} in namespace {ns} was missing some blobs"});
+            }
+
+            return Ok();
+        }
+
+        [HttpGet("{ns}/exists")]
+        [ProducesDefaultResponseType]
+        public async Task<IActionResult> ExistsMultiple(
+            [FromRoute] [Required] NamespaceId ns,
+            [FromQuery] [Required] List<string> names)
+        {
+            ActionResult? accessResult = await _requestHelper.HasAccessToNamespace(User, Request, ns, new [] { DdcAclAction.ReadObject });
+            if (accessResult != null)
+            {
+                return accessResult;
+            }
+
+            ConcurrentBag<(BucketId, IoHashKey)> missingObject = new ();
+
+            List<(BucketId, IoHashKey)> requestedNames = new List<(BucketId, IoHashKey)>();
+            foreach (string name in names)
+            {
+                int separatorIndex = name.IndexOf(".", StringComparison.Ordinal);
+                if (separatorIndex == -1)
+                {
+                    return BadRequest(new ProblemDetails() { Title = $"Key {name} did not contain a '.' separator" });
+                }
+
+                BucketId bucket = new BucketId(name.Substring(0, separatorIndex));
+                IoHashKey key = new IoHashKey(name.Substring(separatorIndex + 1));
+                requestedNames.Add((bucket, key));
+            }
+
+            IEnumerable<Task> tasks = requestedNames.Select(async pair =>
+            {
+                (BucketId bucket, IoHashKey key) = pair;
+                try
+                {
+                    (ObjectRecord record, BlobContents? blob) =
+                        await _objectService.Get(ns, bucket, key, new string[] { "blobIdentifier" });
+
+                    blob ??= await _blobStore.GetObject(ns, record.BlobIdentifier);
+
+                    // we have to verify the blobs are available locally, as the record of the key is replicated a head of the content
+                    // TODO: Once we support inline replication this step is not needed as at least one region as this blob, just maybe not this current one
+                    byte[] blobContents = await blob.Stream.ToByteArray();
+                    CbObject cb = new CbObject(blobContents);
+                    // the reference resolver will throw if any blob is missing, so no need to do anything other then process each reference
+                    IAsyncEnumerable<BlobIdentifier> references = _referenceResolver.GetReferencedBlobs(ns, cb);
+                    List<BlobIdentifier>? _ = await references.ToListAsync();
+                }
+                catch (ObjectNotFoundException)
+                {
+                    missingObject.Add((bucket, key));
+                }
+                catch (PartialReferenceResolveException)
+                {
+                    missingObject.Add((bucket, key));
+                }
+                catch (ReferenceIsMissingBlobsException)
+                {
+                    missingObject.Add((bucket, key));
+                }
+                catch (BlobNotFoundException)
+                {
+                    missingObject.Add((bucket, key));
+                }
+            });
+            await Task.WhenAll(tasks);
+
+            return Ok(new ExistCheckMultipleRefsResponse(missingObject.ToList()));
+        }
+
+        [HttpPut("{ns}/{bucket}/{key}.{format?}", Order = 500)]
+        [DisableRequestSizeLimit]
+        public async Task<IActionResult> PutObject(
+            [FromRoute] [Required] NamespaceId ns,
+            [FromRoute] [Required] BucketId bucket,
+            [FromRoute] [Required] IoHashKey key)
+        {
+            ActionResult? accessResult = await _requestHelper.HasAccessToNamespace(User, Request, ns, new [] { DdcAclAction.WriteObject });
+            if (accessResult != null)
+            {
+                return accessResult;
+            }
+
+            _diagnosticContext.Set("Content-Length", Request.ContentLength ?? -1);
+
+            CbObject payloadObject;
+            BlobIdentifier blobHeader;
+
+            try
+            {
+                using IBufferedPayload payload = await _bufferedPayloadFactory.CreateFromRequest(Request);
+
+                BlobIdentifier headerHash;
+                if (Request.Headers.ContainsKey(CommonHeaders.HashHeaderName))
+                {
+                    headerHash = new BlobIdentifier(Request.Headers[CommonHeaders.HashHeaderName]);
+                }
+                else
+                {
+                    return BadRequest(new ProblemDetails
+                    {
+                        Title = $"Missing expected header {CommonHeaders.HashHeaderName}"
+                    });
+                }
+
+                blobHeader = headerHash;
+
+                switch (Request.ContentType)
+                {
+                    case MediaTypeNames.Application.Json:
+                    {
+                        // TODO: define a scheme for how a json object specifies references
+
+                        blobHeader = await _blobStore.PutObject(ns, payload, headerHash);
+
+                        // TODO: convert the json object into a compact binary instead
+                        CbWriter writer = new CbWriter();
+                        writer.BeginObject();
+                        writer.WriteBinaryAttachmentValue(blobHeader.AsIoHash());
+                        writer.EndObject();
+
+                        byte[] blob = writer.ToByteArray();
+                        payloadObject = new CbObject(blob);
+                        blobHeader = BlobIdentifier.FromBlob(blob);
+                        break;
+                    }
+                    case CustomMediaTypeNames.UnrealCompactBinary:
+                    {
+                        await using MemoryStream ms = new MemoryStream();
+                        await using Stream payloadStream = payload.GetStream();
+                        await payloadStream.CopyToAsync(ms);
+                        payloadObject = new CbObject(ms.ToArray());
+                        break;
+                    }
+                    case MediaTypeNames.Application.Octet:
+                    {
+                        blobHeader = await _blobStore.PutObject(ns, payload, headerHash);
+
+                        CbWriter writer = new CbWriter();
+                        writer.BeginObject();
+                        writer.WriteBinaryAttachment("RawHash", blobHeader.AsIoHash());
+                        writer.WriteInteger("RawSize", payload.Length);
+                        writer.EndObject();
+
+                        byte[] blob = writer.ToByteArray();
+                        payloadObject = new CbObject(blob);
+                        blobHeader = BlobIdentifier.FromBlob(blob);
+                        break;
+                    }
+                    default:
+                        throw new Exception($"Unknown request type {Request.ContentType}, if submitting a blob please use {MediaTypeNames.Application.Octet}");
+                }
+            }
+            catch (HashMismatchException e)
+            {
+                return BadRequest(new ProblemDetails
+                {
+                    Title = $"Incorrect hash, got hash \"{e.SuppliedHash}\" but hash of content was determined to be \"{e.ContentHash}\""
+                });
+            }
+            catch (ClientSendSlowException e)
+            {
+                return Problem(e.Message, null, (int)HttpStatusCode.RequestTimeout);
+            }
+
+            (ContentId[] missingReferences, BlobIdentifier[] missingBlobs) = await _objectService.Put(ns, bucket, key, blobHeader, payloadObject);
+
+            List<JupiterContentHash> missingHashes = new List<JupiterContentHash>(missingReferences);
+            missingHashes.AddRange(missingBlobs);
+            return Ok(new PutObjectResponse(missingHashes.ToArray()));
+        }
+
+        [HttpPut("{ns}/{bucket}/{key}", Order = 300)]
+        [DisableRequestSizeLimit]
+        [RequiredContentType(CustomMediaTypeNames.UnrealCompactBinaryPackage)]
+        public async Task<IActionResult> PutPackage(
+            [FromRoute][Required] NamespaceId ns,
+            [FromRoute][Required] BucketId bucket,
+            [FromRoute][Required] IoHashKey key)
+        {
+            ActionResult? accessResult = await _requestHelper.HasAccessToNamespace(User, Request, ns, new [] { DdcAclAction.WriteObject });
+            if (accessResult != null)
+            {
+                return accessResult;
+            }
+
+            _diagnosticContext.Set("Content-Length", Request.ContentLength ?? -1);
+
+            byte[] b = await RequestUtil.ReadRawBody(Request);
+            CbPackageReader packageReader = await CbPackageReader.Create(new MemoryStream(b));
+
+            try
+            {
+                await foreach ((CbPackageAttachmentEntry entry, byte[] blob) in packageReader.IterateAttachments())
+                {
+                    if (entry.Flags.HasFlag(CbPackageAttachmentFlags.IsError))
+                    {
+                        return BadRequest(new ProblemDetails
+                        {
+                            Title = $"Package contained attachment with error {entry.AttachmentHash}\""
+                        });
+                    }
+                    if (entry.Flags.HasFlag(CbPackageAttachmentFlags.IsCompressed))
+                    {
+                        using MemoryBufferedPayload payload = new MemoryBufferedPayload(blob);
+                        await _blobStore.PutCompressedObject(ns, payload, ContentId.FromIoHash(entry.AttachmentHash), HttpContext.RequestServices);
+                    }
+                    else
+                    {
+                        await _blobStore.PutObject(ns, blob, BlobIdentifier.FromIoHash(entry.AttachmentHash));
+                    }
+                }
+            }
+            catch (HashMismatchException e)
+            {
+                return BadRequest(new ProblemDetails
+                {
+                    Title = $"Incorrect hash, got hash \"{e.SuppliedHash}\" but hash of content was determined to be \"{e.ContentHash}\""
+                });
+            }
+            
+            CbObject rootObject = packageReader.RootObject;
+            BlobIdentifier rootObjectHash = BlobIdentifier.FromIoHash(packageReader.RootHash);
+
+            (ContentId[] missingReferences, BlobIdentifier[] missingBlobs) = await _objectService.Put(ns, bucket, key, rootObjectHash, rootObject);
+
+            List<JupiterContentHash> missingHashes = new List<JupiterContentHash>(missingReferences);
+            missingHashes.AddRange(missingBlobs);
+            return Ok(new PutObjectResponse(missingHashes.ToArray()));
+        }
+
+        [HttpPost("{ns}/{bucket}/{key}/finalize/{hash}.{format?}")]
+        public async Task<IActionResult> FinalizeObject(
+            [FromRoute] [Required] NamespaceId ns,
+            [FromRoute] [Required] BucketId bucket,
+            [FromRoute] [Required] IoHashKey key,
+            [FromRoute] [Required] BlobIdentifier hash)
+        {
+            ActionResult? accessResult = await _requestHelper.HasAccessToNamespace(User, Request, ns, new [] { DdcAclAction.WriteObject });
+            if (accessResult != null)
+            {
+                return accessResult;
+            }
+
+            try
+            {
+                (ContentId[] missingReferences, BlobIdentifier[] missingBlobs) = await _objectService.Finalize(ns, bucket, key, hash);
+                List<JupiterContentHash> missingHashes = new List<JupiterContentHash>(missingReferences);
+                missingHashes.AddRange(missingBlobs);
+
+                return Ok(new PutObjectResponse(missingHashes.ToArray()));
+            }
+            catch (ObjectHashMismatchException e)
+            {
+                return BadRequest(e.Message);
+            }
+            catch (ObjectNotFoundException e)
+            {
+                return NotFound(e.Message);
+            }
+        }
+
+        [HttpPost("{ns}")]
+        [Consumes(CustomMediaTypeNames.UnrealCompactBinary)]
+        [Produces(CustomMediaTypeNames.UnrealCompactBinary)]
+        public async Task<IActionResult> Batch(
+            [FromRoute] [Required] NamespaceId ns,
+            [FromBody] [Required] RefBatchOps ops)
+        {
+            AclAction ActionForOp(RefBatchOps.RefBatchOp.RefOperation op)
+            {
+                switch (op)
+                {
+                    case RefBatchOps.RefBatchOp.RefOperation.GET:
+                        return DdcAclAction.ReadObject;
+                    case RefBatchOps.RefBatchOp.RefOperation.PUT:
+                        return DdcAclAction.WriteObject;
+                    case RefBatchOps.RefBatchOp.RefOperation.HEAD:
+                        return DdcAclAction.ReadObject;
+                    default:
+                        throw new ArgumentOutOfRangeException(nameof(op), op, null);
+                }
+            }
+
+            AclAction[] requiredActions = ops.Ops.Select(op => ActionForOp(op.Op)).ToArray();
+
+            ActionResult? accessResult = await _requestHelper.HasAccessToNamespace(User, Request, ns, requiredActions);
+            if (accessResult != null)
+            {
+                return accessResult;
+            }
+
+            HashSet<uint> usedOpIds = new HashSet<uint>();
+            foreach (RefBatchOps.RefBatchOp batchOp in ops.Ops)
+            {
+                bool added = usedOpIds.Add(batchOp.OpId);
+                if (!added)
+                {
+                    return BadRequest(new ProblemDetails { Title = $"Duplicate op ids used for id: {batchOp.OpId}" });
+                }
+            }
+            ConcurrentDictionary<uint, (CbObject, HttpStatusCode)> results = new();
+
+            async Task<(CbObject, HttpStatusCode)> BatchGetOp(RefBatchOps.RefBatchOp op)
+            {
+                try
+                {
+                    (ObjectRecord objectRecord, BlobContents? blob) = await _objectService.Get(ns, op.Bucket, op.Key, Array.Empty<string>());
+
+                    if (!objectRecord.IsFinalized)
+                    {
+                        throw new Exception("Object is not finalized");
+                    }
+
+                    if (blob == null)
+                    {
+                        throw new Exception();
+                    }
+
+                    CbObject cb = new CbObject(await blob.Stream.ToByteArray());
+
+                    if (op.ResolveAttachments ?? false)
+                    {
+                        IAsyncEnumerable<BlobIdentifier> references = _referenceResolver.GetReferencedBlobs(ns, cb);
+                        List<BlobIdentifier>? _ = await references.ToListAsync();
+                    }
+
+                    return (cb, HttpStatusCode.OK);
+                }
+                catch (Exception ex) when( ex is ObjectNotFoundException or PartialReferenceResolveException or ReferenceIsMissingBlobsException)
+                {
+                    return ToErrorResult(ex, HttpStatusCode.NotFound);
+                }
+                catch (Exception e)
+                {
+                    return ToErrorResult(e);
+                }
+            }
+
+            async Task<(CbObject, HttpStatusCode)> BatchHeadOp(RefBatchOps.RefBatchOp op)
+            {
+                try
+                {
+                    (ObjectRecord record, BlobContents? blob) = await _objectService.Get(ns, op.Bucket, op.Key, new string[] { "blobIdentifier" });
+
+                    if (!record.IsFinalized)
+                    {
+                        return (CbObject.Build(writer => writer.WriteBool("exists", false)), HttpStatusCode.NotFound);
+                    }
+
+                    blob ??= await _blobStore.GetObject(ns, record.BlobIdentifier);
+
+                    if (op.ResolveAttachments ?? false)
+                    {
+                        byte[] blobContents = await blob.Stream.ToByteArray();
+                        CbObject cb = new CbObject(blobContents);
+                        // the reference resolver will throw if any blob is missing, so no need to do anything other then process each reference
+                        IAsyncEnumerable<BlobIdentifier> references = _referenceResolver.GetReferencedBlobs(ns, cb);
+                        List<BlobIdentifier>? _ = await references.ToListAsync();
+                    }
+
+                    if (blob == null)
+                    {
+                        throw new Exception();
+                    }
+
+                    return (CbObject.Build(writer => writer.WriteBool("exists", true)), HttpStatusCode.OK);
+                }
+                catch (Exception ex) when( ex is ObjectNotFoundException or PartialReferenceResolveException or ReferenceIsMissingBlobsException)
+                {
+                    return (CbObject.Build(writer => writer.WriteBool("exists", false)), HttpStatusCode.NotFound);
+                }
+                catch (Exception e)
+                {
+                    return ToErrorResult(e);
+                }
+            }
+
+            async Task<(CbObject, HttpStatusCode)> BatchPutOp(RefBatchOps.RefBatchOp op)
+            {
+                try
+                {
+                    if (op.Payload == null || op.Payload.Equals(CbObject.Empty))
+                    {
+                        throw new Exception($"Missing payload for operation: {op.OpId}");
+                    }
+
+                    if (op.PayloadHash == null)
+                    {
+                        throw new Exception($"Missing payload hash for operation: {op.OpId}");
+                    }
+                    BlobIdentifier headerHash = BlobIdentifier.FromContentHash(op.PayloadHash);
+                    BlobIdentifier objectHash = BlobIdentifier.FromBlob(op.Payload.GetView().ToArray());
+
+                    if (!headerHash.Equals(objectHash))
+                    {
+                        throw new HashMismatchException(headerHash, objectHash);
+                    }
+
+                    (ContentId[] missingReferences, BlobIdentifier[] missingBlobs) = await _objectService.Put(ns, op.Bucket, op.Key, objectHash, op.Payload);
+                    List<JupiterContentHash> missingHashes = new List<JupiterContentHash>(missingReferences);
+
+                    return (CbSerializer.Serialize(new PutObjectResponse(missingHashes.ToArray())), HttpStatusCode.OK);
+                }
+                catch (Exception e)
+                {
+                    return ToErrorResult(e);
+                }
+            }
+
+            await Parallel.ForEachAsync(ops.Ops, CancellationToken.None, async (op, token) =>
+            {
+                switch (op.Op)
+                {
+                    case RefBatchOps.RefBatchOp.RefOperation.GET:
+                        results.TryAdd(op.OpId, await BatchGetOp(op));
+                        break;
+                    case RefBatchOps.RefBatchOp.RefOperation.PUT:
+                        results.TryAdd(op.OpId, await BatchPutOp(op));
+                        break;
+                    case RefBatchOps.RefBatchOp.RefOperation.HEAD:
+                        results.TryAdd(op.OpId, await BatchHeadOp(op));
+                        break;
+                    case RefBatchOps.RefBatchOp.RefOperation.INVALID:
+                    default:
+                        throw new NotImplementedException($"Unknown op type {op.Op}");
+                }
+
+                await Task.CompletedTask;
+            });
+
+            return Ok(new BatchOpsResponse()
+            {
+                Results = results.Select(result =>
+                {
+                    return new BatchOpsResponse.OpResponses()
+                    {
+                        OpId = result.Key,
+                        Response = result.Value.Item1,
+                        StatusCode = (int)result.Value.Item2
+                    };
+                }).ToList()
+            });
+        }
+
+        private static (CbObject, HttpStatusCode) ToErrorResult(Exception exception, HttpStatusCode statusCode = HttpStatusCode.InternalServerError)
+        {
+            Exception e = exception;
+            CbWriter writer = new CbWriter();
+            writer.BeginObject();
+            writer.WriteString("title", e.Message);
+            writer.WriteInteger("status", (int)statusCode);
+            if (e.StackTrace != null)
+            {
+                writer.WriteString("stackTrace", e.StackTrace);
+            }
+            writer.EndObject();
+            return (writer.ToObject(), statusCode);
+        }
+
+        /// <summary>
+        /// Drop all refs records in the namespace
+        /// </summary>
+        /// <param name="ns">Namespace. Each namespace is completely separated from each other. Use for different types of data that is never expected to be similar (between two different games for instance)</param>
+        [HttpDelete("{ns}", Order = 500)]
+        [ProducesResponseType(204)]
+        public async Task<IActionResult> DeleteNamespace(
+            [FromRoute] [Required] NamespaceId ns
+        )
+        {
+            ActionResult? accessResult = await _requestHelper.HasAccessToNamespace(User, Request, ns, new [] { DdcAclAction.DeleteNamespace });
+            if (accessResult != null)
+            {
+                return accessResult;
+            }
+
+            try
+            {
+                await _objectService.DropNamespace(ns);
+            }
+            catch (NamespaceNotFoundException e)
+            {
+                return NotFound(new ProblemDetails {Title = $"Namespace {e.Namespace} did not exist"});
+            }
+
+            return NoContent();
+        }
+
+        /// <summary>
+        /// Drop all refs records in the bucket
+        /// </summary>
+        /// <param name="ns">Namespace. Each namespace is completely separated from each other. Use for different types of data that is never expected to be similar (between two different games for instance)</param>
+        /// <param name="bucket">The category/type of record you are caching. Is a clustered key together with the actual key, but all records in the same bucket can be dropped easily.</param>
+        [HttpDelete("{ns}/{bucket}", Order = 500)]
+        [ProducesResponseType(200)]
+        public async Task<IActionResult> DeleteBucket(
+            [FromRoute] [Required] NamespaceId ns,
+            [FromRoute] [Required] BucketId bucket)
+        {
+            ActionResult? accessResult = await _requestHelper.HasAccessToNamespace(User, Request, ns, new [] { DdcAclAction.DeleteBucket });
+            if (accessResult != null)
+            {
+                return accessResult;
+            }
+
+            long countOfDeletedRecords;
+            try
+            {
+                countOfDeletedRecords = await _objectService.DeleteBucket(ns, bucket);
+            }
+            catch (NamespaceNotFoundException e)
+            {
+                return NotFound(new ProblemDetails {Title = $"Namespace {e.Namespace} did not exist"});
+            }
+
+            return Ok(new BucketDeletedResponse(countOfDeletedRecords));
+        }
+
+        /// <summary>
+        /// Delete a individual refs key
+        /// </summary>
+        /// <param name="ns">Namespace. Each namespace is completely separated from each other. Use for different types of data that is never expected to be similar (between two different games for instance)</param>
+        /// <param name="bucket">The category/type of record you are caching. Is a clustered key together with the actual key, but all records in the same bucket can be dropped easily.</param>
+        /// <param name="key">The unique name of this particular key</param>
+        [HttpDelete("{ns}/{bucket}/{key}", Order = 500)]
+        [ProducesResponseType(200)]
+        [ProducesResponseType(404)]
+        public async Task<IActionResult> Delete(
+            [FromRoute] [Required] NamespaceId ns,
+            [FromRoute] [Required] BucketId bucket,
+            [FromRoute] [Required] IoHashKey key)
+        {
+            ActionResult? accessResult = await _requestHelper.HasAccessToNamespace(User, Request, ns, new [] { DdcAclAction.DeleteObject });
+            if (accessResult != null)
+            {
+                return accessResult;
+            }
+
+            try
+            {
+                bool deleted = await _objectService.Delete(ns, bucket, key);
+                if (!deleted)
+                {
+                    return NotFound(new ProblemDetails { Title = $"Object {key} in bucket {bucket} and namespace {ns} did not exist" });
+                }
+                return Ok(new RefDeletedResponse(deleted ? 1 : 0));
+            }
+            catch (NamespaceNotFoundException e)
+            {
+                return NotFound(new ProblemDetails { Title = $"Namespace {e.Namespace} did not exist" });
+            }
+            catch (ObjectNotFoundException)
+            {
+                return NotFound(new ProblemDetails { Title = $"Object {key} in bucket {bucket} and namespace {ns} did not exist" });
+            }
+        }
+    }
+
+    public class RefDeletedResponse
+    {
+        public RefDeletedResponse()
+        {
+
+        }
+
+        public RefDeletedResponse(int deletedCount)
+        {
+            DeletedCount = deletedCount;
+        }
+
+        [CbField("deletedCount")]
+        public int DeletedCount { get; set; }
+    }
+
+    public class BucketDeletedResponse
+    {
+        public BucketDeletedResponse()
+        {
+
+        }
+
+        public BucketDeletedResponse(long countOfDeletedRecords)
+        {
+            CountOfDeletedRecords = countOfDeletedRecords;
+        }
+
+        [CbField("countOfDeletedRecords")]
+        public long CountOfDeletedRecords { get; set; }
+    }
+
+    public class RefBatchOps
+    {
+        public RefBatchOps()
+        {
+            Ops = Array.Empty<RefBatchOp>();
+        }
+
+        [System.Diagnostics.CodeAnalysis.SuppressMessage("Design", "CA1034:Nested types should not be visible", Justification = "For serialization only")]
+        public class RefBatchOp
+        {
+            public RefBatchOp()
+            {
+                Payload = null;
+                PayloadHash = null;
+            }
+
+            public enum RefOperation
+            {
+                INVALID,
+                GET,
+                PUT,
+                HEAD,
+            }
+
+            [Required]
+            [CbField("opId")]
+            public uint OpId { get; set; }
+
+            [CbField("op")]
+            [JsonIgnore]
+            [UsedImplicitly]
+            public string OpString
+            {
+                get => Op.ToString();
+                set => Op = Enum.Parse<RefOperation>(value);
+            }
+
+            [Required]
+            public RefOperation Op { get; set; } = RefOperation.INVALID;
+
+            [Required]
+            [CbField("bucket")]
+            public BucketId Bucket { get; set; }
+
+            [Required]
+            [CbField("key")]
+            public IoHashKey Key { get; set; }
+
+            [CbField("resolveAttachments")]
+            public bool? ResolveAttachments { get; set; } = null;
+
+            [CbField("payload")]
+            public CbObject? Payload { get; set; } = null;
+
+            [CbField("payloadHash")] 
+            public JupiterContentHash? PayloadHash { get; set; } = null;
+        }
+
+        [CbField("ops")]
+        public RefBatchOp[] Ops { get; set; }
+    }
+
+    public class BatchOpsResponse
+    {
+        public BatchOpsResponse()
+        {
+
+        }
+
+        [System.Diagnostics.CodeAnalysis.SuppressMessage("Design", "CA1034:Nested types should not be visible", Justification = "For serialization only")]
+        public class OpResponses
+        {
+            public OpResponses()
+            {
+
+            }
+            [CbField("opId")]
+            public uint OpId { get; set; }
+
+            [CbField("response")]
+            public CbObject Response { get; set; } = null!;
+
+            [CbField("statusCode")]
+            public int StatusCode { get; set; }
+        }
+
+        [CbField("results")]
+        [System.Diagnostics.CodeAnalysis.SuppressMessage("Usage", "CA2227:Collection properties should be read only", Justification = "Used by serialization")]
+        public List<OpResponses> Results { get; set; } = new List<OpResponses>();
+    }
+
+    public class RefMetadataResponse
+    {
+        public RefMetadataResponse()
+        {
+            PayloadIdentifier = null!;
+            InlinePayload = null!;
+        }
+
+        [JsonConstructor]
+        public RefMetadataResponse(NamespaceId ns, BucketId bucket, IoHashKey name, BlobIdentifier payloadIdentifier, DateTime lastAccess, bool isFinalized, byte[]? inlinePayload)
+        {
+            Ns = ns;
+            Bucket = bucket;
+            Name = name;
+            PayloadIdentifier = payloadIdentifier;
+            LastAccess = lastAccess;
+            IsFinalized = isFinalized;
+            InlinePayload = inlinePayload;
+        }
+
+        public RefMetadataResponse(ObjectRecord objectRecord)
+        {
+            Ns = objectRecord.Namespace;
+            Bucket = objectRecord.Bucket;
+            Name = objectRecord.Name;
+            PayloadIdentifier = objectRecord.BlobIdentifier;
+            LastAccess = objectRecord.LastAccess;
+            IsFinalized = objectRecord.IsFinalized;
+            InlinePayload = objectRecord.InlinePayload;
+        }
+
+        [CbField("ns")]
+        public NamespaceId Ns { get; set; }
+        
+        [CbField("bucket")]
+        public BucketId Bucket { get; set; }
+
+        [CbField("name")]
+        public IoHashKey Name { get; set; }
+
+        [CbField("payloadIdentifier")]
+        public BlobIdentifier PayloadIdentifier { get; set; }
+
+        [CbField("lastAccess")]
+        public DateTime LastAccess { get; set; }
+
+        [CbField("isFinalized")]
+        public bool IsFinalized { get; set; }
+
+        [CbField("inlinePayload")]
+        public byte[]? InlinePayload { get; set; }
+    }
+
+    public class PutObjectResponse
+    {
+        public PutObjectResponse()
+        {
+            Needs = null!;
+        }
+
+        public PutObjectResponse(JupiterContentHash[] missingReferences)
+        {
+            Needs = missingReferences;
+        }
+
+        [CbField("needs")]
+        public JupiterContentHash[] Needs { get; set; }
+    }
+
+    public class ExistCheckMultipleRefsResponse
+    {
+        public ExistCheckMultipleRefsResponse(List<(BucketId,IoHashKey)> missing)
+        {
+            Missing = missing.Select(pair =>
+            {
+                (BucketId bucketId, IoHashKey ioHashKey) = pair;
+                return new MissingReference()
+                {
+                    Bucket = bucketId,
+                    Key = ioHashKey,
+                };
+            }).ToList();
+        }
+
+        [JsonConstructor]
+        public ExistCheckMultipleRefsResponse(List<MissingReference> missing)
+        {
+            Missing = missing;
+        }
+
+        [CbField("missing")]
+        [System.Diagnostics.CodeAnalysis.SuppressMessage("Usage", "CA2227:Collection properties should be read only", Justification = "Used by serialization")]
+        public List<MissingReference> Missing { get; set; }
+
+        [System.Diagnostics.CodeAnalysis.SuppressMessage("Design", "CA1034:Nested types should not be visible", Justification = "For serialization only")]
+        public class MissingReference
+        {
+            [CbField("bucket")]
+            public BucketId Bucket { get; set; }
+
+            [CbField("key")]
+            public IoHashKey Key { get; set; }
+        }
+    }
+
+    public class GetNamespacesResponse
+    {
+        [JsonConstructor]
+        public GetNamespacesResponse(NamespaceId[] namespaces)
+        {
+            Namespaces = namespaces;
+        }
+
+        public NamespaceId[] Namespaces { get; set; }
+    }
+
+    public class HashMismatchException : Exception
+    {
+        public JupiterContentHash SuppliedHash { get; }
+        public JupiterContentHash ContentHash { get; }
+
+        public HashMismatchException(JupiterContentHash suppliedHash, JupiterContentHash contentHash) : base($"ID was not a hash of the content uploaded. Supplied hash was: {suppliedHash} but hash of content was {contentHash}")
+        {
+            SuppliedHash = suppliedHash;
+            ContentHash = contentHash;
+        }
+    }
+}
diff --git a/Engine/Source/Programs/Horde/Horde.Server/Ddc/IBlobService.cs b/Engine/Source/Programs/Horde/Horde.Server/Ddc/IBlobService.cs
deleted file mode 100644
diff --git a/Engine/Source/Programs/Horde/Horde.Server/Ddc/IBlobStore.cs b/Engine/Source/Programs/Horde/Horde.Server/Ddc/IBlobStore.cs
deleted file mode 100644
diff --git a/Engine/Source/Programs/Horde/Horde.Server/Ddc/IContentIdStore.cs b/Engine/Source/Programs/Horde/Horde.Server/Ddc/IContentIdStore.cs
deleted file mode 100644
diff --git a/Engine/Source/Programs/Horde/Horde.Server/Ddc/IDdcBlobService.cs b/Engine/Source/Programs/Horde/Horde.Server/Ddc/IDdcBlobService.cs
new file mode 100644
--- /dev/null
+++ b/Engine/Source/Programs/Horde/Horde.Server/Ddc/IDdcBlobService.cs
@@ -0,0 +1,212 @@
+﻿// Copyright Epic Games, Inc. All Rights Reserved.
+
+using System;
+using System.Collections.Generic;
+using System.IO;
+using System.Net.Mime;
+using System.Threading.Tasks;
+using EpicGames.AspNet;
+using EpicGames.Horde.Storage;
+using Microsoft.Extensions.DependencyInjection;
+using OpenTelemetry.Trace;
+
+#pragma warning disable CS1591
+
+namespace Horde.Server.Ddc
+{
+	public interface IDdcBlobService
+	{
+		Task<JupiterContentHash> VerifyContentMatchesHash(Stream content, JupiterContentHash identifier);
+		Task<BlobIdentifier> PutObjectKnownHash(NamespaceId ns, IBufferedPayload content, BlobIdentifier identifier);
+		Task<BlobIdentifier> PutObject(NamespaceId ns, IBufferedPayload payload, BlobIdentifier identifier);
+		Task<BlobIdentifier> PutObject(NamespaceId ns, byte[] payload, BlobIdentifier identifier);
+		Task<Uri?> MaybePutObjectWithRedirect(NamespaceId ns, BlobIdentifier identifier);
+
+		Task<BlobContents> GetObject(NamespaceId ns, BlobIdentifier blob, List<string>? storageLayers = null, bool supportsRedirectUri = false);
+
+		Task<Uri?> GetObjectWithRedirect(NamespaceId ns, BlobIdentifier blobIdentifier, List<string>? storageLayers = null);
+
+		Task<BlobContents> ReplicateObject(NamespaceId ns, BlobIdentifier blob, bool force = false);
+
+		Task<bool> Exists(NamespaceId ns, BlobIdentifier blob, List<string>? storageLayers = null);
+
+		/// <summary>
+		/// Checks that the blob exists in the root store, the store which is last in the list and thus is intended to have every blob in it
+		/// </summary>
+		/// <param name="ns">The namespace</param>
+		/// <param name="blob">The identifier of the blob</param>
+		/// <returns></returns>
+		Task<bool> ExistsInRootStore(NamespaceId ns, BlobIdentifier blob);
+
+		// Delete a object
+		Task DeleteObject(NamespaceId ns, BlobIdentifier blob);
+
+		// delete the whole namespace
+		Task DeleteNamespace(NamespaceId ns);
+
+		IAsyncEnumerable<(BlobIdentifier, DateTime)> ListObjects(NamespaceId ns);
+		Task<BlobIdentifier[]> FilterOutKnownBlobs(NamespaceId ns, IEnumerable<BlobIdentifier> blobs);
+		Task<BlobIdentifier[]> FilterOutKnownBlobs(NamespaceId ns, IAsyncEnumerable<BlobIdentifier> blobs);
+		Task<BlobContents> GetObjects(NamespaceId ns, BlobIdentifier[] refRequestBlobReferences);
+
+		bool ShouldFetchBlobOnDemand(NamespaceId ns);
+	}
+
+	public class BlobNotFoundException : Exception
+	{
+		public NamespaceId Ns { get; }
+		public BlobIdentifier Blob { get; }
+
+		public BlobNotFoundException(NamespaceId ns, BlobIdentifier blob) : base($"No Blob in Namespace {ns} with id {blob}")
+		{
+			Ns = ns;
+			Blob = blob;
+		}
+
+		public BlobNotFoundException(NamespaceId ns, BlobIdentifier blob, string message) : base(message)
+		{
+			Ns = ns;
+			Blob = blob;
+		}
+	}
+
+	public class BlobReplicationException : BlobNotFoundException
+	{
+		public BlobReplicationException(NamespaceId ns, BlobIdentifier blob, string message) : base(ns, blob, message)
+		{
+		}
+	}
+
+	public class BlobTooLargeException : Exception
+	{
+		public BlobIdentifier Blob { get; }
+
+		public BlobTooLargeException(BlobIdentifier blob) : base($"Blob {blob} was to large to cache")
+		{
+			Blob = blob;
+		}
+	}
+
+	public class ResourceHasToManyRequestsException : Exception
+	{
+		public ResourceHasToManyRequestsException(Exception originalException) : base($"To many requests to resource", originalException)
+		{
+		}
+	}
+
+	public class NamespaceNotFoundException : Exception
+	{
+		public NamespaceId Namespace { get; }
+
+		public NamespaceNotFoundException(NamespaceId @namespace) : base($"Could not find namespace {@namespace}")
+		{
+			Namespace = @namespace;
+		}
+
+		public NamespaceNotFoundException(NamespaceId @namespace, string message) : base(message)
+		{
+			Namespace = @namespace;
+		}
+	}
+
+	public static class BlobServiceExtensions
+	{
+		public static async Task<ContentId> PutCompressedObject(this IDdcBlobService blobService, NamespaceId ns, IBufferedPayload payload, ContentId? id, IServiceProvider provider)
+		{
+			IDdcContentIdStore contentIdStore = provider.GetService<IDdcContentIdStore>()!;
+			CompressedBufferUtils compressedBufferUtils = provider.GetService<CompressedBufferUtils>()!;
+			Tracer tracer = provider.GetService<Tracer>()!;
+
+			// decompress the content and generate a identifier from it to verify the identifier we got
+			await using Stream decompressStream = payload.GetStream();
+			// TODO: we should add a overload for decompress content that can work on streams, otherwise we are still limited to 2GB compressed blobs
+			byte[] decompressedContent = compressedBufferUtils.DecompressContent(await decompressStream.ToByteArray());
+
+			await using MemoryStream decompressedStream = new MemoryStream(decompressedContent);
+			ContentId identifierDecompressedPayload;
+			if (id != null)
+			{
+				identifierDecompressedPayload = ContentId.FromContentHash(await blobService.VerifyContentMatchesHash(decompressedStream, id));
+			}
+			else
+			{
+				JupiterContentHash blobHash;
+				{
+					using TelemetrySpan _ = tracer.StartActiveSpan("web.hash").SetAttribute("operation.name", "web.hash");
+					blobHash = await BlobIdentifier.FromStream(decompressedStream);
+				}
+
+				identifierDecompressedPayload = ContentId.FromContentHash(blobHash);
+			}
+
+			BlobIdentifier identifierCompressedPayload;
+			{
+				using TelemetrySpan _ = tracer.StartActiveSpan("web.hash").SetAttribute("operation.name", "web.hash");
+				await using Stream hashStream = payload.GetStream();
+				identifierCompressedPayload = await BlobIdentifier.FromStream(hashStream);
+			}
+
+			// commit the mapping from the decompressed hash to the compressed hash, we run this in parallel with the blob store submit
+			// TODO: let users specify weight of the blob compared to previously submitted content ids
+			int contentIdWeight = (int)payload.Length;
+			Task contentIdStoreTask = contentIdStore.Put(ns, identifierDecompressedPayload, identifierCompressedPayload, contentIdWeight);
+
+			// we still commit the compressed buffer to the object store using the hash of the compressed content
+			{
+				await blobService.PutObjectKnownHash(ns, payload, identifierCompressedPayload);
+			}
+
+			await contentIdStoreTask;
+
+			return identifierDecompressedPayload;
+		}
+
+		public static async Task<(BlobContents, string)> GetCompressedObject(this IDdcBlobService blobService, NamespaceId ns, ContentId contentId, IServiceProvider provider, bool supportsRedirectUri = false)
+		{
+			IDdcContentIdStore contentIdStore = provider.GetService<IDdcContentIdStore>()!;
+			Tracer tracer = provider.GetService<Tracer>()!;
+
+			BlobIdentifier[]? chunks = await contentIdStore.Resolve(ns, contentId, mustBeContentId: false);
+			if (chunks == null || chunks.Length == 0)
+			{
+				throw new ContentIdResolveException(contentId);
+			}
+
+			// single chunk, we just return that chunk
+			if (chunks.Length == 1)
+			{
+				BlobIdentifier blobToReturn = chunks[0];
+				string mimeType = CustomMediaTypeNames.UnrealCompressedBuffer;
+				if (contentId.Equals(blobToReturn))
+				{
+					// this was actually the unmapped blob, meaning its not a compressed buffer
+					mimeType = MediaTypeNames.Application.Octet;
+				}
+
+				return (await blobService.GetObject(ns, blobToReturn, supportsRedirectUri: supportsRedirectUri), mimeType);
+			}
+
+			// chunked content, combine the chunks into a single stream
+			using TelemetrySpan _ = tracer.StartActiveSpan("blob.combine").SetAttribute("operation.name", "blob.combine");
+			Task<BlobContents>[] tasks = new Task<BlobContents>[chunks.Length];
+			for (int i = 0; i < chunks.Length; i++)
+			{
+				// even if it was requested to support redirect, since we need to combine the chunks using redirects is not possible
+				tasks[i] = blobService.GetObject(ns, chunks[i], supportsRedirectUri: false);
+			}
+
+			MemoryStream ms = new MemoryStream();
+			foreach (Task<BlobContents> task in tasks)
+			{
+				BlobContents blob = await task;
+				await using Stream s = blob.Stream;
+				await s.CopyToAsync(ms);
+			}
+
+			ms.Seek(0, SeekOrigin.Begin);
+
+			// chunking could not have happened for a non compressed buffer so assume it is compressed
+			return (new BlobContents(ms, ms.Length), CustomMediaTypeNames.UnrealCompressedBuffer);
+		}
+	}
+}
diff --git a/Engine/Source/Programs/Horde/Horde.Server/Ddc/IDdcContentIdService.cs b/Engine/Source/Programs/Horde/Horde.Server/Ddc/IDdcContentIdService.cs
new file mode 100644
--- /dev/null
+++ b/Engine/Source/Programs/Horde/Horde.Server/Ddc/IDdcContentIdService.cs
@@ -0,0 +1,50 @@
+// Copyright Epic Games, Inc. All Rights Reserved.
+
+using EpicGames.Horde.Storage;
+using System;
+using System.Threading.Tasks;
+
+#pragma warning disable CS1591
+
+namespace Horde.Server.Ddc
+{
+	public interface IDdcContentIdStore
+    {
+        /// <summary>
+        /// Resolve a content id from its hash into the actual blob (that can in turn be chunked into a set of blobs)
+        /// </summary>
+        /// <param name="ns">The namespace to operate in</param>
+        /// <param name="contentId">The identifier for the content id</param>
+        /// <param name="mustBeContentId"></param>
+        /// <returns></returns>
+        Task<BlobIdentifier[]?> Resolve(NamespaceId ns, ContentId contentId, bool mustBeContentId = false);
+
+        /// <summary>
+        /// Add a mapping from contentId to blobIdentifier
+        /// </summary>
+        /// <param name="ns">The namespace to operate in</param>
+        /// <param name="contentId">The contentId</param>
+        /// <param name="blobIdentifier">The blob the content id maps to</param>
+        /// <param name="contentWeight">Weight of this identifier compared to previous mappings, used to determine which is more important, lower weight is considered a better fit</param>
+        /// <returns></returns>
+        Task Put(NamespaceId ns, ContentId contentId, BlobIdentifier blobIdentifier, int contentWeight);
+    }
+
+    public class InvalidContentIdException : Exception
+    {
+        public InvalidContentIdException(ContentId contentId) : base($"Unknown content id {contentId}")
+        {
+
+        }
+    }
+
+    public class ContentIdResolveException : Exception
+    {
+        public ContentId ContentId { get; }
+
+        public ContentIdResolveException(ContentId contentId) : base($"Unable to find any mapping of contentId {contentId} that has all blobs present")
+        {
+            ContentId = contentId;
+        }
+    }
+}
diff --git a/Engine/Source/Programs/Horde/Horde.Server/Ddc/IDdcObjectService.cs b/Engine/Source/Programs/Horde/Horde.Server/Ddc/IDdcObjectService.cs
new file mode 100644
--- /dev/null
+++ b/Engine/Source/Programs/Horde/Horde.Server/Ddc/IDdcObjectService.cs
@@ -0,0 +1,71 @@
+// Copyright Epic Games, Inc. All Rights Reserved.
+
+using System;
+using System.Collections.Generic;
+using System.Threading.Tasks;
+using EpicGames.Horde.Storage;
+using EpicGames.Serialization;
+
+#pragma warning disable CS1591
+
+namespace Horde.Server.Ddc
+{
+	public interface IDdcObjectService
+    {
+        Task<(ObjectRecord, BlobContents?)> Get(NamespaceId ns, BucketId bucket, IoHashKey key, string[] fields, bool doLastAccessTracking = true);
+        Task<(ContentId[], BlobIdentifier[])> Put(NamespaceId ns, BucketId bucket, IoHashKey key, BlobIdentifier blobHash, CbObject payload);
+        Task<(ContentId[], BlobIdentifier[])> Finalize(NamespaceId ns, BucketId bucket, IoHashKey key, BlobIdentifier blobHash);
+
+        IAsyncEnumerable<NamespaceId> GetNamespaces();
+
+        Task<bool> Delete(NamespaceId ns, BucketId bucket, IoHashKey key);
+        Task<long> DropNamespace(NamespaceId ns);
+        Task<long> DeleteBucket(NamespaceId ns, BucketId bucket);
+
+        Task<bool> Exists(NamespaceId ns, BucketId bucket, IoHashKey key);
+        Task<List<BlobIdentifier>> GetReferencedBlobs(NamespaceId ns, BucketId bucket, IoHashKey key);
+    }
+
+	public class ObjectRecord
+	{
+		public ObjectRecord(NamespaceId ns, BucketId bucket, IoHashKey name, DateTime lastAccess, byte[]? inlinePayload, BlobIdentifier blobIdentifier, bool isFinalized)
+		{
+			Namespace = ns;
+			Bucket = bucket;
+			Name = name;
+			LastAccess = lastAccess;
+			InlinePayload = inlinePayload;
+			BlobIdentifier = blobIdentifier;
+			IsFinalized = isFinalized;
+		}
+
+		public NamespaceId Namespace { get; }
+		public BucketId Bucket { get; }
+		public IoHashKey Name { get; }
+		public DateTime LastAccess { get; }
+		public byte[]? InlinePayload { get; set; }
+		public BlobIdentifier BlobIdentifier { get; set; }
+		public bool IsFinalized { get; }
+	}
+
+	public class ObjectNotFoundException : Exception
+	{
+		public ObjectNotFoundException(NamespaceId ns, BucketId bucket, IoHashKey key) : base($"Object not found {key} in bucket {bucket} namespace {ns}")
+		{
+			Namespace = ns;
+			Bucket = bucket;
+			Key = key;
+		}
+
+		public NamespaceId Namespace { get; }
+		public BucketId Bucket { get; }
+		public IoHashKey Key { get; }
+	}
+
+	public class ObjectHashMismatchException : Exception
+    {
+        public ObjectHashMismatchException(NamespaceId ns, BucketId bucket, IoHashKey name, BlobIdentifier suppliedHash, BlobIdentifier actualHash) : base($"Object {name} in bucket {bucket} and namespace {ns} did not reference hash {suppliedHash} was referencing {actualHash}")
+        {
+        }
+    }
+}
diff --git a/Engine/Source/Programs/Horde/Horde.Server/Ddc/IObjectService.cs b/Engine/Source/Programs/Horde/Horde.Server/Ddc/IObjectService.cs
deleted file mode 100644
diff --git a/Engine/Source/Programs/Horde/Horde.Server/Ddc/IReferenceResolver.cs b/Engine/Source/Programs/Horde/Horde.Server/Ddc/IReferenceResolver.cs
--- a/Engine/Source/Programs/Horde/Horde.Server/Ddc/IReferenceResolver.cs
+++ b/Engine/Source/Programs/Horde/Horde.Server/Ddc/IReferenceResolver.cs
@@ -89,11 +89,11 @@ namespace Horde.Server.Ddc
 
     public class ReferenceResolver : IReferenceResolver
     {
-        private readonly IBlobService _blobStore;
-        private readonly IContentIdStore _contentIdStore;
+        private readonly IDdcBlobService _blobStore;
+        private readonly IDdcContentIdStore _contentIdStore;
         private readonly Tracer _tracer;
 
-        public ReferenceResolver(IBlobService blobStore, IContentIdStore contentIdStore, Tracer tracer)
+        public ReferenceResolver(IDdcBlobService blobStore, IDdcContentIdStore contentIdStore, Tracer tracer)
         {
             _blobStore = blobStore;
             _contentIdStore = contentIdStore;
diff --git a/Engine/Source/Programs/Horde/Horde.Server/Ddc/IReferencesStore.cs b/Engine/Source/Programs/Horde/Horde.Server/Ddc/IReferencesStore.cs
deleted file mode 100644
diff --git a/Engine/Source/Programs/Horde/Horde.Server/Ddc/NamespaceNotFoundException.cs b/Engine/Source/Programs/Horde/Horde.Server/Ddc/NamespaceNotFoundException.cs
deleted file mode 100644
diff --git a/Engine/Source/Programs/Horde/Horde.Server/Ddc/ObjectController.cs b/Engine/Source/Programs/Horde/Horde.Server/Ddc/ObjectController.cs
deleted file mode 100644
diff --git a/Engine/Source/Programs/Horde/Horde.Server/Ddc/ObjectService.cs b/Engine/Source/Programs/Horde/Horde.Server/Ddc/ObjectService.cs
--- a/Engine/Source/Programs/Horde/Horde.Server/Ddc/ObjectService.cs
+++ b/Engine/Source/Programs/Horde/Horde.Server/Ddc/ObjectService.cs
@@ -7,7 +7,7 @@ using EpicGames.Serialization;
 
 namespace Horde.Server.Ddc
 {
-	class ObjectService : IObjectService
+	class ObjectService : IDdcObjectService
 	{
 		public Task<bool> Delete(NamespaceId ns, BucketId bucket, IoHashKey key)
 		{
diff --git a/Engine/Source/Programs/Horde/Horde.Server/Ddc/ReferencesController.cs b/Engine/Source/Programs/Horde/Horde.Server/Ddc/ReferencesController.cs
deleted file mode 100644
diff --git a/Engine/Source/Programs/Horde/Horde.Server/Ddc/ReferencesStore.cs b/Engine/Source/Programs/Horde/Horde.Server/Ddc/ReferencesStore.cs
deleted file mode 100644
diff --git a/Engine/Source/Programs/Horde/Horde.Server/Startup.cs b/Engine/Source/Programs/Horde/Horde.Server/Startup.cs
--- a/Engine/Source/Programs/Horde/Horde.Server/Startup.cs
+++ b/Engine/Source/Programs/Horde/Horde.Server/Startup.cs
@@ -521,12 +521,10 @@ namespace Horde.Server
 			}
 
 			// DDC
-			services.AddSingleton<IBlobStore, BlobStore>();
-			services.AddSingleton<IBlobService, BlobService>();
-			services.AddSingleton<IObjectService, ObjectService>();
-			services.AddSingleton<IReferencesStore, ReferencesStore>();
+			services.AddSingleton<IDdcBlobService, BlobService>();
+			services.AddSingleton<IDdcObjectService, ObjectService>();
 			services.AddSingleton<IReferenceResolver, ReferenceResolver>();
-			services.AddSingleton<IContentIdStore, ContentIdStore>();
+			services.AddSingleton<IDdcContentIdStore, ContentIdStore>();
 
 			// Secret providers
 			services.AddSingleton<SecretService>();
