Horde: Rename DDC services to Ddc*.

diff --git a/Engine/Source/Programs/Horde/Horde.Server/Ddc/BlobService.cs b/Engine/Source/Programs/Horde/Horde.Server/Ddc/BlobService.cs
deleted file mode 100644
diff --git a/Engine/Source/Programs/Horde/Horde.Server/Ddc/ContentIdStore.cs b/Engine/Source/Programs/Horde/Horde.Server/Ddc/ContentIdStore.cs
deleted file mode 100644
diff --git a/Engine/Source/Programs/Horde/Horde.Server/Ddc/DdcBlobService.cs b/Engine/Source/Programs/Horde/Horde.Server/Ddc/DdcBlobService.cs
new file mode 100644
--- /dev/null
+++ b/Engine/Source/Programs/Horde/Horde.Server/Ddc/DdcBlobService.cs
@@ -0,0 +1,98 @@
+// Copyright Epic Games, Inc. All Rights Reserved.
+
+using System;
+using System.Collections.Generic;
+using System.IO;
+using System.Threading.Tasks;
+using EpicGames.Horde.Storage;
+
+namespace Horde.Server.Ddc
+{
+	class BlobService : IDdcBlobService
+	{
+		public Task DeleteNamespace(NamespaceId ns)
+		{
+			throw new NotImplementedException();
+		}
+
+		public Task DeleteObject(NamespaceId ns, BlobIdentifier blob)
+		{
+			throw new NotImplementedException();
+		}
+
+		public Task<bool> Exists(NamespaceId ns, BlobIdentifier blob, List<string>? storageLayers = null)
+		{
+			throw new NotImplementedException();
+		}
+
+		public Task<bool> ExistsInRootStore(NamespaceId ns, BlobIdentifier blob)
+		{
+			throw new NotImplementedException();
+		}
+
+		public Task<BlobIdentifier[]> FilterOutKnownBlobs(NamespaceId ns, IEnumerable<BlobIdentifier> blobs)
+		{
+			throw new NotImplementedException();
+		}
+
+		public Task<BlobIdentifier[]> FilterOutKnownBlobs(NamespaceId ns, IAsyncEnumerable<BlobIdentifier> blobs)
+		{
+			throw new NotImplementedException();
+		}
+
+		public Task<BlobContents> GetObject(NamespaceId ns, BlobIdentifier blob, List<string>? storageLayers = null, bool supportsRedirectUri = false)
+		{
+			throw new NotImplementedException();
+		}
+
+		public Task<BlobContents> GetObjects(NamespaceId ns, BlobIdentifier[] refRequestBlobReferences)
+		{
+			throw new NotImplementedException();
+		}
+
+		public Task<Uri?> GetObjectWithRedirect(NamespaceId ns, BlobIdentifier blobIdentifier, List<string>? storageLayers = null)
+		{
+			throw new NotImplementedException();
+		}
+
+		public IAsyncEnumerable<(BlobIdentifier, DateTime)> ListObjects(NamespaceId ns)
+		{
+			throw new NotImplementedException();
+		}
+
+		public Task<Uri?> MaybePutObjectWithRedirect(NamespaceId ns, BlobIdentifier identifier)
+		{
+			throw new NotImplementedException();
+		}
+
+		public Task<BlobIdentifier> PutObject(NamespaceId ns, IBufferedPayload payload, BlobIdentifier identifier)
+		{
+			throw new NotImplementedException();
+		}
+
+		public Task<BlobIdentifier> PutObject(NamespaceId ns, byte[] payload, BlobIdentifier identifier)
+		{
+			throw new NotImplementedException();
+		}
+
+		public Task<BlobIdentifier> PutObjectKnownHash(NamespaceId ns, IBufferedPayload content, BlobIdentifier identifier)
+		{
+			throw new NotImplementedException();
+		}
+
+		public Task<BlobContents> ReplicateObject(NamespaceId ns, BlobIdentifier blob, bool force = false)
+		{
+			throw new NotImplementedException();
+		}
+
+		public bool ShouldFetchBlobOnDemand(NamespaceId ns)
+		{
+			throw new NotImplementedException();
+		}
+
+		public Task<JupiterContentHash> VerifyContentMatchesHash(Stream content, JupiterContentHash identifier)
+		{
+			throw new NotImplementedException();
+		}
+	}
+}
diff --git a/Engine/Source/Programs/Horde/Horde.Server/Ddc/DdcContentIdController.cs b/Engine/Source/Programs/Horde/Horde.Server/Ddc/DdcContentIdController.cs
--- a/Engine/Source/Programs/Horde/Horde.Server/Ddc/DdcContentIdController.cs
+++ b/Engine/Source/Programs/Horde/Horde.Server/Ddc/DdcContentIdController.cs
@@ -19,9 +19,9 @@ namespace Horde.Server.Ddc
     public class DdcContentIdController : ControllerBase
     {
         private readonly IRequestHelper _requestHelper;
-        private readonly IDdcContentIdStore _contentIdStore;
+        private readonly IDdcContentIdService _contentIdStore;
 
-        public DdcContentIdController(IRequestHelper requestHelper, IDdcContentIdStore contentIdStore)
+        public DdcContentIdController(IRequestHelper requestHelper, IDdcContentIdService contentIdStore)
         {
             _requestHelper = requestHelper;
             _contentIdStore = contentIdStore;
diff --git a/Engine/Source/Programs/Horde/Horde.Server/Ddc/DdcContentIdService.cs b/Engine/Source/Programs/Horde/Horde.Server/Ddc/DdcContentIdService.cs
new file mode 100644
--- /dev/null
+++ b/Engine/Source/Programs/Horde/Horde.Server/Ddc/DdcContentIdService.cs
@@ -0,0 +1,21 @@
+// Copyright Epic Games, Inc. All Rights Reserved.
+
+using EpicGames.Horde.Storage;
+using System;
+using System.Threading.Tasks;
+
+namespace Horde.Server.Ddc
+{
+	class ContentIdStore : IDdcContentIdService
+	{
+		public Task<BlobIdentifier[]?> Resolve(NamespaceId ns, ContentId contentId, bool mustBeContentId = false)
+		{
+			throw new NotImplementedException();
+		}
+
+		public Task Put(NamespaceId ns, ContentId contentId, BlobIdentifier blobIdentifier, int contentWeight)
+		{
+			throw new NotImplementedException();
+		}
+	}
+}
diff --git a/Engine/Source/Programs/Horde/Horde.Server/Ddc/DdcObjectService.cs b/Engine/Source/Programs/Horde/Horde.Server/Ddc/DdcObjectService.cs
new file mode 100644
--- /dev/null
+++ b/Engine/Source/Programs/Horde/Horde.Server/Ddc/DdcObjectService.cs
@@ -0,0 +1,57 @@
+// Copyright Epic Games, Inc. All Rights Reserved.
+
+using System.Collections.Generic;
+using System.Threading.Tasks;
+using EpicGames.Horde.Storage;
+using EpicGames.Serialization;
+
+namespace Horde.Server.Ddc
+{
+	class ObjectService : IDdcObjectService
+	{
+		public Task<bool> Delete(NamespaceId ns, BucketId bucket, IoHashKey key)
+		{
+			throw new System.NotImplementedException();
+		}
+
+		public Task<long> DeleteBucket(NamespaceId ns, BucketId bucket)
+		{
+			throw new System.NotImplementedException();
+		}
+
+		public Task<long> DropNamespace(NamespaceId ns)
+		{
+			throw new System.NotImplementedException();
+		}
+
+		public Task<bool> Exists(NamespaceId ns, BucketId bucket, IoHashKey key)
+		{
+			throw new System.NotImplementedException();
+		}
+
+		public Task<(ContentId[], BlobIdentifier[])> Finalize(NamespaceId ns, BucketId bucket, IoHashKey key, BlobIdentifier blobHash)
+		{
+			throw new System.NotImplementedException();
+		}
+
+		public Task<(ObjectRecord, BlobContents?)> Get(NamespaceId ns, BucketId bucket, IoHashKey key, string[] fields, bool doLastAccessTracking = true)
+		{
+			throw new System.NotImplementedException();
+		}
+
+		public IAsyncEnumerable<NamespaceId> GetNamespaces()
+		{
+			throw new System.NotImplementedException();
+		}
+
+		public Task<List<BlobIdentifier>> GetReferencedBlobs(NamespaceId ns, BucketId bucket, IoHashKey key)
+		{
+			throw new System.NotImplementedException();
+		}
+
+		public Task<(ContentId[], BlobIdentifier[])> Put(NamespaceId ns, BucketId bucket, IoHashKey key, BlobIdentifier blobHash, CbObject payload)
+		{
+			throw new System.NotImplementedException();
+		}
+	}
+}
diff --git a/Engine/Source/Programs/Horde/Horde.Server/Ddc/IDdcBlobService.cs b/Engine/Source/Programs/Horde/Horde.Server/Ddc/IDdcBlobService.cs
--- a/Engine/Source/Programs/Horde/Horde.Server/Ddc/IDdcBlobService.cs
+++ b/Engine/Source/Programs/Horde/Horde.Server/Ddc/IDdcBlobService.cs
@@ -113,7 +113,7 @@ namespace Horde.Server.Ddc
 	{
 		public static async Task<ContentId> PutCompressedObject(this IDdcBlobService blobService, NamespaceId ns, IBufferedPayload payload, ContentId? id, IServiceProvider provider)
 		{
-			IDdcContentIdStore contentIdStore = provider.GetService<IDdcContentIdStore>()!;
+			IDdcContentIdService contentIdStore = provider.GetService<IDdcContentIdService>()!;
 			CompressedBufferUtils compressedBufferUtils = provider.GetService<CompressedBufferUtils>()!;
 			Tracer tracer = provider.GetService<Tracer>()!;
 
@@ -163,7 +163,7 @@ namespace Horde.Server.Ddc
 
 		public static async Task<(BlobContents, string)> GetCompressedObject(this IDdcBlobService blobService, NamespaceId ns, ContentId contentId, IServiceProvider provider, bool supportsRedirectUri = false)
 		{
-			IDdcContentIdStore contentIdStore = provider.GetService<IDdcContentIdStore>()!;
+			IDdcContentIdService contentIdStore = provider.GetService<IDdcContentIdService>()!;
 			Tracer tracer = provider.GetService<Tracer>()!;
 
 			BlobIdentifier[]? chunks = await contentIdStore.Resolve(ns, contentId, mustBeContentId: false);
diff --git a/Engine/Source/Programs/Horde/Horde.Server/Ddc/IDdcContentIdService.cs b/Engine/Source/Programs/Horde/Horde.Server/Ddc/IDdcContentIdService.cs
--- a/Engine/Source/Programs/Horde/Horde.Server/Ddc/IDdcContentIdService.cs
+++ b/Engine/Source/Programs/Horde/Horde.Server/Ddc/IDdcContentIdService.cs
@@ -8,7 +8,7 @@ using System.Threading.Tasks;
 
 namespace Horde.Server.Ddc
 {
-	public interface IDdcContentIdStore
+	public interface IDdcContentIdService
     {
         /// <summary>
         /// Resolve a content id from its hash into the actual blob (that can in turn be chunked into a set of blobs)
diff --git a/Engine/Source/Programs/Horde/Horde.Server/Ddc/IReferenceResolver.cs b/Engine/Source/Programs/Horde/Horde.Server/Ddc/IReferenceResolver.cs
--- a/Engine/Source/Programs/Horde/Horde.Server/Ddc/IReferenceResolver.cs
+++ b/Engine/Source/Programs/Horde/Horde.Server/Ddc/IReferenceResolver.cs
@@ -90,10 +90,10 @@ namespace Horde.Server.Ddc
     public class ReferenceResolver : IReferenceResolver
     {
         private readonly IDdcBlobService _blobStore;
-        private readonly IDdcContentIdStore _contentIdStore;
+        private readonly IDdcContentIdService _contentIdStore;
         private readonly Tracer _tracer;
 
-        public ReferenceResolver(IDdcBlobService blobStore, IDdcContentIdStore contentIdStore, Tracer tracer)
+        public ReferenceResolver(IDdcBlobService blobStore, IDdcContentIdService contentIdStore, Tracer tracer)
         {
             _blobStore = blobStore;
             _contentIdStore = contentIdStore;
diff --git a/Engine/Source/Programs/Horde/Horde.Server/Ddc/ObjectService.cs b/Engine/Source/Programs/Horde/Horde.Server/Ddc/ObjectService.cs
deleted file mode 100644
diff --git a/Engine/Source/Programs/Horde/Horde.Server/Startup.cs b/Engine/Source/Programs/Horde/Horde.Server/Startup.cs
--- a/Engine/Source/Programs/Horde/Horde.Server/Startup.cs
+++ b/Engine/Source/Programs/Horde/Horde.Server/Startup.cs
@@ -524,7 +524,7 @@ namespace Horde.Server
 			services.AddSingleton<IDdcBlobService, BlobService>();
 			services.AddSingleton<IDdcObjectService, ObjectService>();
 			services.AddSingleton<IReferenceResolver, ReferenceResolver>();
-			services.AddSingleton<IDdcContentIdStore, ContentIdStore>();
+			services.AddSingleton<IDdcContentIdService, ContentIdStore>();
 
 			// Secret providers
 			services.AddSingleton<SecretService>();
