Horde: Use EpicGames.Core implementation of StringUtils, rather than version taken from Jupiter.

diff --git a/Engine/Source/Programs/Horde/Horde.Server/Ddc/BlobIdentifier.cs b/Engine/Source/Programs/Horde/Horde.Server/Ddc/BlobIdentifier.cs
--- a/Engine/Source/Programs/Horde/Horde.Server/Ddc/BlobIdentifier.cs
+++ b/Engine/Source/Programs/Horde/Horde.Server/Ddc/BlobIdentifier.cs
@@ -76,7 +76,7 @@ namespace Horde.Server.Ddc
         {
             if (_stringIdentifier == null)
             {
-                _stringIdentifier = StringUtils.FormatAsHexString(Identifier);
+                _stringIdentifier = StringUtils.FormatHexString(Identifier);
             }
 
             return _stringIdentifier;
diff --git a/Engine/Source/Programs/Horde/Horde.Server/Ddc/JupiterContentHash.cs b/Engine/Source/Programs/Horde/Horde.Server/Ddc/JupiterContentHash.cs
--- a/Engine/Source/Programs/Horde/Horde.Server/Ddc/JupiterContentHash.cs
+++ b/Engine/Source/Programs/Horde/Horde.Server/Ddc/JupiterContentHash.cs
@@ -41,7 +41,7 @@ namespace Horde.Server.Ddc
                 throw new ArgumentNullException(nameof(identifier));
             }
 
-            byte[] byteIdentifier = StringUtils.ToHashFromHexString(identifier);
+            byte[] byteIdentifier = StringUtils.ParseHexString(identifier);
             Identifier = byteIdentifier;
             /*if (byteIdentifier.Length != HashLength)
             {
@@ -96,7 +96,7 @@ namespace Horde.Server.Ddc
 
         public override string ToString()
         {
-            return StringUtils.FormatAsHexString(Identifier);
+            return StringUtils.FormatHexString(Identifier);
         }
 
         public static JupiterContentHash FromBlob(byte[] blobMemory)
diff --git a/Engine/Source/Programs/Horde/Horde.Server/Ddc/Sha1Utils.cs b/Engine/Source/Programs/Horde/Horde.Server/Ddc/Sha1Utils.cs
--- a/Engine/Source/Programs/Horde/Horde.Server/Ddc/Sha1Utils.cs
+++ b/Engine/Source/Programs/Horde/Horde.Server/Ddc/Sha1Utils.cs
@@ -3,6 +3,7 @@
 using System.Collections.Generic;
 using System.Security.Cryptography;
 using System.Threading.Tasks;
+using EpicGames.Core;
 
 #pragma warning disable CS1591
 
@@ -31,7 +32,7 @@ namespace Horde.Server.Ddc
         {
             byte[] sha1 = GetSHA1(payload);
 
-            return StringUtils.FormatAsHexString(sha1);
+            return StringUtils.FormatHexString(sha1);
         }
     }
 
diff --git a/Engine/Source/Programs/Horde/Horde.Server/Ddc/StringUtils.cs b/Engine/Source/Programs/Horde/Horde.Server/Ddc/StringUtils.cs
deleted file mode 100644
