Horde: Change device pool telemetry to use ascending key for creation time
#rnx

diff --git a/Engine/Source/Programs/Horde/Horde.Server.Tests/DeviceServiceTests.cs b/Engine/Source/Programs/Horde/Horde.Server.Tests/DeviceServiceTests.cs
--- a/Engine/Source/Programs/Horde/Horde.Server.Tests/DeviceServiceTests.cs
+++ b/Engine/Source/Programs/Horde/Horde.Server.Tests/DeviceServiceTests.cs
@@ -266,8 +266,7 @@ namespace Horde.Server.Tests
 			// make sure the reswrvation time was finished for the problem device
 			Assert.IsNotNull(telemetry[0].Telemetry[0].ReservationFinishUtc);
 		}
-
-		[Ignore]
+		
 		[TestMethod]
 		public async Task TestDevicePoolTelemetryCapture()
 		{
diff --git a/Engine/Source/Programs/Horde/Horde.Server/Devices/DeviceCollection.cs b/Engine/Source/Programs/Horde/Horde.Server/Devices/DeviceCollection.cs
--- a/Engine/Source/Programs/Horde/Horde.Server/Devices/DeviceCollection.cs
+++ b/Engine/Source/Programs/Horde/Horde.Server/Devices/DeviceCollection.cs
@@ -506,11 +506,11 @@ namespace Horde.Server.Devices
 			_reservations = mongoService.GetCollection<DeviceReservationDocument>("Devices.Reservations");
 
 			List<MongoIndex<DeviceTelemetryDocument>> deviceTelemetryIndexes = new List<MongoIndex<DeviceTelemetryDocument>>();
-			deviceTelemetryIndexes.Add((keys => keys.Descending(x => x.CreateTimeUtc)));
+			deviceTelemetryIndexes.Add((keys => keys.Ascending(x => x.CreateTimeUtc)));
 			_deviceTelemetry = mongoService.GetCollection<DeviceTelemetryDocument>("Devices.DeviceTelemetryV2", deviceTelemetryIndexes);
 
 			List<MongoIndex<DevicePoolTelemetryDocument>> poolTelemetryIndexes = new List<MongoIndex<DevicePoolTelemetryDocument>>();
-			poolTelemetryIndexes.Add((keys => keys.Descending(x => x.CreateTimeUtc)));
+			poolTelemetryIndexes.Add((keys => keys.Ascending(x => x.CreateTimeUtc)));
 			_poolTelemetry = mongoService.GetCollection<DevicePoolTelemetryDocument>("Devices.PoolTelemetryV2", poolTelemetryIndexes);
 		}
 
