Horde: Remove Sha1Utils; move functionality directly into identifier classes.

diff --git a/Engine/Source/Programs/Horde/Horde.Server/Ddc/BlobIdentifier.cs b/Engine/Source/Programs/Horde/Horde.Server/Ddc/BlobIdentifier.cs
--- a/Engine/Source/Programs/Horde/Horde.Server/Ddc/BlobIdentifier.cs
+++ b/Engine/Source/Programs/Horde/Horde.Server/Ddc/BlobIdentifier.cs
@@ -39,18 +39,12 @@ namespace Horde.Server.Ddc
 
         public override int GetHashCode()
         {
-            return Comparer.GetHashCode(Identifier);
+			HashCode hashCode = new HashCode();
+			hashCode.AddBytes(Identifier);
+			return hashCode.ToHashCode();
         }
 
-        public bool Equals(BlobIdentifier? other)
-        {
-            if (other == null)
-            {
-                return false;
-            }
-
-            return Comparer.Equals(Identifier, other.Identifier);
-        }
+		public bool Equals(BlobIdentifier? other) => other != null && other.Identifier.AsSpan().SequenceEqual(other.Identifier);
 
         public override bool Equals(object? obj)
         {
diff --git a/Engine/Source/Programs/Horde/Horde.Server/Ddc/ContentId.cs b/Engine/Source/Programs/Horde/Horde.Server/Ddc/ContentId.cs
--- a/Engine/Source/Programs/Horde/Horde.Server/Ddc/ContentId.cs
+++ b/Engine/Source/Programs/Horde/Horde.Server/Ddc/ContentId.cs
@@ -25,18 +25,12 @@ namespace Horde.Server.Ddc
 
         public override int GetHashCode()
         {
-            return Comparer.GetHashCode(Identifier);
+			HashCode hashCode = new HashCode();
+			hashCode.AddBytes(Identifier);
+			return hashCode.ToHashCode();
         }
 
-        public bool Equals(ContentId? other)
-        {
-            if (other == null)
-            {
-                return false;
-            }
-
-            return Comparer.Equals(Identifier, other.Identifier);
-        }
+        public bool Equals(ContentId? other) => other != null && other.Identifier.AsSpan().SequenceEqual(Identifier);
 
         public override bool Equals(object? obj)
         {
diff --git a/Engine/Source/Programs/Horde/Horde.Server/Ddc/JupiterContentHash.cs b/Engine/Source/Programs/Horde/Horde.Server/Ddc/JupiterContentHash.cs
--- a/Engine/Source/Programs/Horde/Horde.Server/Ddc/JupiterContentHash.cs
+++ b/Engine/Source/Programs/Horde/Horde.Server/Ddc/JupiterContentHash.cs
@@ -19,7 +19,6 @@ namespace Horde.Server.Ddc
     [CbConverter(typeof(ContentHashCbConverter))]
     public class JupiterContentHash : IEquatable<JupiterContentHash>, IEquatable<byte[]>
     {
-        protected ByteArrayComparer Comparer { get; } = new ByteArrayComparer();
         protected byte[] Identifier { get; init; }
         public const int HashLength = 20;
 
@@ -51,28 +50,14 @@ namespace Horde.Server.Ddc
 
         public override int GetHashCode()
         {
-            return Comparer.GetHashCode(Identifier);
+			HashCode hashCode = new HashCode();
+			hashCode.AddBytes(Identifier);
+			return hashCode.ToHashCode();
         }
 
-        public bool Equals(JupiterContentHash? other)
-        {
-            if (other == null)
-            {
-                return false;
-            }
-
-            return Comparer.Equals(Identifier, other.Identifier);
-        }
+        public bool Equals(JupiterContentHash? other) => other != null && Identifier.AsSpan().SequenceEqual(other.Identifier);
 
-        public bool Equals(byte[]? other)
-        {
-            if (other == null)
-            {
-                return false;
-            }
-
-            return Comparer.Equals(Identifier, other);
-        }
+		public bool Equals(byte[]? other) => other != null && Identifier.AsSpan().SequenceEqual(other);
 
         public override bool Equals(object? obj)
         {
diff --git a/Engine/Source/Programs/Horde/Horde.Server/Ddc/Sha1Utils.cs b/Engine/Source/Programs/Horde/Horde.Server/Ddc/Sha1Utils.cs
deleted file mode 100644
