Horde: Remove RequestUtil. Functionality is only used in one place; move it there instead.

diff --git a/Engine/Source/Programs/Horde/Horde.Server/Ddc/DdcRefsController.cs b/Engine/Source/Programs/Horde/Horde.Server/Ddc/DdcRefsController.cs
--- a/Engine/Source/Programs/Horde/Horde.Server/Ddc/DdcRefsController.cs
+++ b/Engine/Source/Programs/Horde/Horde.Server/Ddc/DdcRefsController.cs
@@ -722,7 +722,7 @@ namespace Horde.Server.Ddc
 
             _diagnosticContext.Set("Content-Length", Request.ContentLength ?? -1);
 
-            byte[] b = await RequestUtil.ReadRawBody(Request);
+            byte[] b = await ReadRawBodyAsync(Request);
             CbPackageReader packageReader = await CbPackageReader.Create(new MemoryStream(b));
 
             try
@@ -765,7 +765,33 @@ namespace Horde.Server.Ddc
             return Ok(new PutObjectResponse(missingHashes.ToArray()));
         }
 
-        [HttpPost("{ns}/{bucket}/{key}/finalize/{hash}.{format?}")]
+		static async Task<byte[]> ReadRawBodyAsync(HttpRequest request)
+		{
+			long? contentLength = request.GetTypedHeaders().ContentLength;
+
+			if (contentLength == null)
+			{
+				throw new Exception("Expected content-length on all raw body requests");
+			}
+
+			Tracer? tracer = request.HttpContext.RequestServices.GetService<Tracer>();
+			using TelemetrySpan? scope = tracer?.StartActiveSpan("readbody")
+				.SetAttribute("operation.name", "readbody")
+				.SetAttribute("Content-Length", contentLength.Value);
+
+			await using MemoryStream ms = new MemoryStream((int)contentLength);
+			DateTime readStart = DateTime.Now;
+			await request.Body.CopyToAsync(ms);
+			TimeSpan duration = DateTime.Now - readStart;
+
+#if WITH_DOGSTATSD
+			double rate = contentLength.Value / duration.TotalSeconds;
+            StatsdClient.DogStatsd.Histogram("jupiter.stream_throughput", rate, tags: new string[] {"sourceIdentifier:" + "readbody"});
+#endif
+			return ms.GetBuffer();
+		}
+
+		[HttpPost("{ns}/{bucket}/{key}/finalize/{hash}.{format?}")]
         public async Task<IActionResult> FinalizeObject(
             [FromRoute] [Required] NamespaceId ns,
             [FromRoute] [Required] BucketId bucket,
diff --git a/Engine/Source/Programs/Horde/Horde.Server/Ddc/RequestUtil.cs b/Engine/Source/Programs/Horde/Horde.Server/Ddc/RequestUtil.cs
deleted file mode 100644
