Horde: Use IoHash over IoHashKey.

diff --git a/Engine/Source/Programs/Horde/Horde.Server/Ddc/DdcObjectService.cs b/Engine/Source/Programs/Horde/Horde.Server/Ddc/DdcObjectService.cs
--- a/Engine/Source/Programs/Horde/Horde.Server/Ddc/DdcObjectService.cs
+++ b/Engine/Source/Programs/Horde/Horde.Server/Ddc/DdcObjectService.cs
@@ -2,6 +2,7 @@
 
 using System.Collections.Generic;
 using System.Threading.Tasks;
+using EpicGames.Core;
 using EpicGames.Horde.Storage;
 using EpicGames.Serialization;
 
@@ -9,7 +10,7 @@ namespace Horde.Server.Ddc
 {
 	class ObjectService : IDdcObjectService
 	{
-		public Task<bool> Delete(NamespaceId ns, BucketId bucket, IoHashKey key)
+		public Task<bool> Delete(NamespaceId ns, BucketId bucket, IoHash key)
 		{
 			throw new System.NotImplementedException();
 		}
@@ -24,17 +25,17 @@ namespace Horde.Server.Ddc
 			throw new System.NotImplementedException();
 		}
 
-		public Task<bool> Exists(NamespaceId ns, BucketId bucket, IoHashKey key)
+		public Task<bool> Exists(NamespaceId ns, BucketId bucket, IoHash key)
 		{
 			throw new System.NotImplementedException();
 		}
 
-		public Task<(ContentId[], BlobIdentifier[])> Finalize(NamespaceId ns, BucketId bucket, IoHashKey key, BlobIdentifier blobHash)
+		public Task<(ContentId[], BlobIdentifier[])> Finalize(NamespaceId ns, BucketId bucket, IoHash key, BlobIdentifier blobHash)
 		{
 			throw new System.NotImplementedException();
 		}
 
-		public Task<(ObjectRecord, BlobContents?)> Get(NamespaceId ns, BucketId bucket, IoHashKey key, string[] fields, bool doLastAccessTracking = true)
+		public Task<(ObjectRecord, BlobContents?)> Get(NamespaceId ns, BucketId bucket, IoHash key, string[] fields, bool doLastAccessTracking = true)
 		{
 			throw new System.NotImplementedException();
 		}
@@ -44,12 +45,12 @@ namespace Horde.Server.Ddc
 			throw new System.NotImplementedException();
 		}
 
-		public Task<List<BlobIdentifier>> GetReferencedBlobs(NamespaceId ns, BucketId bucket, IoHashKey key)
+		public Task<List<BlobIdentifier>> GetReferencedBlobs(NamespaceId ns, BucketId bucket, IoHash key)
 		{
 			throw new System.NotImplementedException();
 		}
 
-		public Task<(ContentId[], BlobIdentifier[])> Put(NamespaceId ns, BucketId bucket, IoHashKey key, BlobIdentifier blobHash, CbObject payload)
+		public Task<(ContentId[], BlobIdentifier[])> Put(NamespaceId ns, BucketId bucket, IoHash key, BlobIdentifier blobHash, CbObject payload)
 		{
 			throw new System.NotImplementedException();
 		}
diff --git a/Engine/Source/Programs/Horde/Horde.Server/Ddc/DdcRefsController.cs b/Engine/Source/Programs/Horde/Horde.Server/Ddc/DdcRefsController.cs
--- a/Engine/Source/Programs/Horde/Horde.Server/Ddc/DdcRefsController.cs
+++ b/Engine/Source/Programs/Horde/Horde.Server/Ddc/DdcRefsController.cs
@@ -103,7 +103,7 @@ namespace Horde.Server.Ddc
         public async Task<IActionResult> Get(
             [FromRoute] [Required] NamespaceId ns,
             [FromRoute] [Required] BucketId bucket,
-            [FromRoute] [Required] IoHashKey key,
+            [FromRoute] [Required] IoHash key,
             [FromRoute] string? format = null)
         {
             ActionResult? accessResult = await _requestHelper.HasAccessToNamespace(User, Request, ns, new [] { DdcAclAction.ReadObject });
@@ -434,7 +434,7 @@ namespace Horde.Server.Ddc
     public async Task<IActionResult> GetMetadata(
         [FromRoute] [Required] NamespaceId ns,
         [FromRoute] [Required] BucketId bucket,
-        [FromRoute] [Required] IoHashKey key,
+        [FromRoute] [Required] IoHash key,
         [FromQuery] string[] fields)
     {
         ActionResult? accessResult = await _requestHelper.HasAccessToNamespace(User, Request, ns, new [] { DdcAclAction.ReadObject });
@@ -476,7 +476,7 @@ namespace Horde.Server.Ddc
         public async Task<IActionResult> Head(
             [FromRoute] [Required] NamespaceId ns,
             [FromRoute] [Required] BucketId bucket,
-            [FromRoute] [Required] IoHashKey key)
+            [FromRoute] [Required] IoHash key)
         {
             ActionResult? accessResult = await _requestHelper.HasAccessToNamespace(User, Request, ns, new [] { DdcAclAction.ReadObject });
             if (accessResult != null)
@@ -548,9 +548,9 @@ namespace Horde.Server.Ddc
                 return accessResult;
             }
 
-            ConcurrentBag<(BucketId, IoHashKey)> missingObject = new ();
+            ConcurrentBag<(BucketId, IoHash)> missingObject = new ();
 
-            List<(BucketId, IoHashKey)> requestedNames = new List<(BucketId, IoHashKey)>();
+            List<(BucketId, IoHash)> requestedNames = new List<(BucketId, IoHash)>();
             foreach (string name in names)
             {
                 int separatorIndex = name.IndexOf(".", StringComparison.Ordinal);
@@ -560,13 +560,13 @@ namespace Horde.Server.Ddc
                 }
 
                 BucketId bucket = new BucketId(name.Substring(0, separatorIndex));
-                IoHashKey key = new IoHashKey(name.Substring(separatorIndex + 1));
+				IoHash key = IoHash.Parse(name.Substring(separatorIndex + 1));
                 requestedNames.Add((bucket, key));
             }
 
             IEnumerable<Task> tasks = requestedNames.Select(async pair =>
             {
-                (BucketId bucket, IoHashKey key) = pair;
+                (BucketId bucket, IoHash key) = pair;
                 try
                 {
                     (ObjectRecord record, BlobContents? blob) =
@@ -609,7 +609,7 @@ namespace Horde.Server.Ddc
         public async Task<IActionResult> PutObject(
             [FromRoute] [Required] NamespaceId ns,
             [FromRoute] [Required] BucketId bucket,
-            [FromRoute] [Required] IoHashKey key)
+            [FromRoute] [Required] IoHash key)
         {
             ActionResult? accessResult = await _requestHelper.HasAccessToNamespace(User, Request, ns, new [] { DdcAclAction.WriteObject });
             if (accessResult != null)
@@ -712,7 +712,7 @@ namespace Horde.Server.Ddc
         public async Task<IActionResult> PutPackage(
             [FromRoute][Required] NamespaceId ns,
             [FromRoute][Required] BucketId bucket,
-            [FromRoute][Required] IoHashKey key)
+            [FromRoute][Required] IoHash key)
         {
             ActionResult? accessResult = await _requestHelper.HasAccessToNamespace(User, Request, ns, new [] { DdcAclAction.WriteObject });
             if (accessResult != null)
@@ -795,7 +795,7 @@ namespace Horde.Server.Ddc
         public async Task<IActionResult> FinalizeObject(
             [FromRoute] [Required] NamespaceId ns,
             [FromRoute] [Required] BucketId bucket,
-            [FromRoute] [Required] IoHashKey key,
+            [FromRoute] [Required] IoHash key,
             [FromRoute] [Required] BlobIdentifier hash)
         {
             ActionResult? accessResult = await _requestHelper.HasAccessToNamespace(User, Request, ns, new [] { DdcAclAction.WriteObject });
@@ -1090,7 +1090,7 @@ namespace Horde.Server.Ddc
         public async Task<IActionResult> Delete(
             [FromRoute] [Required] NamespaceId ns,
             [FromRoute] [Required] BucketId bucket,
-            [FromRoute] [Required] IoHashKey key)
+            [FromRoute] [Required] IoHash key)
         {
             ActionResult? accessResult = await _requestHelper.HasAccessToNamespace(User, Request, ns, new [] { DdcAclAction.DeleteObject });
             if (accessResult != null)
@@ -1196,7 +1196,7 @@ namespace Horde.Server.Ddc
 
             [Required]
             [CbField("key")]
-            public IoHashKey Key { get; set; }
+            public IoHash Key { get; set; }
 
             [CbField("resolveAttachments")]
             public bool? ResolveAttachments { get; set; } = null;
@@ -1250,7 +1250,7 @@ namespace Horde.Server.Ddc
         }
 
         [JsonConstructor]
-        public RefMetadataResponse(NamespaceId ns, BucketId bucket, IoHashKey name, BlobIdentifier payloadIdentifier, DateTime lastAccess, bool isFinalized, byte[]? inlinePayload)
+        public RefMetadataResponse(NamespaceId ns, BucketId bucket, IoHash name, BlobIdentifier payloadIdentifier, DateTime lastAccess, bool isFinalized, byte[]? inlinePayload)
         {
             Ns = ns;
             Bucket = bucket;
@@ -1279,7 +1279,7 @@ namespace Horde.Server.Ddc
         public BucketId Bucket { get; set; }
 
         [CbField("name")]
-        public IoHashKey Name { get; set; }
+        public IoHash Name { get; set; }
 
         [CbField("payloadIdentifier")]
         public BlobIdentifier PayloadIdentifier { get; set; }
@@ -1312,11 +1312,11 @@ namespace Horde.Server.Ddc
 
     public class ExistCheckMultipleRefsResponse
     {
-        public ExistCheckMultipleRefsResponse(List<(BucketId,IoHashKey)> missing)
+        public ExistCheckMultipleRefsResponse(List<(BucketId, IoHash)> missing)
         {
             Missing = missing.Select(pair =>
             {
-                (BucketId bucketId, IoHashKey ioHashKey) = pair;
+                (BucketId bucketId, IoHash ioHashKey) = pair;
                 return new MissingReference()
                 {
                     Bucket = bucketId,
@@ -1342,7 +1342,7 @@ namespace Horde.Server.Ddc
             public BucketId Bucket { get; set; }
 
             [CbField("key")]
-            public IoHashKey Key { get; set; }
+            public IoHash Key { get; set; }
         }
     }
 
diff --git a/Engine/Source/Programs/Horde/Horde.Server/Ddc/IDdcObjectService.cs b/Engine/Source/Programs/Horde/Horde.Server/Ddc/IDdcObjectService.cs
--- a/Engine/Source/Programs/Horde/Horde.Server/Ddc/IDdcObjectService.cs
+++ b/Engine/Source/Programs/Horde/Horde.Server/Ddc/IDdcObjectService.cs
@@ -3,6 +3,7 @@
 using System;
 using System.Collections.Generic;
 using System.Threading.Tasks;
+using EpicGames.Core;
 using EpicGames.Horde.Storage;
 using EpicGames.Serialization;
 
@@ -12,23 +13,23 @@ namespace Horde.Server.Ddc
 {
 	public interface IDdcObjectService
     {
-        Task<(ObjectRecord, BlobContents?)> Get(NamespaceId ns, BucketId bucket, IoHashKey key, string[] fields, bool doLastAccessTracking = true);
-        Task<(ContentId[], BlobIdentifier[])> Put(NamespaceId ns, BucketId bucket, IoHashKey key, BlobIdentifier blobHash, CbObject payload);
-        Task<(ContentId[], BlobIdentifier[])> Finalize(NamespaceId ns, BucketId bucket, IoHashKey key, BlobIdentifier blobHash);
+        Task<(ObjectRecord, BlobContents?)> Get(NamespaceId ns, BucketId bucket, IoHash key, string[] fields, bool doLastAccessTracking = true);
+        Task<(ContentId[], BlobIdentifier[])> Put(NamespaceId ns, BucketId bucket, IoHash key, BlobIdentifier blobHash, CbObject payload);
+        Task<(ContentId[], BlobIdentifier[])> Finalize(NamespaceId ns, BucketId bucket, IoHash key, BlobIdentifier blobHash);
 
         IAsyncEnumerable<NamespaceId> GetNamespaces();
 
-        Task<bool> Delete(NamespaceId ns, BucketId bucket, IoHashKey key);
+        Task<bool> Delete(NamespaceId ns, BucketId bucket, IoHash key);
         Task<long> DropNamespace(NamespaceId ns);
         Task<long> DeleteBucket(NamespaceId ns, BucketId bucket);
 
-        Task<bool> Exists(NamespaceId ns, BucketId bucket, IoHashKey key);
-        Task<List<BlobIdentifier>> GetReferencedBlobs(NamespaceId ns, BucketId bucket, IoHashKey key);
+        Task<bool> Exists(NamespaceId ns, BucketId bucket, IoHash key);
+        Task<List<BlobIdentifier>> GetReferencedBlobs(NamespaceId ns, BucketId bucket, IoHash key);
     }
 
 	public class ObjectRecord
 	{
-		public ObjectRecord(NamespaceId ns, BucketId bucket, IoHashKey name, DateTime lastAccess, byte[]? inlinePayload, BlobIdentifier blobIdentifier, bool isFinalized)
+		public ObjectRecord(NamespaceId ns, BucketId bucket, IoHash name, DateTime lastAccess, byte[]? inlinePayload, BlobIdentifier blobIdentifier, bool isFinalized)
 		{
 			Namespace = ns;
 			Bucket = bucket;
@@ -41,7 +42,7 @@ namespace Horde.Server.Ddc
 
 		public NamespaceId Namespace { get; }
 		public BucketId Bucket { get; }
-		public IoHashKey Name { get; }
+		public IoHash Name { get; }
 		public DateTime LastAccess { get; }
 		public byte[]? InlinePayload { get; set; }
 		public BlobIdentifier BlobIdentifier { get; set; }
@@ -50,7 +51,7 @@ namespace Horde.Server.Ddc
 
 	public class ObjectNotFoundException : Exception
 	{
-		public ObjectNotFoundException(NamespaceId ns, BucketId bucket, IoHashKey key) : base($"Object not found {key} in bucket {bucket} namespace {ns}")
+		public ObjectNotFoundException(NamespaceId ns, BucketId bucket, IoHash key) : base($"Object not found {key} in bucket {bucket} namespace {ns}")
 		{
 			Namespace = ns;
 			Bucket = bucket;
@@ -59,12 +60,12 @@ namespace Horde.Server.Ddc
 
 		public NamespaceId Namespace { get; }
 		public BucketId Bucket { get; }
-		public IoHashKey Key { get; }
+		public IoHash Key { get; }
 	}
 
 	public class ObjectHashMismatchException : Exception
     {
-        public ObjectHashMismatchException(NamespaceId ns, BucketId bucket, IoHashKey name, BlobIdentifier suppliedHash, BlobIdentifier actualHash) : base($"Object {name} in bucket {bucket} and namespace {ns} did not reference hash {suppliedHash} was referencing {actualHash}")
+        public ObjectHashMismatchException(NamespaceId ns, BucketId bucket, IoHash name, BlobIdentifier suppliedHash, BlobIdentifier actualHash) : base($"Object {name} in bucket {bucket} and namespace {ns} did not reference hash {suppliedHash} was referencing {actualHash}")
         {
         }
     }
diff --git a/Engine/Source/Programs/Horde/Horde.Server/Ddc/IoHashKey.cs b/Engine/Source/Programs/Horde/Horde.Server/Ddc/IoHashKey.cs
deleted file mode 100644
