Horde: Rename IDdcObjectService to IDdcRefService.

diff --git a/Engine/Source/Programs/Horde/Horde.Server/Ddc/DdcBlobService.cs b/Engine/Source/Programs/Horde/Horde.Server/Ddc/DdcBlobService.cs
--- a/Engine/Source/Programs/Horde/Horde.Server/Ddc/DdcBlobService.cs
+++ b/Engine/Source/Programs/Horde/Horde.Server/Ddc/DdcBlobService.cs
@@ -8,7 +8,7 @@ using EpicGames.Horde.Storage;
 
 namespace Horde.Server.Ddc
 {
-	class BlobService : IDdcBlobService
+	class DdcBlobService : IDdcBlobService
 	{
 		public Task DeleteNamespace(NamespaceId ns)
 		{
diff --git a/Engine/Source/Programs/Horde/Horde.Server/Ddc/DdcObjectService.cs b/Engine/Source/Programs/Horde/Horde.Server/Ddc/DdcObjectService.cs
deleted file mode 100644
diff --git a/Engine/Source/Programs/Horde/Horde.Server/Ddc/DdcRefService.cs b/Engine/Source/Programs/Horde/Horde.Server/Ddc/DdcRefService.cs
new file mode 100644
--- /dev/null
+++ b/Engine/Source/Programs/Horde/Horde.Server/Ddc/DdcRefService.cs
@@ -0,0 +1,58 @@
+// Copyright Epic Games, Inc. All Rights Reserved.
+
+using System.Collections.Generic;
+using System.Threading.Tasks;
+using EpicGames.Core;
+using EpicGames.Horde.Storage;
+using EpicGames.Serialization;
+
+namespace Horde.Server.Ddc
+{
+	class DdcRefService : IDdcRefService
+	{
+		public Task<bool> Delete(NamespaceId ns, BucketId bucket, IoHash key)
+		{
+			throw new System.NotImplementedException();
+		}
+
+		public Task<long> DeleteBucket(NamespaceId ns, BucketId bucket)
+		{
+			throw new System.NotImplementedException();
+		}
+
+		public Task<long> DropNamespace(NamespaceId ns)
+		{
+			throw new System.NotImplementedException();
+		}
+
+		public Task<bool> Exists(NamespaceId ns, BucketId bucket, IoHash key)
+		{
+			throw new System.NotImplementedException();
+		}
+
+		public Task<(ContentId[], BlobIdentifier[])> Finalize(NamespaceId ns, BucketId bucket, IoHash key, BlobIdentifier blobHash)
+		{
+			throw new System.NotImplementedException();
+		}
+
+		public Task<(ObjectRecord, BlobContents?)> Get(NamespaceId ns, BucketId bucket, IoHash key, string[] fields, bool doLastAccessTracking = true)
+		{
+			throw new System.NotImplementedException();
+		}
+
+		public IAsyncEnumerable<NamespaceId> GetNamespaces()
+		{
+			throw new System.NotImplementedException();
+		}
+
+		public Task<List<BlobIdentifier>> GetReferencedBlobs(NamespaceId ns, BucketId bucket, IoHash key)
+		{
+			throw new System.NotImplementedException();
+		}
+
+		public Task<(ContentId[], BlobIdentifier[])> Put(NamespaceId ns, BucketId bucket, IoHash key, BlobIdentifier blobHash, CbObject payload)
+		{
+			throw new System.NotImplementedException();
+		}
+	}
+}
diff --git a/Engine/Source/Programs/Horde/Horde.Server/Ddc/DdcRefsController.cs b/Engine/Source/Programs/Horde/Horde.Server/Ddc/DdcRefsController.cs
--- a/Engine/Source/Programs/Horde/Horde.Server/Ddc/DdcRefsController.cs
+++ b/Engine/Source/Programs/Horde/Horde.Server/Ddc/DdcRefsController.cs
@@ -49,10 +49,10 @@ namespace Horde.Server.Ddc
         private readonly Tracer _tracer;
 
         private readonly ILogger _logger;
-        private readonly IDdcObjectService _objectService;
+        private readonly IDdcRefService _objectService;
         private readonly IDdcBlobService _blobStore;
 
-        public DdcRefsController(IDdcObjectService objectService, IDdcBlobService blobStore, IDiagnosticContext diagnosticContext, FormatResolver formatResolver, BufferedPayloadFactory bufferedPayloadFactory, IReferenceResolver referenceResolver, NginxRedirectHelper nginxRedirectHelper, IRequestHelper requestHelper, Tracer tracer, ILogger<DdcRefsController> logger)
+        public DdcRefsController(IDdcRefService objectService, IDdcBlobService blobStore, IDiagnosticContext diagnosticContext, FormatResolver formatResolver, BufferedPayloadFactory bufferedPayloadFactory, IReferenceResolver referenceResolver, NginxRedirectHelper nginxRedirectHelper, IRequestHelper requestHelper, Tracer tracer, ILogger<DdcRefsController> logger)
         {
             _objectService = objectService;
             _blobStore = blobStore;
@@ -806,7 +806,7 @@ namespace Horde.Server.Ddc
 
             try
             {
-                (ContentId[] missingReferences, BlobIdentifier[] missingBlobs) = await _objectService.Finalize(ns, bucket, key, hash);
+                (ContentId[] missingReferences, BlobIdentifier[] missingBlobs) = await _objectService.IDdcRefService(ns, bucket, key, hash);
                 List<JupiterContentHash> missingHashes = new List<JupiterContentHash>(missingReferences);
                 missingHashes.AddRange(missingBlobs);
 
diff --git a/Engine/Source/Programs/Horde/Horde.Server/Ddc/IDdcObjectService.cs b/Engine/Source/Programs/Horde/Horde.Server/Ddc/IDdcObjectService.cs
deleted file mode 100644
diff --git a/Engine/Source/Programs/Horde/Horde.Server/Ddc/IDdcRefService.cs b/Engine/Source/Programs/Horde/Horde.Server/Ddc/IDdcRefService.cs
new file mode 100644
--- /dev/null
+++ b/Engine/Source/Programs/Horde/Horde.Server/Ddc/IDdcRefService.cs
@@ -0,0 +1,72 @@
+// Copyright Epic Games, Inc. All Rights Reserved.
+
+using System;
+using System.Collections.Generic;
+using System.Threading.Tasks;
+using EpicGames.Core;
+using EpicGames.Horde.Storage;
+using EpicGames.Serialization;
+
+#pragma warning disable CS1591
+
+namespace Horde.Server.Ddc
+{
+	public interface IDdcRefService
+    {
+        Task<(ObjectRecord, BlobContents?)> Get(NamespaceId ns, BucketId bucket, IoHash key, string[] fields, bool doLastAccessTracking = true);
+        Task<(ContentId[], BlobIdentifier[])> Put(NamespaceId ns, BucketId bucket, IoHash key, BlobIdentifier blobHash, CbObject payload);
+        Task<(ContentId[], BlobIdentifier[])> IDdcRefService(NamespaceId ns, BucketId bucket, IoHash key, BlobIdentifier blobHash);
+
+        IAsyncEnumerable<NamespaceId> GetNamespaces();
+
+        Task<bool> Delete(NamespaceId ns, BucketId bucket, IoHash key);
+        Task<long> DropNamespace(NamespaceId ns);
+        Task<long> DeleteBucket(NamespaceId ns, BucketId bucket);
+
+        Task<bool> Exists(NamespaceId ns, BucketId bucket, IoHash key);
+        Task<List<BlobIdentifier>> GetReferencedBlobs(NamespaceId ns, BucketId bucket, IoHash key);
+    }
+
+	public class ObjectRecord
+	{
+		public ObjectRecord(NamespaceId ns, BucketId bucket, IoHash name, DateTime lastAccess, byte[]? inlinePayload, BlobIdentifier blobIdentifier, bool isFinalized)
+		{
+			Namespace = ns;
+			Bucket = bucket;
+			Name = name;
+			LastAccess = lastAccess;
+			InlinePayload = inlinePayload;
+			BlobIdentifier = blobIdentifier;
+			IsFinalized = isFinalized;
+		}
+
+		public NamespaceId Namespace { get; }
+		public BucketId Bucket { get; }
+		public IoHash Name { get; }
+		public DateTime LastAccess { get; }
+		public byte[]? InlinePayload { get; set; }
+		public BlobIdentifier BlobIdentifier { get; set; }
+		public bool IsFinalized { get; }
+	}
+
+	public class ObjectNotFoundException : Exception
+	{
+		public ObjectNotFoundException(NamespaceId ns, BucketId bucket, IoHash key) : base($"Object not found {key} in bucket {bucket} namespace {ns}")
+		{
+			Namespace = ns;
+			Bucket = bucket;
+			Key = key;
+		}
+
+		public NamespaceId Namespace { get; }
+		public BucketId Bucket { get; }
+		public IoHash Key { get; }
+	}
+
+	public class ObjectHashMismatchException : Exception
+    {
+        public ObjectHashMismatchException(NamespaceId ns, BucketId bucket, IoHash name, BlobIdentifier suppliedHash, BlobIdentifier actualHash) : base($"Object {name} in bucket {bucket} and namespace {ns} did not reference hash {suppliedHash} was referencing {actualHash}")
+        {
+        }
+    }
+}
diff --git a/Engine/Source/Programs/Horde/Horde.Server/Startup.cs b/Engine/Source/Programs/Horde/Horde.Server/Startup.cs
--- a/Engine/Source/Programs/Horde/Horde.Server/Startup.cs
+++ b/Engine/Source/Programs/Horde/Horde.Server/Startup.cs
@@ -521,8 +521,8 @@ namespace Horde.Server
 			}
 
 			// DDC
-			services.AddSingleton<IDdcBlobService, BlobService>();
-			services.AddSingleton<IDdcObjectService, ObjectService>();
+			services.AddSingleton<IDdcBlobService, DdcBlobService>();
+			services.AddSingleton<IDdcRefService, DdcRefService>();
 			services.AddSingleton<IReferenceResolver, ReferenceResolver>();
 			services.AddSingleton<IDdcContentIdService, ContentIdStore>();
 
