Horde: Add a custom type for ref ids, to ensure they don't get mixed up with blob ids.

diff --git a/Engine/Source/Programs/Horde/Horde.Server/Ddc/DdcRefService.cs b/Engine/Source/Programs/Horde/Horde.Server/Ddc/DdcRefService.cs
--- a/Engine/Source/Programs/Horde/Horde.Server/Ddc/DdcRefService.cs
+++ b/Engine/Source/Programs/Horde/Horde.Server/Ddc/DdcRefService.cs
@@ -10,7 +10,7 @@ namespace Horde.Server.Ddc
 {
 	class DdcRefService : IDdcRefService
 	{
-		public Task<bool> Delete(NamespaceId ns, BucketId bucket, IoHash key)
+		public Task<bool> Delete(NamespaceId ns, BucketId bucket, RefId key)
 		{
 			throw new System.NotImplementedException();
 		}
@@ -25,17 +25,17 @@ namespace Horde.Server.Ddc
 			throw new System.NotImplementedException();
 		}
 
-		public Task<bool> Exists(NamespaceId ns, BucketId bucket, IoHash key)
+		public Task<bool> Exists(NamespaceId ns, BucketId bucket, RefId key)
 		{
 			throw new System.NotImplementedException();
 		}
 
-		public Task<(ContentId[], BlobIdentifier[])> Finalize(NamespaceId ns, BucketId bucket, IoHash key, BlobIdentifier blobHash)
+		public Task<(ContentId[], BlobIdentifier[])> Finalize(NamespaceId ns, BucketId bucket, RefId key, BlobIdentifier blobHash)
 		{
 			throw new System.NotImplementedException();
 		}
 
-		public Task<(ObjectRecord, BlobContents?)> Get(NamespaceId ns, BucketId bucket, IoHash key, string[] fields, bool doLastAccessTracking = true)
+		public Task<(RefRecord, BlobContents?)> Get(NamespaceId ns, BucketId bucket, RefId key, string[] fields, bool doLastAccessTracking = true)
 		{
 			throw new System.NotImplementedException();
 		}
@@ -45,12 +45,12 @@ namespace Horde.Server.Ddc
 			throw new System.NotImplementedException();
 		}
 
-		public Task<List<BlobIdentifier>> GetReferencedBlobs(NamespaceId ns, BucketId bucket, IoHash key)
+		public Task<List<BlobIdentifier>> GetReferencedBlobs(NamespaceId ns, BucketId bucket, RefId key)
 		{
 			throw new System.NotImplementedException();
 		}
 
-		public Task<(ContentId[], BlobIdentifier[])> Put(NamespaceId ns, BucketId bucket, IoHash key, BlobIdentifier blobHash, CbObject payload)
+		public Task<(ContentId[], BlobIdentifier[])> Put(NamespaceId ns, BucketId bucket, RefId key, BlobIdentifier blobHash, CbObject payload)
 		{
 			throw new System.NotImplementedException();
 		}
diff --git a/Engine/Source/Programs/Horde/Horde.Server/Ddc/DdcRefsController.cs b/Engine/Source/Programs/Horde/Horde.Server/Ddc/DdcRefsController.cs
--- a/Engine/Source/Programs/Horde/Horde.Server/Ddc/DdcRefsController.cs
+++ b/Engine/Source/Programs/Horde/Horde.Server/Ddc/DdcRefsController.cs
@@ -49,12 +49,12 @@ namespace Horde.Server.Ddc
         private readonly Tracer _tracer;
 
         private readonly ILogger _logger;
-        private readonly IDdcRefService _objectService;
+        private readonly IDdcRefService _refService;
         private readonly IDdcBlobService _blobStore;
 
         public DdcRefsController(IDdcRefService objectService, IDdcBlobService blobStore, IDiagnosticContext diagnosticContext, FormatResolver formatResolver, BufferedPayloadFactory bufferedPayloadFactory, IReferenceResolver referenceResolver, NginxRedirectHelper nginxRedirectHelper, IRequestHelper requestHelper, Tracer tracer, ILogger<DdcRefsController> logger)
         {
-            _objectService = objectService;
+            _refService = objectService;
             _blobStore = blobStore;
             _diagnosticContext = diagnosticContext;
             _formatResolver = formatResolver;
@@ -75,7 +75,7 @@ namespace Horde.Server.Ddc
         [ProducesResponseType(type: typeof(ProblemDetails), 400)]
         public async Task<IActionResult> GetNamespaces()
         {
-            NamespaceId[] namespaces = await _objectService.GetNamespaces().ToArrayAsync();
+            NamespaceId[] namespaces = await _refService.GetNamespaces().ToArrayAsync();
 
             // filter namespaces down to only the namespaces the user has access to
             List<NamespaceId> namespacesWithAccess = new();
@@ -103,7 +103,7 @@ namespace Horde.Server.Ddc
         public async Task<IActionResult> Get(
             [FromRoute] [Required] NamespaceId ns,
             [FromRoute] [Required] BucketId bucket,
-            [FromRoute] [Required] IoHash key,
+            [FromRoute] [Required] RefId key,
             [FromRoute] string? format = null)
         {
             ActionResult? accessResult = await _requestHelper.HasAccessToNamespace(User, Request, ns, new [] { DdcAclAction.ReadObject });
@@ -114,7 +114,7 @@ namespace Horde.Server.Ddc
 
             try
             {
-                (ObjectRecord objectRecord, BlobContents? blob) = await _objectService.Get(ns, bucket, key, Array.Empty<string>());
+                (RefRecord objectRecord, BlobContents? blob) = await _refService.Get(ns, bucket, key, Array.Empty<string>());
 
                 if (blob == null)
                 {
@@ -411,7 +411,7 @@ namespace Horde.Server.Ddc
             {
                 return NotFound(new ProblemDetails {Title = $"Namespace {e.Namespace} did not exist"});
             }
-            catch (ObjectNotFoundException e)
+            catch (RefNotFoundException e)
             {
                 return NotFound(new ProblemDetails { Title = $"Object {e.Bucket} {e.Key} did not exist" });
             }
@@ -434,7 +434,7 @@ namespace Horde.Server.Ddc
     public async Task<IActionResult> GetMetadata(
         [FromRoute] [Required] NamespaceId ns,
         [FromRoute] [Required] BucketId bucket,
-        [FromRoute] [Required] IoHash key,
+        [FromRoute] [Required] RefId key,
         [FromQuery] string[] fields)
     {
         ActionResult? accessResult = await _requestHelper.HasAccessToNamespace(User, Request, ns, new [] { DdcAclAction.ReadObject });
@@ -445,7 +445,7 @@ namespace Horde.Server.Ddc
 
         try
         {
-            (ObjectRecord objectRecord, BlobContents? _) = await _objectService.Get(ns, bucket, key, fields);
+            (RefRecord objectRecord, BlobContents? _) = await _refService.Get(ns, bucket, key, fields);
 
             return Ok(new RefMetadataResponse(objectRecord));
         }
@@ -453,7 +453,7 @@ namespace Horde.Server.Ddc
         {
             return NotFound(new ProblemDetails {Title = $"Namespace {e.Namespace} did not exist"});
         }
-        catch (ObjectNotFoundException e)
+        catch (RefNotFoundException e)
         {
             return NotFound(new ProblemDetails { Title = $"Object {e.Bucket} {e.Key} did not exist" });
         }
@@ -476,7 +476,7 @@ namespace Horde.Server.Ddc
         public async Task<IActionResult> Head(
             [FromRoute] [Required] NamespaceId ns,
             [FromRoute] [Required] BucketId bucket,
-            [FromRoute] [Required] IoHash key)
+            [FromRoute] [Required] RefId key)
         {
             ActionResult? accessResult = await _requestHelper.HasAccessToNamespace(User, Request, ns, new [] { DdcAclAction.ReadObject });
             if (accessResult != null)
@@ -486,7 +486,7 @@ namespace Horde.Server.Ddc
 
             try
             {
-                (ObjectRecord record, BlobContents? blob) = await _objectService.Get(ns, bucket, key, new string[] {"blobIdentifier", "IsFinalized"});
+                (RefRecord record, BlobContents? blob) = await _refService.Get(ns, bucket, key, new string[] {"blobIdentifier", "IsFinalized"});
                 Response.Headers[CommonHeaders.HashHeaderName] = record.BlobIdentifier.ToString();
 
                 if (!record.IsFinalized)
@@ -520,7 +520,7 @@ namespace Horde.Server.Ddc
             {
                 return NotFound(new ProblemDetails { Title = $"Blob {e.Blob} in namespace {ns} did not exist" });
             }
-            catch (ObjectNotFoundException e)
+            catch (RefNotFoundException e)
             {
                 return NotFound(new ProblemDetails {Title = $"Object {e.Bucket} {e.Key} in namespace {e.Namespace} did not exist"});
             }
@@ -548,9 +548,9 @@ namespace Horde.Server.Ddc
                 return accessResult;
             }
 
-            ConcurrentBag<(BucketId, IoHash)> missingObject = new ();
+            ConcurrentBag<(BucketId, RefId)> missingObject = new ();
 
-            List<(BucketId, IoHash)> requestedNames = new List<(BucketId, IoHash)>();
+            List<(BucketId, RefId)> requestedNames = new List<(BucketId, RefId)>();
             foreach (string name in names)
             {
                 int separatorIndex = name.IndexOf(".", StringComparison.Ordinal);
@@ -560,17 +560,17 @@ namespace Horde.Server.Ddc
                 }
 
                 BucketId bucket = new BucketId(name.Substring(0, separatorIndex));
-				IoHash key = IoHash.Parse(name.Substring(separatorIndex + 1));
+				RefId key = RefId.Parse(name.Substring(separatorIndex + 1));
                 requestedNames.Add((bucket, key));
             }
 
             IEnumerable<Task> tasks = requestedNames.Select(async pair =>
             {
-                (BucketId bucket, IoHash key) = pair;
+                (BucketId bucket, RefId key) = pair;
                 try
                 {
-                    (ObjectRecord record, BlobContents? blob) =
-                        await _objectService.Get(ns, bucket, key, new string[] { "blobIdentifier" });
+                    (RefRecord record, BlobContents? blob) =
+                        await _refService.Get(ns, bucket, key, new string[] { "blobIdentifier" });
 
                     blob ??= await _blobStore.GetObject(ns, record.BlobIdentifier);
 
@@ -582,7 +582,7 @@ namespace Horde.Server.Ddc
                     IAsyncEnumerable<BlobIdentifier> references = _referenceResolver.GetReferencedBlobs(ns, cb);
                     List<BlobIdentifier>? _ = await references.ToListAsync();
                 }
-                catch (ObjectNotFoundException)
+                catch (RefNotFoundException)
                 {
                     missingObject.Add((bucket, key));
                 }
@@ -609,7 +609,7 @@ namespace Horde.Server.Ddc
         public async Task<IActionResult> PutObject(
             [FromRoute] [Required] NamespaceId ns,
             [FromRoute] [Required] BucketId bucket,
-            [FromRoute] [Required] IoHash key)
+            [FromRoute] [Required] RefId key)
         {
             ActionResult? accessResult = await _requestHelper.HasAccessToNamespace(User, Request, ns, new [] { DdcAclAction.WriteObject });
             if (accessResult != null)
@@ -699,7 +699,7 @@ namespace Horde.Server.Ddc
                 return Problem(e.Message, null, (int)HttpStatusCode.RequestTimeout);
             }
 
-            (ContentId[] missingReferences, BlobIdentifier[] missingBlobs) = await _objectService.Put(ns, bucket, key, blobHeader, payloadObject);
+            (ContentId[] missingReferences, BlobIdentifier[] missingBlobs) = await _refService.Put(ns, bucket, key, blobHeader, payloadObject);
 
             List<JupiterContentHash> missingHashes = new List<JupiterContentHash>(missingReferences);
             missingHashes.AddRange(missingBlobs);
@@ -712,7 +712,7 @@ namespace Horde.Server.Ddc
         public async Task<IActionResult> PutPackage(
             [FromRoute][Required] NamespaceId ns,
             [FromRoute][Required] BucketId bucket,
-            [FromRoute][Required] IoHash key)
+            [FromRoute][Required] RefId key)
         {
             ActionResult? accessResult = await _requestHelper.HasAccessToNamespace(User, Request, ns, new [] { DdcAclAction.WriteObject });
             if (accessResult != null)
@@ -758,7 +758,7 @@ namespace Horde.Server.Ddc
             CbObject rootObject = packageReader.RootObject;
             BlobIdentifier rootObjectHash = BlobIdentifier.FromIoHash(packageReader.RootHash);
 
-            (ContentId[] missingReferences, BlobIdentifier[] missingBlobs) = await _objectService.Put(ns, bucket, key, rootObjectHash, rootObject);
+            (ContentId[] missingReferences, BlobIdentifier[] missingBlobs) = await _refService.Put(ns, bucket, key, rootObjectHash, rootObject);
 
             List<JupiterContentHash> missingHashes = new List<JupiterContentHash>(missingReferences);
             missingHashes.AddRange(missingBlobs);
@@ -795,7 +795,7 @@ namespace Horde.Server.Ddc
         public async Task<IActionResult> FinalizeObject(
             [FromRoute] [Required] NamespaceId ns,
             [FromRoute] [Required] BucketId bucket,
-            [FromRoute] [Required] IoHash key,
+            [FromRoute] [Required] RefId key,
             [FromRoute] [Required] BlobIdentifier hash)
         {
             ActionResult? accessResult = await _requestHelper.HasAccessToNamespace(User, Request, ns, new [] { DdcAclAction.WriteObject });
@@ -806,7 +806,7 @@ namespace Horde.Server.Ddc
 
             try
             {
-                (ContentId[] missingReferences, BlobIdentifier[] missingBlobs) = await _objectService.IDdcRefService(ns, bucket, key, hash);
+                (ContentId[] missingReferences, BlobIdentifier[] missingBlobs) = await _refService.Finalize(ns, bucket, key, hash);
                 List<JupiterContentHash> missingHashes = new List<JupiterContentHash>(missingReferences);
                 missingHashes.AddRange(missingBlobs);
 
@@ -816,7 +816,7 @@ namespace Horde.Server.Ddc
             {
                 return BadRequest(e.Message);
             }
-            catch (ObjectNotFoundException e)
+            catch (RefNotFoundException e)
             {
                 return NotFound(e.Message);
             }
@@ -867,7 +867,7 @@ namespace Horde.Server.Ddc
             {
                 try
                 {
-                    (ObjectRecord objectRecord, BlobContents? blob) = await _objectService.Get(ns, op.Bucket, op.Key, Array.Empty<string>());
+                    (RefRecord objectRecord, BlobContents? blob) = await _refService.Get(ns, op.Bucket, op.Key, Array.Empty<string>());
 
                     if (!objectRecord.IsFinalized)
                     {
@@ -889,7 +889,7 @@ namespace Horde.Server.Ddc
 
                     return (cb, HttpStatusCode.OK);
                 }
-                catch (Exception ex) when( ex is ObjectNotFoundException or PartialReferenceResolveException or ReferenceIsMissingBlobsException)
+                catch (Exception ex) when( ex is RefNotFoundException or PartialReferenceResolveException or ReferenceIsMissingBlobsException)
                 {
                     return ToErrorResult(ex, HttpStatusCode.NotFound);
                 }
@@ -903,7 +903,7 @@ namespace Horde.Server.Ddc
             {
                 try
                 {
-                    (ObjectRecord record, BlobContents? blob) = await _objectService.Get(ns, op.Bucket, op.Key, new string[] { "blobIdentifier" });
+                    (RefRecord record, BlobContents? blob) = await _refService.Get(ns, op.Bucket, op.Key, new string[] { "blobIdentifier" });
 
                     if (!record.IsFinalized)
                     {
@@ -928,7 +928,7 @@ namespace Horde.Server.Ddc
 
                     return (CbObject.Build(writer => writer.WriteBool("exists", true)), HttpStatusCode.OK);
                 }
-                catch (Exception ex) when( ex is ObjectNotFoundException or PartialReferenceResolveException or ReferenceIsMissingBlobsException)
+                catch (Exception ex) when( ex is RefNotFoundException or PartialReferenceResolveException or ReferenceIsMissingBlobsException)
                 {
                     return (CbObject.Build(writer => writer.WriteBool("exists", false)), HttpStatusCode.NotFound);
                 }
@@ -959,7 +959,7 @@ namespace Horde.Server.Ddc
                         throw new HashMismatchException(headerHash, objectHash);
                     }
 
-                    (ContentId[] missingReferences, BlobIdentifier[] missingBlobs) = await _objectService.Put(ns, op.Bucket, op.Key, objectHash, op.Payload);
+                    (ContentId[] missingReferences, BlobIdentifier[] missingBlobs) = await _refService.Put(ns, op.Bucket, op.Key, objectHash, op.Payload);
                     List<JupiterContentHash> missingHashes = new List<JupiterContentHash>(missingReferences);
 
                     return (CbSerializer.Serialize(new PutObjectResponse(missingHashes.ToArray())), HttpStatusCode.OK);
@@ -1038,7 +1038,7 @@ namespace Horde.Server.Ddc
 
             try
             {
-                await _objectService.DropNamespace(ns);
+                await _refService.DropNamespace(ns);
             }
             catch (NamespaceNotFoundException e)
             {
@@ -1068,7 +1068,7 @@ namespace Horde.Server.Ddc
             long countOfDeletedRecords;
             try
             {
-                countOfDeletedRecords = await _objectService.DeleteBucket(ns, bucket);
+                countOfDeletedRecords = await _refService.DeleteBucket(ns, bucket);
             }
             catch (NamespaceNotFoundException e)
             {
@@ -1090,7 +1090,7 @@ namespace Horde.Server.Ddc
         public async Task<IActionResult> Delete(
             [FromRoute] [Required] NamespaceId ns,
             [FromRoute] [Required] BucketId bucket,
-            [FromRoute] [Required] IoHash key)
+            [FromRoute] [Required] RefId key)
         {
             ActionResult? accessResult = await _requestHelper.HasAccessToNamespace(User, Request, ns, new [] { DdcAclAction.DeleteObject });
             if (accessResult != null)
@@ -1100,7 +1100,7 @@ namespace Horde.Server.Ddc
 
             try
             {
-                bool deleted = await _objectService.Delete(ns, bucket, key);
+                bool deleted = await _refService.Delete(ns, bucket, key);
                 if (!deleted)
                 {
                     return NotFound(new ProblemDetails { Title = $"Object {key} in bucket {bucket} and namespace {ns} did not exist" });
@@ -1111,7 +1111,7 @@ namespace Horde.Server.Ddc
             {
                 return NotFound(new ProblemDetails { Title = $"Namespace {e.Namespace} did not exist" });
             }
-            catch (ObjectNotFoundException)
+            catch (RefNotFoundException)
             {
                 return NotFound(new ProblemDetails { Title = $"Object {key} in bucket {bucket} and namespace {ns} did not exist" });
             }
@@ -1196,7 +1196,7 @@ namespace Horde.Server.Ddc
 
             [Required]
             [CbField("key")]
-            public IoHash Key { get; set; }
+            public RefId Key { get; set; }
 
             [CbField("resolveAttachments")]
             public bool? ResolveAttachments { get; set; } = null;
@@ -1250,7 +1250,7 @@ namespace Horde.Server.Ddc
         }
 
         [JsonConstructor]
-        public RefMetadataResponse(NamespaceId ns, BucketId bucket, IoHash name, BlobIdentifier payloadIdentifier, DateTime lastAccess, bool isFinalized, byte[]? inlinePayload)
+        public RefMetadataResponse(NamespaceId ns, BucketId bucket, RefId name, BlobIdentifier payloadIdentifier, DateTime lastAccess, bool isFinalized, byte[]? inlinePayload)
         {
             Ns = ns;
             Bucket = bucket;
@@ -1261,7 +1261,7 @@ namespace Horde.Server.Ddc
             InlinePayload = inlinePayload;
         }
 
-        public RefMetadataResponse(ObjectRecord objectRecord)
+        public RefMetadataResponse(RefRecord objectRecord)
         {
             Ns = objectRecord.Namespace;
             Bucket = objectRecord.Bucket;
@@ -1279,7 +1279,7 @@ namespace Horde.Server.Ddc
         public BucketId Bucket { get; set; }
 
         [CbField("name")]
-        public IoHash Name { get; set; }
+        public RefId Name { get; set; }
 
         [CbField("payloadIdentifier")]
         public BlobIdentifier PayloadIdentifier { get; set; }
@@ -1312,11 +1312,11 @@ namespace Horde.Server.Ddc
 
     public class ExistCheckMultipleRefsResponse
     {
-        public ExistCheckMultipleRefsResponse(List<(BucketId, IoHash)> missing)
+        public ExistCheckMultipleRefsResponse(List<(BucketId, RefId)> missing)
         {
             Missing = missing.Select(pair =>
             {
-                (BucketId bucketId, IoHash ioHashKey) = pair;
+                (BucketId bucketId, RefId ioHashKey) = pair;
                 return new MissingReference()
                 {
                     Bucket = bucketId,
@@ -1342,7 +1342,7 @@ namespace Horde.Server.Ddc
             public BucketId Bucket { get; set; }
 
             [CbField("key")]
-            public IoHash Key { get; set; }
+            public RefId Key { get; set; }
         }
     }
 
diff --git a/Engine/Source/Programs/Horde/Horde.Server/Ddc/IDdcRefService.cs b/Engine/Source/Programs/Horde/Horde.Server/Ddc/IDdcRefService.cs
--- a/Engine/Source/Programs/Horde/Horde.Server/Ddc/IDdcRefService.cs
+++ b/Engine/Source/Programs/Horde/Horde.Server/Ddc/IDdcRefService.cs
@@ -13,23 +13,23 @@ namespace Horde.Server.Ddc
 {
 	public interface IDdcRefService
     {
-        Task<(ObjectRecord, BlobContents?)> Get(NamespaceId ns, BucketId bucket, IoHash key, string[] fields, bool doLastAccessTracking = true);
-        Task<(ContentId[], BlobIdentifier[])> Put(NamespaceId ns, BucketId bucket, IoHash key, BlobIdentifier blobHash, CbObject payload);
-        Task<(ContentId[], BlobIdentifier[])> IDdcRefService(NamespaceId ns, BucketId bucket, IoHash key, BlobIdentifier blobHash);
+        Task<(RefRecord, BlobContents?)> Get(NamespaceId ns, BucketId bucket, RefId key, string[] fields, bool doLastAccessTracking = true);
+        Task<(ContentId[], BlobIdentifier[])> Put(NamespaceId ns, BucketId bucket, RefId key, BlobIdentifier blobHash, CbObject payload);
+        Task<(ContentId[], BlobIdentifier[])> Finalize(NamespaceId ns, BucketId bucket, RefId key, BlobIdentifier blobHash);
 
         IAsyncEnumerable<NamespaceId> GetNamespaces();
 
-        Task<bool> Delete(NamespaceId ns, BucketId bucket, IoHash key);
+        Task<bool> Delete(NamespaceId ns, BucketId bucket, RefId key);
         Task<long> DropNamespace(NamespaceId ns);
         Task<long> DeleteBucket(NamespaceId ns, BucketId bucket);
 
-        Task<bool> Exists(NamespaceId ns, BucketId bucket, IoHash key);
-        Task<List<BlobIdentifier>> GetReferencedBlobs(NamespaceId ns, BucketId bucket, IoHash key);
+        Task<bool> Exists(NamespaceId ns, BucketId bucket, RefId key);
+        Task<List<BlobIdentifier>> GetReferencedBlobs(NamespaceId ns, BucketId bucket, RefId key);
     }
 
-	public class ObjectRecord
+	public class RefRecord
 	{
-		public ObjectRecord(NamespaceId ns, BucketId bucket, IoHash name, DateTime lastAccess, byte[]? inlinePayload, BlobIdentifier blobIdentifier, bool isFinalized)
+		public RefRecord(NamespaceId ns, BucketId bucket, RefId name, DateTime lastAccess, byte[]? inlinePayload, BlobIdentifier blobIdentifier, bool isFinalized)
 		{
 			Namespace = ns;
 			Bucket = bucket;
@@ -42,16 +42,16 @@ namespace Horde.Server.Ddc
 
 		public NamespaceId Namespace { get; }
 		public BucketId Bucket { get; }
-		public IoHash Name { get; }
+		public RefId Name { get; }
 		public DateTime LastAccess { get; }
 		public byte[]? InlinePayload { get; set; }
 		public BlobIdentifier BlobIdentifier { get; set; }
 		public bool IsFinalized { get; }
 	}
 
-	public class ObjectNotFoundException : Exception
+	public class RefNotFoundException : Exception
 	{
-		public ObjectNotFoundException(NamespaceId ns, BucketId bucket, IoHash key) : base($"Object not found {key} in bucket {bucket} namespace {ns}")
+		public RefNotFoundException(NamespaceId ns, BucketId bucket, RefId key) : base($"Object not found {key} in bucket {bucket} namespace {ns}")
 		{
 			Namespace = ns;
 			Bucket = bucket;
@@ -60,7 +60,7 @@ namespace Horde.Server.Ddc
 
 		public NamespaceId Namespace { get; }
 		public BucketId Bucket { get; }
-		public IoHash Key { get; }
+		public RefId Key { get; }
 	}
 
 	public class ObjectHashMismatchException : Exception
diff --git a/Engine/Source/Programs/Horde/Horde.Server/Ddc/RefId.cs b/Engine/Source/Programs/Horde/Horde.Server/Ddc/RefId.cs
new file mode 100644
--- /dev/null
+++ b/Engine/Source/Programs/Horde/Horde.Server/Ddc/RefId.cs
@@ -0,0 +1,92 @@
+// Copyright Epic Games, Inc. All Rights Reserved.
+
+using System;
+using System.ComponentModel;
+using System.Globalization;
+using System.Text.Json;
+using System.Text.Json.Serialization;
+using EpicGames.Core;
+using EpicGames.Serialization;
+
+namespace Horde.Server.Ddc
+{
+	/// <summary>
+	/// Identifier for a DDC ref. Wrapper around an IoHash.
+	/// </summary>
+	[JsonConverter(typeof(RefIdJsonConverter))]
+	[TypeConverter(typeof(RefIdTypeConverter))]
+	[CbConverter(typeof(RefIdCbConverter))]
+	public readonly struct RefId : IEquatable<RefId>
+	{
+		/// <summary>
+		/// Hash of the ref name
+		/// </summary>
+		public IoHash Hash { get; }
+
+		/// <summary>
+		/// Constructor
+		/// </summary>
+		public RefId(IoHash hash) => Hash = hash;
+
+		/// <summary>
+		/// Parses a ref id from a string
+		/// </summary>
+		public static RefId Parse(string str) => new RefId(IoHash.Parse(str));
+
+		/// <inheritdoc/>
+		public override readonly bool Equals(object? obj) => obj is RefId id && Hash.Equals(id.Hash);
+
+		/// <inheritdoc/>
+		public override readonly int GetHashCode() => Hash.GetHashCode();
+
+		/// <inheritdoc/>
+		public readonly bool Equals(RefId other) => Hash.Equals(other.Hash);
+
+		/// <inheritdoc/>
+		public override readonly string ToString() => Hash.ToString();
+
+		/// <inheritdoc cref="IoHash.op_Equality"/>
+		public static bool operator ==(RefId left, RefId right) => left.Hash == right.Hash;
+
+		/// <inheritdoc cref="IoHash.op_Inequality"/>
+		public static bool operator !=(RefId left, RefId right) => left.Hash != right.Hash;
+	}
+
+	/// <summary>
+	/// Type converter for RefId to and from JSON
+	/// </summary>
+	sealed class RefIdJsonConverter : JsonConverter<RefId>
+	{
+		/// <inheritdoc/>
+		public override RefId Read(ref Utf8JsonReader reader, Type typeToConvert, JsonSerializerOptions options) => RefId.Parse(reader.GetString() ?? String.Empty);
+
+		/// <inheritdoc/>
+		public override void Write(Utf8JsonWriter writer, RefId value, JsonSerializerOptions options) => writer.WriteStringValue(value.ToString());
+	}
+
+	/// <summary>
+	/// Type converter from strings to RefId objects
+	/// </summary>
+	sealed class RefIdTypeConverter : TypeConverter
+	{
+		/// <inheritdoc/>
+		public override bool CanConvertFrom(ITypeDescriptorContext? context, Type? sourceType) => sourceType == typeof(string);
+
+		/// <inheritdoc/>
+		public override object ConvertFrom(ITypeDescriptorContext? context, CultureInfo? culture, object? value) => RefId.Parse((string)value!);
+	}
+
+	/// <summary>
+	/// Type converter to compact binary objects
+	/// </summary>
+	sealed class RefIdCbConverter : CbConverterBase<RefId>
+	{
+		public override RefId Read(CbField field) => RefId.Parse(field.AsString());
+
+		/// <inheritdoc/>
+		public override void Write(CbWriter writer, RefId value) => writer.WriteStringValue(value.ToString());
+
+		/// <inheritdoc/>
+		public override void WriteNamed(CbWriter writer, Utf8String name, RefId value) => writer.WriteString(name, value.ToString());
+	}
+}
