Horde: Rename BlobIdentifier -> BlobId.

diff --git a/Engine/Source/Programs/Horde/Horde.Server/Ddc/BlobId.cs b/Engine/Source/Programs/Horde/Horde.Server/Ddc/BlobId.cs
new file mode 100644
--- /dev/null
+++ b/Engine/Source/Programs/Horde/Horde.Server/Ddc/BlobId.cs
@@ -0,0 +1,235 @@
+﻿// Copyright Epic Games, Inc. All Rights Reserved.
+
+using System;
+using System.ComponentModel;
+using System.Globalization;
+using System.IO;
+using System.Text;
+using System.Text.Json;
+using System.Text.Json.Serialization;
+using System.Threading.Tasks;
+using Blake3;
+using EpicGames.Core;
+using EpicGames.Horde.Storage;
+using EpicGames.Horde.Storage.Bundles;
+using EpicGames.Serialization;
+
+#pragma warning disable CS1591
+
+namespace Horde.Server.Ddc
+{
+	[TypeConverter(typeof(BlobIdTypeConverter))]
+    [JsonConverter(typeof(BlobIdJsonConverter))]
+    [CbConverter(typeof(BlobIdCbConverter))]
+    public class BlobId : JupiterContentHash,  IEquatable<BlobId>
+    {
+        // multi thread the hashing for blobs larger then this size
+        private const int MultiThreadedSize = 1_000_000;
+        private string? _stringIdentifier;
+
+        public BlobId(byte[] identifier) : base(identifier)
+        {
+        }
+
+        [JsonConstructor]
+        public BlobId(string identifier) : base(identifier)
+        {
+
+        }
+
+        public override int GetHashCode()
+        {
+			HashCode hashCode = new HashCode();
+			hashCode.AddBytes(Identifier);
+			return hashCode.ToHashCode();
+        }
+
+		public bool Equals(BlobId? other) => other != null && other.Identifier.AsSpan().SequenceEqual(other.Identifier);
+
+        public override bool Equals(object? obj)
+        {
+            if (ReferenceEquals(null, obj))
+            {
+                return false;
+            }
+
+            if (ReferenceEquals(this, obj))
+            {
+                return true;
+            }
+
+            if (obj.GetType() != GetType())
+            {
+                return false;
+            }
+
+            return Equals((BlobId) obj);
+        }
+
+        public override string ToString()
+        {
+            if (_stringIdentifier == null)
+            {
+                _stringIdentifier = StringUtils.FormatHexString(Identifier);
+            }
+
+            return _stringIdentifier;
+        }
+
+        public static new BlobId FromBlob(byte[] blobMemory)
+        {
+            Hash blake3Hash;
+            if (blobMemory.Length < MultiThreadedSize)
+            {
+                using Hasher hasher = Hasher.New();
+                hasher.Update(blobMemory);
+                blake3Hash = hasher.Finalize();
+            }
+            else
+            {
+                using Hasher hasher = Hasher.New();
+                hasher.UpdateWithJoin(blobMemory);
+                blake3Hash = hasher.Finalize();
+            }
+            
+            // we only keep the first 20 bytes of the Blake3 hash
+            Span<byte> hash = blake3Hash.AsSpanUnsafe().Slice(0, 20);
+            return new BlobId(hash.ToArray());
+        }
+
+        public static BlobId FromBlob(in Memory<byte> blobMemory)
+        {
+            Hash blake3Hash;
+            if (blobMemory.Length < MultiThreadedSize)
+            {
+                using Hasher hasher = Hasher.New();
+                hasher.Update(blobMemory.Span);
+                blake3Hash = hasher.Finalize();
+            }
+            else
+            {
+                using Hasher hasher = Hasher.New();
+                hasher.UpdateWithJoin(blobMemory.Span);
+                blake3Hash = hasher.Finalize();
+            }
+
+            // we only keep the first 20 bytes of the Blake3 hash
+            Span<byte> hash = blake3Hash.AsSpanUnsafe().Slice(0, 20);
+            return new BlobId(hash.ToArray());
+        }
+
+        public static BlobId FromContentHash(JupiterContentHash testObjectHash)
+        {
+            return new BlobId(testObjectHash.HashData);
+        }
+
+        public static async Task<BlobId> FromStream(Stream stream)
+        {
+            using Hasher hasher = Hasher.New();
+            const int bufferSize = 1024 * 1024 * 5;
+            byte[] buffer = new byte[bufferSize];
+            int read = await stream.ReadAsync(buffer, 0, buffer.Length);
+            while (read > 0)
+            {
+                hasher.UpdateWithJoin(new ReadOnlySpan<byte>(buffer, 0, read));
+                read = await stream.ReadAsync(buffer, 0, buffer.Length);
+            }
+            Hash blake3Hash = hasher.Finalize();
+
+            // we only keep the first 20 bytes of the Blake3 hash
+            byte[] hash = blake3Hash.AsSpanUnsafe().Slice(0, 20).ToArray();
+            return new BlobId(hash);
+        }
+
+        public static BlobId FromIoHash(IoHash blobIdentifier)
+        {
+            return new BlobId(blobIdentifier.ToByteArray());
+        }
+
+        public IoHash AsIoHash()
+        {
+            return new IoHash(HashData);
+        }
+
+        public static BlobId FromBlobLocator(BundleLocator locator)
+        {
+            return new BlobId(Encoding.UTF8.GetBytes(locator.ToString()));
+        }
+
+        public BundleLocator AsBlobLocator()
+        {
+            return new BundleLocator(Encoding.UTF8.GetString(HashData));
+        }
+    }
+
+    public class BlobIdTypeConverter : TypeConverter
+    {
+        public override bool CanConvertFrom(ITypeDescriptorContext? context, Type sourceType)
+        {  
+            if (sourceType == typeof(string))  
+            {  
+                return true;
+            }  
+            return base.CanConvertFrom(context, sourceType);
+        }  
+  
+        public override object? ConvertFrom(ITypeDescriptorContext? context, CultureInfo? culture, object value)  
+        {
+            if (value is string s)
+            {
+                return new BlobId(s);
+            }
+
+            return base.ConvertFrom(context, culture, value);  
+        }
+
+        public override bool CanConvertTo(ITypeDescriptorContext? context, Type? destinationType)
+        {
+            if (destinationType == typeof(string))
+            {
+                return true;
+            }
+            return base.CanConvertTo(context, destinationType);
+        }
+
+        public override object? ConvertTo(ITypeDescriptorContext? context, CultureInfo? culture, object? value, Type destinationType)
+        {
+            if (destinationType == typeof(string))
+            {
+                BlobId? identifier = (BlobId?)value;
+                return identifier?.ToString();
+            }
+            return base.ConvertTo(context, culture, value, destinationType);
+        }
+    }
+
+    public class BlobIdJsonConverter : JsonConverter<BlobId>
+    {
+        public override BlobId? Read(ref Utf8JsonReader reader, Type typeToConvert, JsonSerializerOptions options)
+        {
+            string? str = reader.GetString();
+            if (str == null)
+            {
+                throw new InvalidDataException("Unable to parse blob identifier");
+            }
+
+            return new BlobId(str);
+        }
+
+        public override void Write(Utf8JsonWriter writer, BlobId value, JsonSerializerOptions options)
+        {
+            writer.WriteStringValue(value.ToString());
+        }
+    }
+
+    public class BlobIdCbConverter : CbConverterBase<BlobId>
+    {
+        public override BlobId Read(CbField field) => new BlobId(field.AsHash().ToByteArray());
+
+        /// <inheritdoc/>
+        public override void Write(CbWriter writer, BlobId value) => writer.WriteBinaryAttachmentValue(new IoHash(value.HashData));
+
+        /// <inheritdoc/>
+        public override void WriteNamed(CbWriter writer, Utf8String name, BlobId value) => writer.WriteBinaryAttachment(name, new IoHash(value.HashData));
+    }
+}
diff --git a/Engine/Source/Programs/Horde/Horde.Server/Ddc/BlobIdentifier.cs b/Engine/Source/Programs/Horde/Horde.Server/Ddc/BlobIdentifier.cs
deleted file mode 100644
diff --git a/Engine/Source/Programs/Horde/Horde.Server/Ddc/CompressedBufferUtils.cs b/Engine/Source/Programs/Horde/Horde.Server/Ddc/CompressedBufferUtils.cs
--- a/Engine/Source/Programs/Horde/Horde.Server/Ddc/CompressedBufferUtils.cs
+++ b/Engine/Source/Programs/Horde/Horde.Server/Ddc/CompressedBufferUtils.cs
@@ -235,8 +235,8 @@ namespace Horde.Server.Ddc
                 byte[] slicedHash = new byte[20];
                 Array.Copy(header.RawHash, 0, slicedHash, 0, 20);
 
-                BlobIdentifier headerIdentifier = new BlobIdentifier(slicedHash);
-                BlobIdentifier contentHash = BlobIdentifier.FromBlob(decompressedPayload);
+                BlobId headerIdentifier = new BlobId(slicedHash);
+                BlobId contentHash = BlobId.FromBlob(decompressedPayload);
 
                 if (!headerIdentifier.Equals(contentHash))
                 {
diff --git a/Engine/Source/Programs/Horde/Horde.Server/Ddc/DdcBlobsController.cs b/Engine/Source/Programs/Horde/Horde.Server/Ddc/DdcBlobsController.cs
--- a/Engine/Source/Programs/Horde/Horde.Server/Ddc/DdcBlobsController.cs
+++ b/Engine/Source/Programs/Horde/Horde.Server/Ddc/DdcBlobsController.cs
@@ -49,7 +49,7 @@ namespace Horde.Server.Ddc
         [ProducesDefaultResponseType]
         public async Task<IActionResult> Get(
             [Required] NamespaceId ns,
-            [Required] BlobIdentifier id,
+            [Required] BlobId id,
             [FromQuery] List<string>? storageLayers = null)
         {
             ActionResult? result = await _requestHelper.HasAccessToNamespace(User, Request, ns, new [] { DdcAclAction.ReadObject });
@@ -86,7 +86,7 @@ namespace Horde.Server.Ddc
         [ProducesDefaultResponseType]
         public async Task<IActionResult> Head(
             [Required] NamespaceId ns,
-            [Required] BlobIdentifier id,
+            [Required] BlobId id,
             [FromQuery] List<string>? storageLayers = null)
         {
             ActionResult? result = await _requestHelper.HasAccessToNamespace(User, Request, ns, new [] { DdcAclAction.ReadObject });
@@ -108,7 +108,7 @@ namespace Horde.Server.Ddc
         [ProducesDefaultResponseType]
         public async Task<IActionResult> ExistsMultiple(
             [Required] NamespaceId ns,
-            [Required] [FromQuery] List<BlobIdentifier> id)
+            [Required] [FromQuery] List<BlobId> id)
         {
             ActionResult? result = await _requestHelper.HasAccessToNamespace(User, Request, ns, new [] { DdcAclAction.ReadObject });
             if (result != null)
@@ -116,7 +116,7 @@ namespace Horde.Server.Ddc
                 return result;
             }
 
-            ConcurrentBag<BlobIdentifier> missingBlobs = new ConcurrentBag<BlobIdentifier>();
+            ConcurrentBag<BlobId> missingBlobs = new ConcurrentBag<BlobId>();
 
             IEnumerable<Task> tasks = id.Select(async blob =>
             {
@@ -134,7 +134,7 @@ namespace Horde.Server.Ddc
         [ProducesDefaultResponseType]
         public async Task<IActionResult> ExistsBody(
             [Required] NamespaceId ns,
-            [FromBody] BlobIdentifier[] bodyIds)
+            [FromBody] BlobId[] bodyIds)
         {
             ActionResult? result = await _requestHelper.HasAccessToNamespace(User, Request, ns, new [] { DdcAclAction.ReadObject });
             if (result != null)
@@ -142,7 +142,7 @@ namespace Horde.Server.Ddc
                 return result;
             }
 
-            ConcurrentBag<BlobIdentifier> missingBlobs = new ConcurrentBag<BlobIdentifier>();
+            ConcurrentBag<BlobId> missingBlobs = new ConcurrentBag<BlobId>();
 
             IEnumerable<Task> tasks = bodyIds.Select(async blob =>
             {
@@ -156,7 +156,7 @@ namespace Horde.Server.Ddc
             return Ok(new HeadMultipleResponse { Needs = missingBlobs.ToArray()});
         }
 
-        private async Task<BlobContents> GetImpl(NamespaceId ns, BlobIdentifier blob, List<string>? storageLayers = null, bool supportsRedirectUri = false)
+        private async Task<BlobContents> GetImpl(NamespaceId ns, BlobId blob, List<string>? storageLayers = null, bool supportsRedirectUri = false)
         {
             try
             {
@@ -178,7 +178,7 @@ namespace Horde.Server.Ddc
         [DisableRequestSizeLimit]
         public async Task<IActionResult> Put(
             [Required] NamespaceId ns,
-            [Required] BlobIdentifier id)
+            [Required] BlobId id)
         {
             ActionResult? result = await _requestHelper.HasAccessToNamespace(User, Request, ns, new [] { DdcAclAction.WriteObject });
             if (result != null)
@@ -201,7 +201,7 @@ namespace Horde.Server.Ddc
                 }
                 using IBufferedPayload payload = await _bufferedPayloadFactory.CreateFromRequest(Request);
 
-                BlobIdentifier identifier = await _storage.PutObject(ns, payload, id);
+                BlobId identifier = await _storage.PutObject(ns, payload, id);
                 return Ok(new
                 {
                     Identifier = identifier.ToString()
@@ -236,7 +236,7 @@ namespace Horde.Server.Ddc
 
                 await using Stream stream = payload.GetStream();
 
-                BlobIdentifier id = await BlobIdentifier.FromStream(stream);
+                BlobId id = await BlobId.FromStream(stream);
                 await _storage.PutObjectKnownHash(ns, payload, id);
                 
                 return Ok(new
@@ -253,7 +253,7 @@ namespace Horde.Server.Ddc
         [HttpDelete("{ns}/{id}")]
         public async Task<IActionResult> Delete(
             [Required] NamespaceId ns,
-            [Required] BlobIdentifier id)
+            [Required] BlobId id)
         {
             ActionResult? result = await _requestHelper.HasAccessToNamespace(User, Request, ns, new [] { DdcAclAction.DeleteObject });
             if (result != null)
@@ -282,7 +282,7 @@ namespace Horde.Server.Ddc
             return NoContent();
         }
 
-        private async Task DeleteImpl(NamespaceId ns, BlobIdentifier id)
+        private async Task DeleteImpl(NamespaceId ns, BlobId id)
         {
             await _storage.DeleteObject(ns, id);
         }
@@ -304,7 +304,7 @@ namespace Horde.Server.Ddc
 
             [Required] public NamespaceId? Namespace { get; set; }
 
-            public BlobIdentifier? Id { get; set; }
+            public BlobId? Id { get; set; }
 
             [Required] public Operation Op { get; set; }
 
@@ -429,6 +429,6 @@ namespace Horde.Server.Ddc
     public class HeadMultipleResponse
     {
         [CbField("needs")]
-        public BlobIdentifier[] Needs { get; set; } = null!;
+        public BlobId[] Needs { get; set; } = null!;
     }
 }
diff --git a/Engine/Source/Programs/Horde/Horde.Server/Ddc/DdcContentIdController.cs b/Engine/Source/Programs/Horde/Horde.Server/Ddc/DdcContentIdController.cs
--- a/Engine/Source/Programs/Horde/Horde.Server/Ddc/DdcContentIdController.cs
+++ b/Engine/Source/Programs/Horde/Horde.Server/Ddc/DdcContentIdController.cs
@@ -42,7 +42,7 @@ namespace Horde.Server.Ddc
             {
                 return result;
             }
-            BlobIdentifier[]? blobs = await _contentIdStore.Resolve(ns, contentId, mustBeContentId: true);
+            BlobId[]? blobs = await _contentIdStore.Resolve(ns, contentId, mustBeContentId: true);
 
             if (blobs == null)
             {
@@ -62,7 +62,7 @@ namespace Horde.Server.Ddc
         [HttpPut("{ns}/{contentId}/update/{blobIdentifier}/{contentWeight}", Order = 500)]
         [ProducesDefaultResponseType]
         [ProducesResponseType(type: typeof(ProblemDetails), 400)]
-        public async Task<IActionResult> UpdateContentIdMapping(NamespaceId ns, ContentId contentId, BlobIdentifier blobIdentifier, int contentWeight)
+        public async Task<IActionResult> UpdateContentIdMapping(NamespaceId ns, ContentId contentId, BlobId blobIdentifier, int contentWeight)
         {
             ActionResult? result = await _requestHelper.HasAccessToNamespace(User, Request, ns, new [] { DdcAclAction.WriteObject });
             if (result != null)
@@ -77,9 +77,9 @@ namespace Horde.Server.Ddc
 
     public class ResolvedContentIdResponse
     {
-        public BlobIdentifier[] Blobs { get; set; }
+        public BlobId[] Blobs { get; set; }
 
-        public ResolvedContentIdResponse(BlobIdentifier[] blobs)
+        public ResolvedContentIdResponse(BlobId[] blobs)
         {
             Blobs = blobs;
         }
diff --git a/Engine/Source/Programs/Horde/Horde.Server/Ddc/DdcContentIdService.cs b/Engine/Source/Programs/Horde/Horde.Server/Ddc/DdcContentIdService.cs
--- a/Engine/Source/Programs/Horde/Horde.Server/Ddc/DdcContentIdService.cs
+++ b/Engine/Source/Programs/Horde/Horde.Server/Ddc/DdcContentIdService.cs
@@ -8,12 +8,12 @@ namespace Horde.Server.Ddc
 {
 	class ContentIdStore : IDdcContentIdService
 	{
-		public Task<BlobIdentifier[]?> Resolve(NamespaceId ns, ContentId contentId, bool mustBeContentId = false)
+		public Task<BlobId[]?> Resolve(NamespaceId ns, ContentId contentId, bool mustBeContentId = false)
 		{
 			throw new NotImplementedException();
 		}
 
-		public Task Put(NamespaceId ns, ContentId contentId, BlobIdentifier blobIdentifier, int contentWeight)
+		public Task Put(NamespaceId ns, ContentId contentId, BlobId blobIdentifier, int contentWeight)
 		{
 			throw new NotImplementedException();
 		}
diff --git a/Engine/Source/Programs/Horde/Horde.Server/Ddc/DdcObjectsController.cs b/Engine/Source/Programs/Horde/Horde.Server/Ddc/DdcObjectsController.cs
--- a/Engine/Source/Programs/Horde/Horde.Server/Ddc/DdcObjectsController.cs
+++ b/Engine/Source/Programs/Horde/Horde.Server/Ddc/DdcObjectsController.cs
@@ -49,7 +49,7 @@ namespace Horde.Server.Ddc
         [ProducesDefaultResponseType]
         public async Task<IActionResult> Get(
             [Required] NamespaceId ns,
-            [Required] BlobIdentifier id)
+            [Required] BlobId id)
         {
             ActionResult? result = await _requestHelper.HasAccessToNamespace(User, Request, ns, new [] { DdcAclAction.ReadObject });
             if (result != null)
@@ -73,7 +73,7 @@ namespace Horde.Server.Ddc
         [ProducesDefaultResponseType]
         public async Task<IActionResult> Head(
             [Required] NamespaceId ns,
-            [Required] BlobIdentifier id)
+            [Required] BlobId id)
         {
             ActionResult? result = await _requestHelper.HasAccessToNamespace(User, Request, ns, new [] { DdcAclAction.ReadObject });
             if (result != null)
@@ -95,7 +95,7 @@ namespace Horde.Server.Ddc
         [ProducesDefaultResponseType]
         public async Task<IActionResult> ExistsMultiple(
             [Required] NamespaceId ns,
-            [Required] [FromQuery] List<BlobIdentifier> id)
+            [Required] [FromQuery] List<BlobId> id)
         {
             ActionResult? result = await _requestHelper.HasAccessToNamespace(User, Request, ns, new [] { DdcAclAction.ReadObject });
             if (result != null)
@@ -103,7 +103,7 @@ namespace Horde.Server.Ddc
                 return result;
             }
 
-            ConcurrentBag<BlobIdentifier> missingBlobs = new ConcurrentBag<BlobIdentifier>();
+            ConcurrentBag<BlobId> missingBlobs = new ConcurrentBag<BlobId>();
 
             IEnumerable<Task> tasks = id.Select(async blob =>
             {
@@ -121,7 +121,7 @@ namespace Horde.Server.Ddc
         [ProducesDefaultResponseType]
         public async Task<IActionResult> ExistsBody(
             [Required] NamespaceId ns,
-            [FromBody] BlobIdentifier[] bodyIds)
+            [FromBody] BlobId[] bodyIds)
         {
             ActionResult? result = await _requestHelper.HasAccessToNamespace(User, Request, ns, new [] { DdcAclAction.ReadObject });
             if (result != null)
@@ -129,7 +129,7 @@ namespace Horde.Server.Ddc
                 return result;
             }
 
-            ConcurrentBag<BlobIdentifier> missingBlobs = new ConcurrentBag<BlobIdentifier>();
+            ConcurrentBag<BlobId> missingBlobs = new ConcurrentBag<BlobId>();
 
             IEnumerable<Task> tasks = bodyIds.Select(async blob =>
             {
@@ -147,7 +147,7 @@ namespace Horde.Server.Ddc
         [RequiredContentType(CustomMediaTypeNames.UnrealCompactBinary)]
         public async Task<IActionResult> Put(
             [Required] NamespaceId ns,
-            [Required] BlobIdentifier id)
+            [Required] BlobId id)
         {
             ActionResult? result = await _requestHelper.HasAccessToNamespace(User, Request, ns, new [] { DdcAclAction.WriteObject });
             if (result != null)
@@ -160,7 +160,7 @@ namespace Horde.Server.Ddc
             {
                 using IBufferedPayload payload = await _bufferedPayloadFactory.CreateFromRequest(Request);
 
-                BlobIdentifier identifier = await _storage.PutObject(ns, payload, id);
+                BlobId identifier = await _storage.PutObject(ns, payload, id);
                 return Ok(new PutBlobResponse(identifier));
             }
             catch (ClientSendSlowException e)
@@ -172,7 +172,7 @@ namespace Horde.Server.Ddc
         [HttpGet("{ns}/{id}/references")]
         public async Task<IActionResult> ResolveReferences(
             [Required] NamespaceId ns,
-            [Required] BlobIdentifier id)
+            [Required] BlobId id)
         {
             ActionResult? result = await _requestHelper.HasAccessToNamespace(User, Request, ns, new [] { DdcAclAction.ReadObject });
             if (result != null)
@@ -208,7 +208,7 @@ namespace Horde.Server.Ddc
 
             try
             {
-                BlobIdentifier[] references = await _referenceResolver.GetReferencedBlobs(ns, compactBinaryObject).ToArrayAsync();
+                BlobId[] references = await _referenceResolver.GetReferencedBlobs(ns, compactBinaryObject).ToArrayAsync();
                 return Ok(new ResolvedReferencesResult(references));
             }
             catch (PartialReferenceResolveException e)
@@ -224,7 +224,7 @@ namespace Horde.Server.Ddc
         [HttpDelete("{ns}/{id}")]
         public async Task<IActionResult> Delete(
             [Required] NamespaceId ns,
-            [Required] BlobIdentifier id)
+            [Required] BlobId id)
         {
             ActionResult? result = await _requestHelper.HasAccessToNamespace(User, Request, ns, new [] { DdcAclAction.DeleteObject });
             if (result != null)
@@ -263,13 +263,13 @@ namespace Horde.Server.Ddc
             Identifier = null!;
         }
 
-        public PutBlobResponse(BlobIdentifier identifier)
+        public PutBlobResponse(BlobId identifier)
         {
             Identifier = identifier;
         }
 
         [CbField("identifier")]
-        public BlobIdentifier Identifier { get; set; }
+        public BlobId Identifier { get; set; }
     }
 
     public class DeletedResponse
@@ -284,12 +284,12 @@ namespace Horde.Server.Ddc
             References = null!;
         }
 
-        public ResolvedReferencesResult(BlobIdentifier[] references)
+        public ResolvedReferencesResult(BlobId[] references)
         {
             References = references;
         }
 
         [CbField("references")]
-        public BlobIdentifier[] References { get; set; }
+        public BlobId[] References { get; set; }
     }
 }
diff --git a/Engine/Source/Programs/Horde/Horde.Server/Ddc/DdcRefsController.cs b/Engine/Source/Programs/Horde/Horde.Server/Ddc/DdcRefsController.cs
--- a/Engine/Source/Programs/Horde/Horde.Server/Ddc/DdcRefsController.cs
+++ b/Engine/Source/Programs/Horde/Horde.Server/Ddc/DdcRefsController.cs
@@ -198,7 +198,7 @@ namespace Horde.Server.Ddc
 
                             IoHash hash = binaryAttachmentField.AsBinaryAttachment();
 
-                            BlobContents referencedBlobContents = await _blobStore.GetObject(ns, BlobIdentifier.FromIoHash(hash));
+                            BlobContents referencedBlobContents = await _blobStore.GetObject(ns, BlobId.FromIoHash(hash));
 
                             if (_nginxRedirectHelper.CanRedirect(Request, referencedBlobContents))
                             {
@@ -248,13 +248,13 @@ namespace Horde.Server.Ddc
                                 BlobContents attachmentContents;
                                 if (attachment is BlobAttachment blobAttachment)
                                 {
-                                    BlobIdentifier referencedBlob = blobAttachment.Identifier;
+                                    BlobId referencedBlob = blobAttachment.Identifier;
                                     attachmentContents = await _blobStore.GetObject(ns, referencedBlob);
                                 }
                                 else if (attachment is ObjectAttachment objectAttachment)
                                 {
                                     flags |= CbPackageAttachmentFlags.IsObject;
-                                    BlobIdentifier referencedBlob = objectAttachment.Identifier;
+                                    BlobId referencedBlob = objectAttachment.Identifier;
                                     attachmentContents = await _blobStore.GetObject(ns, referencedBlob);
                                 }
                                 else if (attachment is ContentIdAttachment contentIdAttachment)
@@ -330,10 +330,10 @@ namespace Horde.Server.Ddc
                         if (countOfBinaryAttachmentFields == 1 && countOfAttachmentFields == 1)
                         {
                             // fetch the blob so we can resolve any content ids in it
-                            List<BlobIdentifier> referencedBlobs;
+                            List<BlobId> referencedBlobs;
                             try
                             {
-                                IAsyncEnumerable<BlobIdentifier> referencedBlobsEnumerable = _referenceResolver.GetReferencedBlobs(ns, cb);
+                                IAsyncEnumerable<BlobId> referencedBlobsEnumerable = _referenceResolver.GetReferencedBlobs(ns, cb);
                                 referencedBlobs = await referencedBlobsEnumerable.ToListAsync();
                             }
                             catch (PartialReferenceResolveException)
@@ -347,7 +347,7 @@ namespace Horde.Server.Ddc
 
                             if (referencedBlobs.Count == 1)
                             {
-                                BlobIdentifier attachmentToSend = referencedBlobs.First();
+                                BlobId attachmentToSend = referencedBlobs.First();
                                 try
                                 {
                                     BlobContents referencedBlobContents = await _blobStore.GetObject(ns, attachmentToSend);
@@ -501,12 +501,12 @@ namespace Horde.Server.Ddc
                 byte[] blobContents = await blob.Stream.ToByteArray();
                 CbObject compactBinaryObject = new CbObject(blobContents);
                 // the reference resolver will throw if any blob is missing, so no need to do anything other then process each reference
-                IAsyncEnumerable<BlobIdentifier> references = _referenceResolver.GetReferencedBlobs(ns, compactBinaryObject);
-                List<BlobIdentifier>? _ = await references.ToListAsync();
+                IAsyncEnumerable<BlobId> references = _referenceResolver.GetReferencedBlobs(ns, compactBinaryObject);
+                List<BlobId>? _ = await references.ToListAsync();
 
                 // we have to verify the blobs are available locally, as the record of the key is replicated a head of the content
                 // TODO: Once we support inline replication this step is not needed as at least one region as this blob, just maybe not this current one
-                BlobIdentifier[] unknownBlobs = await _blobStore.FilterOutKnownBlobs(ns, new BlobIdentifier[] { record.BlobIdentifier });
+                BlobId[] unknownBlobs = await _blobStore.FilterOutKnownBlobs(ns, new BlobId[] { record.BlobIdentifier });
                 if (unknownBlobs.Length != 0)
                 {
                     return NotFound(new ProblemDetails {Title = $"Object {bucket} {key} in namespace {ns} had at least one missing blob."});
@@ -579,8 +579,8 @@ namespace Horde.Server.Ddc
                     byte[] blobContents = await blob.Stream.ToByteArray();
                     CbObject cb = new CbObject(blobContents);
                     // the reference resolver will throw if any blob is missing, so no need to do anything other then process each reference
-                    IAsyncEnumerable<BlobIdentifier> references = _referenceResolver.GetReferencedBlobs(ns, cb);
-                    List<BlobIdentifier>? _ = await references.ToListAsync();
+                    IAsyncEnumerable<BlobId> references = _referenceResolver.GetReferencedBlobs(ns, cb);
+                    List<BlobId>? _ = await references.ToListAsync();
                 }
                 catch (RefNotFoundException)
                 {
@@ -620,16 +620,16 @@ namespace Horde.Server.Ddc
             _diagnosticContext.Set("Content-Length", Request.ContentLength ?? -1);
 
             CbObject payloadObject;
-            BlobIdentifier blobHeader;
+            BlobId blobHeader;
 
             try
             {
                 using IBufferedPayload payload = await _bufferedPayloadFactory.CreateFromRequest(Request);
 
-                BlobIdentifier headerHash;
+                BlobId headerHash;
                 if (Request.Headers.ContainsKey(CommonHeaders.HashHeaderName))
                 {
-                    headerHash = new BlobIdentifier(Request.Headers[CommonHeaders.HashHeaderName]);
+                    headerHash = new BlobId(Request.Headers[CommonHeaders.HashHeaderName]);
                 }
                 else
                 {
@@ -657,7 +657,7 @@ namespace Horde.Server.Ddc
 
                         byte[] blob = writer.ToByteArray();
                         payloadObject = new CbObject(blob);
-                        blobHeader = BlobIdentifier.FromBlob(blob);
+                        blobHeader = BlobId.FromBlob(blob);
                         break;
                     }
                     case CustomMediaTypeNames.UnrealCompactBinary:
@@ -680,7 +680,7 @@ namespace Horde.Server.Ddc
 
                         byte[] blob = writer.ToByteArray();
                         payloadObject = new CbObject(blob);
-                        blobHeader = BlobIdentifier.FromBlob(blob);
+                        blobHeader = BlobId.FromBlob(blob);
                         break;
                     }
                     default:
@@ -699,7 +699,7 @@ namespace Horde.Server.Ddc
                 return Problem(e.Message, null, (int)HttpStatusCode.RequestTimeout);
             }
 
-            (ContentId[] missingReferences, BlobIdentifier[] missingBlobs) = await _refService.Put(ns, bucket, key, blobHeader, payloadObject);
+            (ContentId[] missingReferences, BlobId[] missingBlobs) = await _refService.Put(ns, bucket, key, blobHeader, payloadObject);
 
             List<JupiterContentHash> missingHashes = new List<JupiterContentHash>(missingReferences);
             missingHashes.AddRange(missingBlobs);
@@ -743,7 +743,7 @@ namespace Horde.Server.Ddc
                     }
                     else
                     {
-                        await _blobStore.PutObject(ns, blob, BlobIdentifier.FromIoHash(entry.AttachmentHash));
+                        await _blobStore.PutObject(ns, blob, BlobId.FromIoHash(entry.AttachmentHash));
                     }
                 }
             }
@@ -756,9 +756,9 @@ namespace Horde.Server.Ddc
             }
             
             CbObject rootObject = packageReader.RootObject;
-            BlobIdentifier rootObjectHash = BlobIdentifier.FromIoHash(packageReader.RootHash);
+            BlobId rootObjectHash = BlobId.FromIoHash(packageReader.RootHash);
 
-            (ContentId[] missingReferences, BlobIdentifier[] missingBlobs) = await _refService.Put(ns, bucket, key, rootObjectHash, rootObject);
+            (ContentId[] missingReferences, BlobId[] missingBlobs) = await _refService.Put(ns, bucket, key, rootObjectHash, rootObject);
 
             List<JupiterContentHash> missingHashes = new List<JupiterContentHash>(missingReferences);
             missingHashes.AddRange(missingBlobs);
@@ -796,7 +796,7 @@ namespace Horde.Server.Ddc
             [FromRoute] [Required] NamespaceId ns,
             [FromRoute] [Required] BucketId bucket,
             [FromRoute] [Required] RefId key,
-            [FromRoute] [Required] BlobIdentifier hash)
+            [FromRoute] [Required] BlobId hash)
         {
             ActionResult? accessResult = await _requestHelper.HasAccessToNamespace(User, Request, ns, new [] { DdcAclAction.WriteObject });
             if (accessResult != null)
@@ -806,7 +806,7 @@ namespace Horde.Server.Ddc
 
             try
             {
-                (ContentId[] missingReferences, BlobIdentifier[] missingBlobs) = await _refService.Finalize(ns, bucket, key, hash);
+                (ContentId[] missingReferences, BlobId[] missingBlobs) = await _refService.Finalize(ns, bucket, key, hash);
                 List<JupiterContentHash> missingHashes = new List<JupiterContentHash>(missingReferences);
                 missingHashes.AddRange(missingBlobs);
 
@@ -883,8 +883,8 @@ namespace Horde.Server.Ddc
 
                     if (op.ResolveAttachments ?? false)
                     {
-                        IAsyncEnumerable<BlobIdentifier> references = _referenceResolver.GetReferencedBlobs(ns, cb);
-                        List<BlobIdentifier>? _ = await references.ToListAsync();
+                        IAsyncEnumerable<BlobId> references = _referenceResolver.GetReferencedBlobs(ns, cb);
+                        List<BlobId>? _ = await references.ToListAsync();
                     }
 
                     return (cb, HttpStatusCode.OK);
@@ -917,8 +917,8 @@ namespace Horde.Server.Ddc
                         byte[] blobContents = await blob.Stream.ToByteArray();
                         CbObject cb = new CbObject(blobContents);
                         // the reference resolver will throw if any blob is missing, so no need to do anything other then process each reference
-                        IAsyncEnumerable<BlobIdentifier> references = _referenceResolver.GetReferencedBlobs(ns, cb);
-                        List<BlobIdentifier>? _ = await references.ToListAsync();
+                        IAsyncEnumerable<BlobId> references = _referenceResolver.GetReferencedBlobs(ns, cb);
+                        List<BlobId>? _ = await references.ToListAsync();
                     }
 
                     if (blob == null)
@@ -951,15 +951,15 @@ namespace Horde.Server.Ddc
                     {
                         throw new Exception($"Missing payload hash for operation: {op.OpId}");
                     }
-                    BlobIdentifier headerHash = BlobIdentifier.FromContentHash(op.PayloadHash);
-                    BlobIdentifier objectHash = BlobIdentifier.FromBlob(op.Payload.GetView().ToArray());
+                    BlobId headerHash = BlobId.FromContentHash(op.PayloadHash);
+                    BlobId objectHash = BlobId.FromBlob(op.Payload.GetView().ToArray());
 
                     if (!headerHash.Equals(objectHash))
                     {
                         throw new HashMismatchException(headerHash, objectHash);
                     }
 
-                    (ContentId[] missingReferences, BlobIdentifier[] missingBlobs) = await _refService.Put(ns, op.Bucket, op.Key, objectHash, op.Payload);
+                    (ContentId[] missingReferences, BlobId[] missingBlobs) = await _refService.Put(ns, op.Bucket, op.Key, objectHash, op.Payload);
                     List<JupiterContentHash> missingHashes = new List<JupiterContentHash>(missingReferences);
 
                     return (CbSerializer.Serialize(new PutObjectResponse(missingHashes.ToArray())), HttpStatusCode.OK);
@@ -1250,7 +1250,7 @@ namespace Horde.Server.Ddc
         }
 
         [JsonConstructor]
-        public RefMetadataResponse(NamespaceId ns, BucketId bucket, RefId name, BlobIdentifier payloadIdentifier, DateTime lastAccess, bool isFinalized, byte[]? inlinePayload)
+        public RefMetadataResponse(NamespaceId ns, BucketId bucket, RefId name, BlobId payloadIdentifier, DateTime lastAccess, bool isFinalized, byte[]? inlinePayload)
         {
             Ns = ns;
             Bucket = bucket;
@@ -1282,7 +1282,7 @@ namespace Horde.Server.Ddc
         public RefId Name { get; set; }
 
         [CbField("payloadIdentifier")]
-        public BlobIdentifier PayloadIdentifier { get; set; }
+        public BlobId PayloadIdentifier { get; set; }
 
         [CbField("lastAccess")]
         public DateTime LastAccess { get; set; }
diff --git a/Engine/Source/Programs/Horde/Horde.Server/Ddc/IDdcBlobService.cs b/Engine/Source/Programs/Horde/Horde.Server/Ddc/IDdcBlobService.cs
--- a/Engine/Source/Programs/Horde/Horde.Server/Ddc/IDdcBlobService.cs
+++ b/Engine/Source/Programs/Horde/Horde.Server/Ddc/IDdcBlobService.cs
@@ -17,18 +17,18 @@ namespace Horde.Server.Ddc
 	public interface IDdcBlobService
 	{
 		Task<JupiterContentHash> VerifyContentMatchesHash(Stream content, JupiterContentHash identifier);
-		Task<BlobIdentifier> PutObjectKnownHash(NamespaceId ns, IBufferedPayload content, BlobIdentifier identifier);
-		Task<BlobIdentifier> PutObject(NamespaceId ns, IBufferedPayload payload, BlobIdentifier identifier);
-		Task<BlobIdentifier> PutObject(NamespaceId ns, byte[] payload, BlobIdentifier identifier);
-		Task<Uri?> MaybePutObjectWithRedirect(NamespaceId ns, BlobIdentifier identifier);
+		Task<BlobId> PutObjectKnownHash(NamespaceId ns, IBufferedPayload content, BlobId identifier);
+		Task<BlobId> PutObject(NamespaceId ns, IBufferedPayload payload, BlobId identifier);
+		Task<BlobId> PutObject(NamespaceId ns, byte[] payload, BlobId identifier);
+		Task<Uri?> MaybePutObjectWithRedirect(NamespaceId ns, BlobId identifier);
 
-		Task<BlobContents> GetObject(NamespaceId ns, BlobIdentifier blob, List<string>? storageLayers = null, bool supportsRedirectUri = false);
+		Task<BlobContents> GetObject(NamespaceId ns, BlobId blob, List<string>? storageLayers = null, bool supportsRedirectUri = false);
 
-		Task<Uri?> GetObjectWithRedirect(NamespaceId ns, BlobIdentifier blobIdentifier, List<string>? storageLayers = null);
+		Task<Uri?> GetObjectWithRedirect(NamespaceId ns, BlobId blobIdentifier, List<string>? storageLayers = null);
 
-		Task<BlobContents> ReplicateObject(NamespaceId ns, BlobIdentifier blob, bool force = false);
+		Task<BlobContents> ReplicateObject(NamespaceId ns, BlobId blob, bool force = false);
 
-		Task<bool> Exists(NamespaceId ns, BlobIdentifier blob, List<string>? storageLayers = null);
+		Task<bool> Exists(NamespaceId ns, BlobId blob, List<string>? storageLayers = null);
 
 		/// <summary>
 		/// Checks that the blob exists in the root store, the store which is last in the list and thus is intended to have every blob in it
@@ -36,18 +36,18 @@ namespace Horde.Server.Ddc
 		/// <param name="ns">The namespace</param>
 		/// <param name="blob">The identifier of the blob</param>
 		/// <returns></returns>
-		Task<bool> ExistsInRootStore(NamespaceId ns, BlobIdentifier blob);
+		Task<bool> ExistsInRootStore(NamespaceId ns, BlobId blob);
 
 		// Delete a object
-		Task DeleteObject(NamespaceId ns, BlobIdentifier blob);
+		Task DeleteObject(NamespaceId ns, BlobId blob);
 
 		// delete the whole namespace
 		Task DeleteNamespace(NamespaceId ns);
 
-		IAsyncEnumerable<(BlobIdentifier, DateTime)> ListObjects(NamespaceId ns);
-		Task<BlobIdentifier[]> FilterOutKnownBlobs(NamespaceId ns, IEnumerable<BlobIdentifier> blobs);
-		Task<BlobIdentifier[]> FilterOutKnownBlobs(NamespaceId ns, IAsyncEnumerable<BlobIdentifier> blobs);
-		Task<BlobContents> GetObjects(NamespaceId ns, BlobIdentifier[] refRequestBlobReferences);
+		IAsyncEnumerable<(BlobId, DateTime)> ListObjects(NamespaceId ns);
+		Task<BlobId[]> FilterOutKnownBlobs(NamespaceId ns, IEnumerable<BlobId> blobs);
+		Task<BlobId[]> FilterOutKnownBlobs(NamespaceId ns, IAsyncEnumerable<BlobId> blobs);
+		Task<BlobContents> GetObjects(NamespaceId ns, BlobId[] refRequestBlobReferences);
 
 		bool ShouldFetchBlobOnDemand(NamespaceId ns);
 	}
@@ -55,15 +55,15 @@ namespace Horde.Server.Ddc
 	public class BlobNotFoundException : Exception
 	{
 		public NamespaceId Ns { get; }
-		public BlobIdentifier Blob { get; }
+		public BlobId Blob { get; }
 
-		public BlobNotFoundException(NamespaceId ns, BlobIdentifier blob) : base($"No Blob in Namespace {ns} with id {blob}")
+		public BlobNotFoundException(NamespaceId ns, BlobId blob) : base($"No Blob in Namespace {ns} with id {blob}")
 		{
 			Ns = ns;
 			Blob = blob;
 		}
 
-		public BlobNotFoundException(NamespaceId ns, BlobIdentifier blob, string message) : base(message)
+		public BlobNotFoundException(NamespaceId ns, BlobId blob, string message) : base(message)
 		{
 			Ns = ns;
 			Blob = blob;
@@ -72,16 +72,16 @@ namespace Horde.Server.Ddc
 
 	public class BlobReplicationException : BlobNotFoundException
 	{
-		public BlobReplicationException(NamespaceId ns, BlobIdentifier blob, string message) : base(ns, blob, message)
+		public BlobReplicationException(NamespaceId ns, BlobId blob, string message) : base(ns, blob, message)
 		{
 		}
 	}
 
 	public class BlobTooLargeException : Exception
 	{
-		public BlobIdentifier Blob { get; }
+		public BlobId Blob { get; }
 
-		public BlobTooLargeException(BlobIdentifier blob) : base($"Blob {blob} was to large to cache")
+		public BlobTooLargeException(BlobId blob) : base($"Blob {blob} was to large to cache")
 		{
 			Blob = blob;
 		}
@@ -133,17 +133,17 @@ namespace Horde.Server.Ddc
 				JupiterContentHash blobHash;
 				{
 					using TelemetrySpan _ = tracer.StartActiveSpan("web.hash").SetAttribute("operation.name", "web.hash");
-					blobHash = await BlobIdentifier.FromStream(decompressedStream);
+					blobHash = await BlobId.FromStream(decompressedStream);
 				}
 
 				identifierDecompressedPayload = ContentId.FromContentHash(blobHash);
 			}
 
-			BlobIdentifier identifierCompressedPayload;
+			BlobId identifierCompressedPayload;
 			{
 				using TelemetrySpan _ = tracer.StartActiveSpan("web.hash").SetAttribute("operation.name", "web.hash");
 				await using Stream hashStream = payload.GetStream();
-				identifierCompressedPayload = await BlobIdentifier.FromStream(hashStream);
+				identifierCompressedPayload = await BlobId.FromStream(hashStream);
 			}
 
 			// commit the mapping from the decompressed hash to the compressed hash, we run this in parallel with the blob store submit
@@ -166,7 +166,7 @@ namespace Horde.Server.Ddc
 			IDdcContentIdService contentIdStore = provider.GetService<IDdcContentIdService>()!;
 			Tracer tracer = provider.GetService<Tracer>()!;
 
-			BlobIdentifier[]? chunks = await contentIdStore.Resolve(ns, contentId, mustBeContentId: false);
+			BlobId[]? chunks = await contentIdStore.Resolve(ns, contentId, mustBeContentId: false);
 			if (chunks == null || chunks.Length == 0)
 			{
 				throw new ContentIdResolveException(contentId);
@@ -175,7 +175,7 @@ namespace Horde.Server.Ddc
 			// single chunk, we just return that chunk
 			if (chunks.Length == 1)
 			{
-				BlobIdentifier blobToReturn = chunks[0];
+				BlobId blobToReturn = chunks[0];
 				string mimeType = CustomMediaTypeNames.UnrealCompressedBuffer;
 				if (contentId.Equals(blobToReturn))
 				{
diff --git a/Engine/Source/Programs/Horde/Horde.Server/Ddc/IReferenceResolver.cs b/Engine/Source/Programs/Horde/Horde.Server/Ddc/IReferenceResolver.cs
--- a/Engine/Source/Programs/Horde/Horde.Server/Ddc/IReferenceResolver.cs
+++ b/Engine/Source/Programs/Horde/Horde.Server/Ddc/IReferenceResolver.cs
@@ -20,9 +20,9 @@ namespace Horde.Server.Ddc
 
     public class BlobAttachment : Attachment
     {
-        public BlobIdentifier Identifier { get; }
+        public BlobId Identifier { get; }
 
-        public BlobAttachment(BlobIdentifier blobIdentifier)
+        public BlobAttachment(BlobId blobIdentifier)
         {
             Identifier = blobIdentifier;
         }
@@ -35,9 +35,9 @@ namespace Horde.Server.Ddc
 
     public class ObjectAttachment : Attachment
     {
-        public BlobIdentifier Identifier { get; }
+        public BlobId Identifier { get; }
 
-        public ObjectAttachment(BlobIdentifier blobIdentifier)
+        public ObjectAttachment(BlobId blobIdentifier)
         {
             Identifier = blobIdentifier;
         }
@@ -51,9 +51,9 @@ namespace Horde.Server.Ddc
     public class ContentIdAttachment : Attachment
     {
         public ContentId Identifier { get; }
-        public BlobIdentifier[] ReferencedBlobs { get; }
+        public BlobId[] ReferencedBlobs { get; }
 
-        public ContentIdAttachment(ContentId contentId, BlobIdentifier[] referencedBlobs)
+        public ContentIdAttachment(ContentId contentId, BlobId[] referencedBlobs)
         {
             Identifier = contentId;
             ReferencedBlobs = referencedBlobs;
@@ -75,7 +75,7 @@ namespace Horde.Server.Ddc
         /// <param name="ns">The namespace to check</param>
         /// <param name="cb">The compact binary object to resolve references for</param>
         /// <returns></returns>
-        IAsyncEnumerable<BlobIdentifier> GetReferencedBlobs(NamespaceId ns, CbObject cb);
+        IAsyncEnumerable<BlobId> GetReferencedBlobs(NamespaceId ns, CbObject cb);
 
         /// <summary>
         /// Returns which attachments exist in the cb object or any children
@@ -104,10 +104,10 @@ namespace Horde.Server.Ddc
         {
             Queue<CbObject> objectsToVisit = new Queue<CbObject>();
             objectsToVisit.Enqueue(cb);
-            List<BlobIdentifier> unresolvedBlobReferences = new List<BlobIdentifier>();
+            List<BlobId> unresolvedBlobReferences = new List<BlobId>();
 
             List<Task<CbObject>> pendingCompactBinaryAttachments = new();
-            List<Task< (ContentId, BlobIdentifier[]?)>> pendingContentIdResolves = new();
+            List<Task< (ContentId, BlobId[]?)>> pendingContentIdResolves = new();
 
             while (pendingCompactBinaryAttachments.Count != 0 || pendingContentIdResolves.Count != 0 || objectsToVisit.Count != 0)
             {
@@ -119,12 +119,12 @@ namespace Horde.Server.Ddc
                     {
                         IoHash attachmentHash = field.AsAttachment();
 
-                        BlobIdentifier blobIdentifier = BlobIdentifier.FromIoHash(attachmentHash);
+                        BlobId blobIdentifier = BlobId.FromIoHash(attachmentHash);
                         ContentId contentId = ContentId.FromIoHash(attachmentHash);
 
                         if (field.IsBinaryAttachment())
                         {
-                            Task<(ContentId, BlobIdentifier[]?)> resolveContentId = ResolveContentId(ns, contentId);
+                            Task<(ContentId, BlobId[]?)> resolveContentId = ResolveContentId(ns, contentId);
                             pendingContentIdResolves.Add(resolveContentId);
                         }
                         else if (field.IsObjectAttachment())
@@ -140,14 +140,14 @@ namespace Horde.Server.Ddc
                 }
 
                 // check for any content id resolves to finish
-                List<Task<(ContentId, BlobIdentifier[]?)>> finishedContentIdResolves = new();
-                foreach (Task<(ContentId, BlobIdentifier[]?)> pendingContentIdResolve in pendingContentIdResolves)
+                List<Task<(ContentId, BlobId[]?)>> finishedContentIdResolves = new();
+                foreach (Task<(ContentId, BlobId[]?)> pendingContentIdResolve in pendingContentIdResolves)
                 {
                     if (pendingContentIdResolve.IsCompleted)
                     {
                         ContentId? contentId = null;
-                        BlobIdentifier[]? resolvedBlobs = null;
-                        BlobIdentifier? blobIdentifier = null;
+                        BlobId[]? resolvedBlobs = null;
+                        BlobId? blobIdentifier = null;
                         bool wasContentId = false;
                         try
                         {
@@ -178,7 +178,7 @@ namespace Horde.Server.Ddc
                 }
 
                 // cleanup finished tasks
-                foreach (Task<(ContentId, BlobIdentifier[]?)> finishedTask in finishedContentIdResolves)
+                foreach (Task<(ContentId, BlobId[]?)> finishedTask in finishedContentIdResolves)
                 {
                     pendingContentIdResolves.Remove(finishedTask);
                 }
@@ -234,12 +234,12 @@ namespace Horde.Server.Ddc
             }
         }
 
-        public async IAsyncEnumerable<BlobIdentifier> GetReferencedBlobs(NamespaceId ns, CbObject cb)
+        public async IAsyncEnumerable<BlobId> GetReferencedBlobs(NamespaceId ns, CbObject cb)
         {
-            List<Task<(BlobIdentifier, bool)>> pendingBlobExistsChecks = new();
+            List<Task<(BlobId, bool)>> pendingBlobExistsChecks = new();
             List<Task<(ContentIdAttachment, bool)>> pendingContentIdChecks = new();
             List<ContentId> unresolvedContentIdReferences = new List<ContentId>();
-            List<BlobIdentifier> unresolvedBlobReferences = new List<BlobIdentifier>();
+            List<BlobId> unresolvedBlobReferences = new List<BlobId>();
 
             // Resolve all the attachments
             await foreach (Attachment attachment in GetAttachments(ns, cb))
@@ -265,9 +265,9 @@ namespace Horde.Server.Ddc
             }
 
             // return any verified blobs
-            foreach (Task<(BlobIdentifier, bool)> pendingBlobExistsTask in pendingBlobExistsChecks)
+            foreach (Task<(BlobId, bool)> pendingBlobExistsTask in pendingBlobExistsChecks)
             {
-                (BlobIdentifier blob, bool exists) = await pendingBlobExistsTask;
+                (BlobId blob, bool exists) = await pendingBlobExistsTask;
                 if (exists)
                 {
                     yield return blob;
@@ -288,7 +288,7 @@ namespace Horde.Server.Ddc
                     continue;
                 }
 
-                foreach (BlobIdentifier b in contentIdAttachment.ReferencedBlobs)
+                foreach (BlobId b in contentIdAttachment.ReferencedBlobs)
                 {
                     yield return b;
                 }
@@ -309,9 +309,9 @@ namespace Horde.Server.Ddc
         private async Task<(ContentIdAttachment, bool)> CheckContentIdExists(NamespaceId ns, ContentIdAttachment contentIdAttachment)
         {
             bool allBlobsExist = true;
-            foreach (BlobIdentifier b in contentIdAttachment.ReferencedBlobs)
+            foreach (BlobId b in contentIdAttachment.ReferencedBlobs)
             {
-                (BlobIdentifier _, bool exists) = await CheckBlobExists(ns, b);
+                (BlobId _, bool exists) = await CheckBlobExists(ns, b);
 
                 if (!exists)
                 {
@@ -323,7 +323,7 @@ namespace Horde.Server.Ddc
             return (contentIdAttachment, allBlobsExist);
         }
 
-        private async Task<CbObject> ParseCompactBinaryAttachment(NamespaceId ns, BlobIdentifier blobIdentifier)
+        private async Task<CbObject> ParseCompactBinaryAttachment(NamespaceId ns, BlobId blobIdentifier)
         {
             BlobContents contents = await _blobStore.GetObject(ns, blobIdentifier);
             byte[] data = await contents.Stream.ToByteArray();
@@ -332,16 +332,16 @@ namespace Horde.Server.Ddc
             return childBinaryObject;
         }
 
-        private async Task<(ContentId, BlobIdentifier[]?)> ResolveContentId(NamespaceId ns, ContentId contentId)
+        private async Task<(ContentId, BlobId[]?)> ResolveContentId(NamespaceId ns, ContentId contentId)
         {
             using TelemetrySpan scope = _tracer.StartActiveSpan("ReferenceResolver.ResolveContentId")
                 .SetAttribute("operation.name", "ReferenceResolver.ResolveContentId")
                 .SetAttribute("resource.name", contentId.ToString());
-            BlobIdentifier[]? resolvedBlobs = await _contentIdStore.Resolve(ns, contentId);
+            BlobId[]? resolvedBlobs = await _contentIdStore.Resolve(ns, contentId);
             return (contentId, resolvedBlobs);
         }
 
-        private async Task<(BlobIdentifier, bool)> CheckBlobExists(NamespaceId ns, BlobIdentifier blob)
+        private async Task<(BlobId, bool)> CheckBlobExists(NamespaceId ns, BlobId blob)
         {
             return (blob, await _blobStore.Exists(ns, blob));
         }
@@ -359,9 +359,9 @@ namespace Horde.Server.Ddc
 
     public class ReferenceIsMissingBlobsException : Exception
     {
-        public List<BlobIdentifier> MissingBlobs { get; }
+        public List<BlobId> MissingBlobs { get; }
 
-        public ReferenceIsMissingBlobsException(List<BlobIdentifier> missingBlobs) : base($"References is missing these blobs: {string.Join(',', missingBlobs)}")
+        public ReferenceIsMissingBlobsException(List<BlobId> missingBlobs) : base($"References is missing these blobs: {string.Join(',', missingBlobs)}")
         {
             MissingBlobs = missingBlobs;
         }
