Horde: Add 'Async' to the end of all asynchronous DDC methods.

diff --git a/Engine/Source/Programs/Horde/Horde.Server/Ddc/ContentId.cs b/Engine/Source/Programs/Horde/Horde.Server/Ddc/ContentId.cs
--- a/Engine/Source/Programs/Horde/Horde.Server/Ddc/ContentId.cs
+++ b/Engine/Source/Programs/Horde/Horde.Server/Ddc/ContentId.cs
@@ -57,14 +57,14 @@ namespace Horde.Server.Ddc
             return new ContentId(contentHash.HashData);
         }
 
-        public static ContentId FromBlobIdentifier(BlobIdentifier blobIdentifier)
+        public static ContentId FromBlobIdentifier(BlobId blobIdentifier)
         {
             return new ContentId(blobIdentifier.HashData);
         }
 
-        public BlobIdentifier AsBlobIdentifier()
+        public BlobId AsBlobIdentifier()
         {
-            return new BlobIdentifier(HashData);
+            return new BlobId(HashData);
         }
 
         public static ContentId FromIoHash(IoHash ioHash)
diff --git a/Engine/Source/Programs/Horde/Horde.Server/Ddc/DdcBlobService.cs b/Engine/Source/Programs/Horde/Horde.Server/Ddc/DdcBlobService.cs
--- a/Engine/Source/Programs/Horde/Horde.Server/Ddc/DdcBlobService.cs
+++ b/Engine/Source/Programs/Horde/Horde.Server/Ddc/DdcBlobService.cs
@@ -10,77 +10,77 @@ namespace Horde.Server.Ddc
 {
 	class DdcBlobService : IDdcBlobService
 	{
-		public Task DeleteNamespace(NamespaceId ns)
+		public Task DeleteNamespaceAsync(NamespaceId ns)
 		{
 			throw new NotImplementedException();
 		}
 
-		public Task DeleteObject(NamespaceId ns, BlobIdentifier blob)
+		public Task DeleteObjectAsync(NamespaceId ns, BlobId blob)
 		{
 			throw new NotImplementedException();
 		}
 
-		public Task<bool> Exists(NamespaceId ns, BlobIdentifier blob, List<string>? storageLayers = null)
+		public Task<bool> ExistsAsync(NamespaceId ns, BlobId blob, List<string>? storageLayers = null)
 		{
 			throw new NotImplementedException();
 		}
 
-		public Task<bool> ExistsInRootStore(NamespaceId ns, BlobIdentifier blob)
+		public Task<bool> ExistsInRootStore(NamespaceId ns, BlobId blob)
 		{
 			throw new NotImplementedException();
 		}
 
-		public Task<BlobIdentifier[]> FilterOutKnownBlobs(NamespaceId ns, IEnumerable<BlobIdentifier> blobs)
+		public Task<BlobId[]> FilterOutKnownBlobsAsync(NamespaceId ns, IEnumerable<BlobId> blobs)
 		{
 			throw new NotImplementedException();
 		}
 
-		public Task<BlobIdentifier[]> FilterOutKnownBlobs(NamespaceId ns, IAsyncEnumerable<BlobIdentifier> blobs)
+		public Task<BlobId[]> FilterOutKnownBlobsAsync(NamespaceId ns, IAsyncEnumerable<BlobId> blobs)
 		{
 			throw new NotImplementedException();
 		}
 
-		public Task<BlobContents> GetObject(NamespaceId ns, BlobIdentifier blob, List<string>? storageLayers = null, bool supportsRedirectUri = false)
+		public Task<BlobContents> GetObjectAsync(NamespaceId ns, BlobId blob, List<string>? storageLayers = null, bool supportsRedirectUri = false)
 		{
 			throw new NotImplementedException();
 		}
 
-		public Task<BlobContents> GetObjects(NamespaceId ns, BlobIdentifier[] refRequestBlobReferences)
+		public Task<BlobContents> GetObjectsAsync(NamespaceId ns, BlobId[] refRequestBlobReferences)
 		{
 			throw new NotImplementedException();
 		}
 
-		public Task<Uri?> GetObjectWithRedirect(NamespaceId ns, BlobIdentifier blobIdentifier, List<string>? storageLayers = null)
+		public Task<Uri?> GetObjectWithRedirectAsync(NamespaceId ns, BlobId blobIdentifier, List<string>? storageLayers = null)
 		{
 			throw new NotImplementedException();
 		}
 
-		public IAsyncEnumerable<(BlobIdentifier, DateTime)> ListObjects(NamespaceId ns)
+		public IAsyncEnumerable<(BlobId, DateTime)> ListObjects(NamespaceId ns)
 		{
 			throw new NotImplementedException();
 		}
 
-		public Task<Uri?> MaybePutObjectWithRedirect(NamespaceId ns, BlobIdentifier identifier)
+		public Task<Uri?> MaybePutObjectWithRedirectAsync(NamespaceId ns, BlobId identifier)
 		{
 			throw new NotImplementedException();
 		}
 
-		public Task<BlobIdentifier> PutObject(NamespaceId ns, IBufferedPayload payload, BlobIdentifier identifier)
+		public Task<BlobId> PutObjectAsync(NamespaceId ns, IBufferedPayload payload, BlobId identifier)
 		{
 			throw new NotImplementedException();
 		}
 
-		public Task<BlobIdentifier> PutObject(NamespaceId ns, byte[] payload, BlobIdentifier identifier)
+		public Task<BlobId> PutObjectAsync(NamespaceId ns, byte[] payload, BlobId identifier)
 		{
 			throw new NotImplementedException();
 		}
 
-		public Task<BlobIdentifier> PutObjectKnownHash(NamespaceId ns, IBufferedPayload content, BlobIdentifier identifier)
+		public Task<BlobId> PutObjectKnownHashAsync(NamespaceId ns, IBufferedPayload content, BlobId identifier)
 		{
 			throw new NotImplementedException();
 		}
 
-		public Task<BlobContents> ReplicateObject(NamespaceId ns, BlobIdentifier blob, bool force = false)
+		public Task<BlobContents> ReplicateObjectAsync(NamespaceId ns, BlobId blob, bool force = false)
 		{
 			throw new NotImplementedException();
 		}
@@ -90,7 +90,7 @@ namespace Horde.Server.Ddc
 			throw new NotImplementedException();
 		}
 
-		public Task<JupiterContentHash> VerifyContentMatchesHash(Stream content, JupiterContentHash identifier)
+		public Task<JupiterContentHash> VerifyContentMatchesHashAsync(Stream content, JupiterContentHash identifier)
 		{
 			throw new NotImplementedException();
 		}
diff --git a/Engine/Source/Programs/Horde/Horde.Server/Ddc/DdcBlobsController.cs b/Engine/Source/Programs/Horde/Horde.Server/Ddc/DdcBlobsController.cs
--- a/Engine/Source/Programs/Horde/Horde.Server/Ddc/DdcBlobsController.cs
+++ b/Engine/Source/Programs/Horde/Horde.Server/Ddc/DdcBlobsController.cs
@@ -52,7 +52,7 @@ namespace Horde.Server.Ddc
             [Required] BlobId id,
             [FromQuery] List<string>? storageLayers = null)
         {
-            ActionResult? result = await _requestHelper.HasAccessToNamespace(User, Request, ns, new [] { DdcAclAction.ReadObject });
+            ActionResult? result = await _requestHelper.HasAccessToNamespaceAsync(User, Request, ns, new [] { DdcAclAction.ReadObject });
             if (result != null)
             {
                 return result;
@@ -89,12 +89,12 @@ namespace Horde.Server.Ddc
             [Required] BlobId id,
             [FromQuery] List<string>? storageLayers = null)
         {
-            ActionResult? result = await _requestHelper.HasAccessToNamespace(User, Request, ns, new [] { DdcAclAction.ReadObject });
+            ActionResult? result = await _requestHelper.HasAccessToNamespaceAsync(User, Request, ns, new [] { DdcAclAction.ReadObject });
             if (result != null)
             {
                 return result;
             }
-            bool exists = await _storage.Exists(ns, id, storageLayers);
+            bool exists = await _storage.ExistsAsync(ns, id, storageLayers);
 
             if (!exists)
             {
@@ -110,7 +110,7 @@ namespace Horde.Server.Ddc
             [Required] NamespaceId ns,
             [Required] [FromQuery] List<BlobId> id)
         {
-            ActionResult? result = await _requestHelper.HasAccessToNamespace(User, Request, ns, new [] { DdcAclAction.ReadObject });
+            ActionResult? result = await _requestHelper.HasAccessToNamespaceAsync(User, Request, ns, new [] { DdcAclAction.ReadObject });
             if (result != null)
             {
                 return result;
@@ -120,7 +120,7 @@ namespace Horde.Server.Ddc
 
             IEnumerable<Task> tasks = id.Select(async blob =>
             {
-                if (!await _storage.Exists(ns, blob))
+                if (!await _storage.ExistsAsync(ns, blob))
                 {
                     missingBlobs.Add(blob);
                 }
@@ -136,7 +136,7 @@ namespace Horde.Server.Ddc
             [Required] NamespaceId ns,
             [FromBody] BlobId[] bodyIds)
         {
-            ActionResult? result = await _requestHelper.HasAccessToNamespace(User, Request, ns, new [] { DdcAclAction.ReadObject });
+            ActionResult? result = await _requestHelper.HasAccessToNamespaceAsync(User, Request, ns, new [] { DdcAclAction.ReadObject });
             if (result != null)
             {
                 return result;
@@ -146,7 +146,7 @@ namespace Horde.Server.Ddc
 
             IEnumerable<Task> tasks = bodyIds.Select(async blob =>
             {
-                if (!await _storage.Exists(ns, blob))
+                if (!await _storage.ExistsAsync(ns, blob))
                 {
                     missingBlobs.Add(blob);
                 }
@@ -158,19 +158,7 @@ namespace Horde.Server.Ddc
 
         private async Task<BlobContents> GetImpl(NamespaceId ns, BlobId blob, List<string>? storageLayers = null, bool supportsRedirectUri = false)
         {
-            try
-            {
-                return await _storage.GetObject(ns, blob, storageLayers, supportsRedirectUri);
-            }
-            catch (BlobNotFoundException)
-            {
-                if (!_storage.ShouldFetchBlobOnDemand(ns))
-                {
-                    throw;
-                }
-
-                return await _storage.ReplicateObject(ns, blob);
-            }
+            return await _storage.GetObjectAsync(ns, blob, storageLayers, supportsRedirectUri);
         }
 
         [HttpPut("{ns}/{id}")]
@@ -180,7 +168,7 @@ namespace Horde.Server.Ddc
             [Required] NamespaceId ns,
             [Required] BlobId id)
         {
-            ActionResult? result = await _requestHelper.HasAccessToNamespace(User, Request, ns, new [] { DdcAclAction.WriteObject });
+            ActionResult? result = await _requestHelper.HasAccessToNamespaceAsync(User, Request, ns, new [] { DdcAclAction.WriteObject });
             if (result != null)
             {
                 return result;
@@ -190,7 +178,7 @@ namespace Horde.Server.Ddc
 
             try
             {
-                Uri? uri = await _storage.MaybePutObjectWithRedirect(ns, id);
+                Uri? uri = await _storage.MaybePutObjectWithRedirectAsync(ns, id);
                 if (uri != null)
                 {
                     return Ok(new
@@ -201,7 +189,7 @@ namespace Horde.Server.Ddc
                 }
                 using IBufferedPayload payload = await _bufferedPayloadFactory.CreateFromRequest(Request);
 
-                BlobId identifier = await _storage.PutObject(ns, payload, id);
+                BlobId identifier = await _storage.PutObjectAsync(ns, payload, id);
                 return Ok(new
                 {
                     Identifier = identifier.ToString()
@@ -223,7 +211,7 @@ namespace Horde.Server.Ddc
         public async Task<IActionResult> Post(
             [Required] NamespaceId ns)
         {
-            ActionResult? result = await _requestHelper.HasAccessToNamespace(User, Request, ns, new [] { DdcAclAction.WriteObject });
+            ActionResult? result = await _requestHelper.HasAccessToNamespaceAsync(User, Request, ns, new [] { DdcAclAction.WriteObject });
             if (result != null)
             {
                 return result;
@@ -237,7 +225,7 @@ namespace Horde.Server.Ddc
                 await using Stream stream = payload.GetStream();
 
                 BlobId id = await BlobId.FromStream(stream);
-                await _storage.PutObjectKnownHash(ns, payload, id);
+                await _storage.PutObjectKnownHashAsync(ns, payload, id);
                 
                 return Ok(new
                 {
@@ -255,7 +243,7 @@ namespace Horde.Server.Ddc
             [Required] NamespaceId ns,
             [Required] BlobId id)
         {
-            ActionResult? result = await _requestHelper.HasAccessToNamespace(User, Request, ns, new [] { DdcAclAction.DeleteObject });
+            ActionResult? result = await _requestHelper.HasAccessToNamespaceAsync(User, Request, ns, new [] { DdcAclAction.DeleteObject });
             if (result != null)
             {
                 return result;
@@ -271,20 +259,20 @@ namespace Horde.Server.Ddc
         public async Task<IActionResult> DeleteNamespace(
             [Required] NamespaceId ns)
         {
-            ActionResult? result = await _requestHelper.HasAccessToNamespace(User, Request, ns, new [] { DdcAclAction.DeleteNamespace });
+            ActionResult? result = await _requestHelper.HasAccessToNamespaceAsync(User, Request, ns, new [] { DdcAclAction.DeleteNamespace });
             if (result != null)
             {
                 return result;
             }
 
-            await  _storage.DeleteNamespace(ns);
+            await  _storage.DeleteNamespaceAsync(ns);
 
             return NoContent();
         }
 
         private async Task DeleteImpl(NamespaceId ns, BlobId id)
         {
-            await _storage.DeleteObject(ns, id);
+            await _storage.DeleteObjectAsync(ns, id);
         }
 
         // ReSharper disable UnusedAutoPropertyAccessor.Global
@@ -351,7 +339,7 @@ namespace Horde.Server.Ddc
                     throw new Exception(nameof(op.Namespace));
                 }
 
-                ActionResult? result = await _requestHelper.HasAccessToNamespace(User, Request, op.Namespace!.Value, new[] { MapToAclAction(op.Op) });
+                ActionResult? result = await _requestHelper.HasAccessToNamespaceAsync(User, Request, op.Namespace!.Value, new[] { MapToAclAction(op.Op) });
 
                 if (result != null)
                 {
@@ -386,7 +374,7 @@ namespace Horde.Server.Ddc
                             return BadRequest();
                         }
 
-                        tasks[index] = _storage.Exists(op.Namespace.Value, op.Id)
+                        tasks[index] = _storage.ExistsAsync(op.Namespace.Value, op.Id)
                             .ContinueWith((t,_) => t.Result ? (object?) null : op.Id, null, TaskScheduler.Current);
                         break;
                     case BatchOp.Operation.PUT:
@@ -402,7 +390,7 @@ namespace Horde.Server.Ddc
                             }
 
                             using MemoryBufferedPayload payload = new MemoryBufferedPayload(op.Content);
-                            tasks[index] = _storage.PutObject(op.Namespace.Value, payload, op.Id).ContinueWith((t, _) => (object?) t.Result, null, TaskScheduler.Current);
+                            tasks[index] = _storage.PutObjectAsync(op.Namespace.Value, payload, op.Id).ContinueWith((t, _) => (object?) t.Result, null, TaskScheduler.Current);
                             break;
                         }
                     case BatchOp.Operation.DELETE:
diff --git a/Engine/Source/Programs/Horde/Horde.Server/Ddc/DdcContentIdController.cs b/Engine/Source/Programs/Horde/Horde.Server/Ddc/DdcContentIdController.cs
--- a/Engine/Source/Programs/Horde/Horde.Server/Ddc/DdcContentIdController.cs
+++ b/Engine/Source/Programs/Horde/Horde.Server/Ddc/DdcContentIdController.cs
@@ -37,12 +37,12 @@ namespace Horde.Server.Ddc
         [ProducesResponseType(type: typeof(ProblemDetails), 400)]
         public async Task<IActionResult> Resolve(NamespaceId ns, ContentId contentId)
         {
-            ActionResult? result = await _requestHelper.HasAccessToNamespace(User, Request, ns, new [] { DdcAclAction.ReadObject });
+            ActionResult? result = await _requestHelper.HasAccessToNamespaceAsync(User, Request, ns, new [] { DdcAclAction.ReadObject });
             if (result != null)
             {
                 return result;
             }
-            BlobId[]? blobs = await _contentIdStore.Resolve(ns, contentId, mustBeContentId: true);
+            BlobId[]? blobs = await _contentIdStore.ResolveAsync(ns, contentId, mustBeContentId: true);
 
             if (blobs == null)
             {
@@ -64,12 +64,12 @@ namespace Horde.Server.Ddc
         [ProducesResponseType(type: typeof(ProblemDetails), 400)]
         public async Task<IActionResult> UpdateContentIdMapping(NamespaceId ns, ContentId contentId, BlobId blobIdentifier, int contentWeight)
         {
-            ActionResult? result = await _requestHelper.HasAccessToNamespace(User, Request, ns, new [] { DdcAclAction.WriteObject });
+            ActionResult? result = await _requestHelper.HasAccessToNamespaceAsync(User, Request, ns, new [] { DdcAclAction.WriteObject });
             if (result != null)
             {
                 return result;
             }
-            await _contentIdStore.Put(ns, contentId, blobIdentifier, contentWeight);
+            await _contentIdStore.PutAsync(ns, contentId, blobIdentifier, contentWeight);
 
             return Ok();
         }
diff --git a/Engine/Source/Programs/Horde/Horde.Server/Ddc/DdcContentIdService.cs b/Engine/Source/Programs/Horde/Horde.Server/Ddc/DdcContentIdService.cs
--- a/Engine/Source/Programs/Horde/Horde.Server/Ddc/DdcContentIdService.cs
+++ b/Engine/Source/Programs/Horde/Horde.Server/Ddc/DdcContentIdService.cs
@@ -8,12 +8,12 @@ namespace Horde.Server.Ddc
 {
 	class ContentIdStore : IDdcContentIdService
 	{
-		public Task<BlobId[]?> Resolve(NamespaceId ns, ContentId contentId, bool mustBeContentId = false)
+		public Task<BlobId[]?> ResolveAsync(NamespaceId ns, ContentId contentId, bool mustBeContentId = false)
 		{
 			throw new NotImplementedException();
 		}
 
-		public Task Put(NamespaceId ns, ContentId contentId, BlobId blobIdentifier, int contentWeight)
+		public Task PutAsync(NamespaceId ns, ContentId contentId, BlobId blobIdentifier, int contentWeight)
 		{
 			throw new NotImplementedException();
 		}
diff --git a/Engine/Source/Programs/Horde/Horde.Server/Ddc/DdcObjectsController.cs b/Engine/Source/Programs/Horde/Horde.Server/Ddc/DdcObjectsController.cs
--- a/Engine/Source/Programs/Horde/Horde.Server/Ddc/DdcObjectsController.cs
+++ b/Engine/Source/Programs/Horde/Horde.Server/Ddc/DdcObjectsController.cs
@@ -51,7 +51,7 @@ namespace Horde.Server.Ddc
             [Required] NamespaceId ns,
             [Required] BlobId id)
         {
-            ActionResult? result = await _requestHelper.HasAccessToNamespace(User, Request, ns, new [] { DdcAclAction.ReadObject });
+            ActionResult? result = await _requestHelper.HasAccessToNamespaceAsync(User, Request, ns, new [] { DdcAclAction.ReadObject });
             if (result != null)
             {
                 return result;
@@ -59,7 +59,7 @@ namespace Horde.Server.Ddc
 
             try
             {
-                BlobContents blobContents = await _storage.GetObject(ns, id);
+                BlobContents blobContents = await _storage.GetObjectAsync(ns, id);
 
                 return File(blobContents.Stream, CustomMediaTypeNames.UnrealCompactBinary);
             }
@@ -75,13 +75,13 @@ namespace Horde.Server.Ddc
             [Required] NamespaceId ns,
             [Required] BlobId id)
         {
-            ActionResult? result = await _requestHelper.HasAccessToNamespace(User, Request, ns, new [] { DdcAclAction.ReadObject });
+            ActionResult? result = await _requestHelper.HasAccessToNamespaceAsync(User, Request, ns, new [] { DdcAclAction.ReadObject });
             if (result != null)
             {
                 return result;
             }
 
-            bool exists = await _storage.Exists(ns, id);
+            bool exists = await _storage.ExistsAsync(ns, id);
 
             if (!exists)
             {
@@ -97,7 +97,7 @@ namespace Horde.Server.Ddc
             [Required] NamespaceId ns,
             [Required] [FromQuery] List<BlobId> id)
         {
-            ActionResult? result = await _requestHelper.HasAccessToNamespace(User, Request, ns, new [] { DdcAclAction.ReadObject });
+            ActionResult? result = await _requestHelper.HasAccessToNamespaceAsync(User, Request, ns, new [] { DdcAclAction.ReadObject });
             if (result != null)
             {
                 return result;
@@ -107,7 +107,7 @@ namespace Horde.Server.Ddc
 
             IEnumerable<Task> tasks = id.Select(async blob =>
             {
-                if (!await _storage.Exists(ns, blob))
+                if (!await _storage.ExistsAsync(ns, blob))
                 {
                     missingBlobs.Add(blob);
                 }
@@ -123,7 +123,7 @@ namespace Horde.Server.Ddc
             [Required] NamespaceId ns,
             [FromBody] BlobId[] bodyIds)
         {
-            ActionResult? result = await _requestHelper.HasAccessToNamespace(User, Request, ns, new [] { DdcAclAction.ReadObject });
+            ActionResult? result = await _requestHelper.HasAccessToNamespaceAsync(User, Request, ns, new [] { DdcAclAction.ReadObject });
             if (result != null)
             {
                 return result;
@@ -133,7 +133,7 @@ namespace Horde.Server.Ddc
 
             IEnumerable<Task> tasks = bodyIds.Select(async blob =>
             {
-                if (!await _storage.Exists(ns, blob))
+                if (!await _storage.ExistsAsync(ns, blob))
                 {
                     missingBlobs.Add(blob);
                 }
@@ -149,7 +149,7 @@ namespace Horde.Server.Ddc
             [Required] NamespaceId ns,
             [Required] BlobId id)
         {
-            ActionResult? result = await _requestHelper.HasAccessToNamespace(User, Request, ns, new [] { DdcAclAction.WriteObject });
+            ActionResult? result = await _requestHelper.HasAccessToNamespaceAsync(User, Request, ns, new [] { DdcAclAction.WriteObject });
             if (result != null)
             {
                 return result;
@@ -160,7 +160,7 @@ namespace Horde.Server.Ddc
             {
                 using IBufferedPayload payload = await _bufferedPayloadFactory.CreateFromRequest(Request);
 
-                BlobId identifier = await _storage.PutObject(ns, payload, id);
+                BlobId identifier = await _storage.PutObjectAsync(ns, payload, id);
                 return Ok(new PutBlobResponse(identifier));
             }
             catch (ClientSendSlowException e)
@@ -174,7 +174,7 @@ namespace Horde.Server.Ddc
             [Required] NamespaceId ns,
             [Required] BlobId id)
         {
-            ActionResult? result = await _requestHelper.HasAccessToNamespace(User, Request, ns, new [] { DdcAclAction.ReadObject });
+            ActionResult? result = await _requestHelper.HasAccessToNamespaceAsync(User, Request, ns, new [] { DdcAclAction.ReadObject });
             if (result != null)
             {
                 return result;
@@ -183,7 +183,7 @@ namespace Horde.Server.Ddc
             BlobContents blob;
             try
             {
-                blob = await _storage.GetObject(ns, id);
+                blob = await _storage.GetObjectAsync(ns, id);
             }
             catch (BlobNotFoundException e)
             {
@@ -208,7 +208,7 @@ namespace Horde.Server.Ddc
 
             try
             {
-                BlobId[] references = await _referenceResolver.GetReferencedBlobs(ns, compactBinaryObject).ToArrayAsync();
+                BlobId[] references = await _referenceResolver.GetReferencedBlobsAsync(ns, compactBinaryObject).ToArrayAsync();
                 return Ok(new ResolvedReferencesResult(references));
             }
             catch (PartialReferenceResolveException e)
@@ -226,13 +226,13 @@ namespace Horde.Server.Ddc
             [Required] NamespaceId ns,
             [Required] BlobId id)
         {
-            ActionResult? result = await _requestHelper.HasAccessToNamespace(User, Request, ns, new [] { DdcAclAction.DeleteObject });
+            ActionResult? result = await _requestHelper.HasAccessToNamespaceAsync(User, Request, ns, new [] { DdcAclAction.DeleteObject });
             if (result != null)
             {
                 return result;
             }
 
-            await _storage.DeleteObject(ns, id);
+            await _storage.DeleteObjectAsync(ns, id);
 
             return Ok( new DeletedResponse
             {
@@ -244,13 +244,13 @@ namespace Horde.Server.Ddc
         public async Task<IActionResult> DeleteNamespace(
             [Required] NamespaceId ns)
         {
-            ActionResult? result = await _requestHelper.HasAccessToNamespace(User, Request, ns, new [] { DdcAclAction.DeleteNamespace });
+            ActionResult? result = await _requestHelper.HasAccessToNamespaceAsync(User, Request, ns, new [] { DdcAclAction.DeleteNamespace });
             if (result != null)
             {
                 return result;
             }
 
-            await _storage.DeleteNamespace(ns);
+            await _storage.DeleteNamespaceAsync(ns);
 
             return Ok();
         }
diff --git a/Engine/Source/Programs/Horde/Horde.Server/Ddc/DdcRefService.cs b/Engine/Source/Programs/Horde/Horde.Server/Ddc/DdcRefService.cs
--- a/Engine/Source/Programs/Horde/Horde.Server/Ddc/DdcRefService.cs
+++ b/Engine/Source/Programs/Horde/Horde.Server/Ddc/DdcRefService.cs
@@ -10,47 +10,47 @@ namespace Horde.Server.Ddc
 {
 	class DdcRefService : IDdcRefService
 	{
-		public Task<bool> Delete(NamespaceId ns, BucketId bucket, RefId key)
+		public Task<bool> DeleteAsync(NamespaceId ns, BucketId bucket, RefId key)
 		{
 			throw new System.NotImplementedException();
 		}
 
-		public Task<long> DeleteBucket(NamespaceId ns, BucketId bucket)
+		public Task<long> DeleteBucketAsync(NamespaceId ns, BucketId bucket)
 		{
 			throw new System.NotImplementedException();
 		}
 
-		public Task<long> DropNamespace(NamespaceId ns)
+		public Task<long> DropNamespaceAsync(NamespaceId ns)
 		{
 			throw new System.NotImplementedException();
 		}
 
-		public Task<bool> Exists(NamespaceId ns, BucketId bucket, RefId key)
+		public Task<bool> ExistsAsync(NamespaceId ns, BucketId bucket, RefId key)
 		{
 			throw new System.NotImplementedException();
 		}
 
-		public Task<(ContentId[], BlobIdentifier[])> Finalize(NamespaceId ns, BucketId bucket, RefId key, BlobIdentifier blobHash)
+		public Task<(ContentId[], BlobId[])> FinalizeAsync(NamespaceId ns, BucketId bucket, RefId key, BlobId blobHash)
 		{
 			throw new System.NotImplementedException();
 		}
 
-		public Task<(RefRecord, BlobContents?)> Get(NamespaceId ns, BucketId bucket, RefId key, string[] fields, bool doLastAccessTracking = true)
+		public Task<(RefRecord, BlobContents?)> GetAsync(NamespaceId ns, BucketId bucket, RefId key, string[] fields, bool doLastAccessTracking = true)
 		{
 			throw new System.NotImplementedException();
 		}
 
-		public IAsyncEnumerable<NamespaceId> GetNamespaces()
+		public IAsyncEnumerable<NamespaceId> GetNamespacesAsync()
 		{
 			throw new System.NotImplementedException();
 		}
 
-		public Task<List<BlobIdentifier>> GetReferencedBlobs(NamespaceId ns, BucketId bucket, RefId key)
+		public Task<List<BlobId>> GetReferencedBlobsAsync(NamespaceId ns, BucketId bucket, RefId key)
 		{
 			throw new System.NotImplementedException();
 		}
 
-		public Task<(ContentId[], BlobIdentifier[])> Put(NamespaceId ns, BucketId bucket, RefId key, BlobIdentifier blobHash, CbObject payload)
+		public Task<(ContentId[], BlobId[])> PutAsync(NamespaceId ns, BucketId bucket, RefId key, BlobId blobHash, CbObject payload)
 		{
 			throw new System.NotImplementedException();
 		}
diff --git a/Engine/Source/Programs/Horde/Horde.Server/Ddc/DdcRefsController.cs b/Engine/Source/Programs/Horde/Horde.Server/Ddc/DdcRefsController.cs
--- a/Engine/Source/Programs/Horde/Horde.Server/Ddc/DdcRefsController.cs
+++ b/Engine/Source/Programs/Horde/Horde.Server/Ddc/DdcRefsController.cs
@@ -75,13 +75,13 @@ namespace Horde.Server.Ddc
         [ProducesResponseType(type: typeof(ProblemDetails), 400)]
         public async Task<IActionResult> GetNamespaces()
         {
-            NamespaceId[] namespaces = await _refService.GetNamespaces().ToArrayAsync();
+            NamespaceId[] namespaces = await _refService.GetNamespacesAsync().ToArrayAsync();
 
             // filter namespaces down to only the namespaces the user has access to
             List<NamespaceId> namespacesWithAccess = new();
             foreach (NamespaceId ns in namespaces)
             {
-                ActionResult? accessResult = await _requestHelper.HasAccessToNamespace(User, Request, ns, new [] { DdcAclAction.ReadObject });
+                ActionResult? accessResult = await _requestHelper.HasAccessToNamespaceAsync(User, Request, ns, new [] { DdcAclAction.ReadObject });
                 if (accessResult == null)
                 {
                     namespacesWithAccess.Add(ns);
@@ -106,7 +106,7 @@ namespace Horde.Server.Ddc
             [FromRoute] [Required] RefId key,
             [FromRoute] string? format = null)
         {
-            ActionResult? accessResult = await _requestHelper.HasAccessToNamespace(User, Request, ns, new [] { DdcAclAction.ReadObject });
+            ActionResult? accessResult = await _requestHelper.HasAccessToNamespaceAsync(User, Request, ns, new [] { DdcAclAction.ReadObject });
             if (accessResult != null)
             {
                 return accessResult;
@@ -114,7 +114,7 @@ namespace Horde.Server.Ddc
 
             try
             {
-                (RefRecord objectRecord, BlobContents? blob) = await _refService.Get(ns, bucket, key, Array.Empty<string>());
+                (RefRecord objectRecord, BlobContents? blob) = await _refService.GetAsync(ns, bucket, key, Array.Empty<string>());
 
                 if (blob == null)
                 {
@@ -198,7 +198,7 @@ namespace Horde.Server.Ddc
 
                             IoHash hash = binaryAttachmentField.AsBinaryAttachment();
 
-                            BlobContents referencedBlobContents = await _blobStore.GetObject(ns, BlobId.FromIoHash(hash));
+                            BlobContents referencedBlobContents = await _blobStore.GetObjectAsync(ns, BlobId.FromIoHash(hash));
 
                             if (_nginxRedirectHelper.CanRedirect(Request, referencedBlobContents))
                             {
@@ -233,7 +233,7 @@ namespace Horde.Server.Ddc
                         byte[] blobMemory = await blob.Stream.ToByteArray();
                         CbObject cb = new CbObject(blobMemory);
 
-                        IAsyncEnumerable<Attachment> attachments = _referenceResolver.GetAttachments(ns, cb);
+                        IAsyncEnumerable<Attachment> attachments = _referenceResolver.GetAttachmentsAsync(ns, cb);
 
                         using CbPackageBuilder writer = new CbPackageBuilder();
                         writer.AddAttachment(objectRecord.BlobIdentifier.AsIoHash(), CbPackageAttachmentFlags.IsObject, blobMemory);
@@ -249,13 +249,13 @@ namespace Horde.Server.Ddc
                                 if (attachment is BlobAttachment blobAttachment)
                                 {
                                     BlobId referencedBlob = blobAttachment.Identifier;
-                                    attachmentContents = await _blobStore.GetObject(ns, referencedBlob);
+                                    attachmentContents = await _blobStore.GetObjectAsync(ns, referencedBlob);
                                 }
                                 else if (attachment is ObjectAttachment objectAttachment)
                                 {
                                     flags |= CbPackageAttachmentFlags.IsObject;
                                     BlobId referencedBlob = objectAttachment.Identifier;
-                                    attachmentContents = await _blobStore.GetObject(ns, referencedBlob);
+                                    attachmentContents = await _blobStore.GetObjectAsync(ns, referencedBlob);
                                 }
                                 else if (attachment is ContentIdAttachment contentIdAttachment)
                                 {
@@ -333,7 +333,7 @@ namespace Horde.Server.Ddc
                             List<BlobId> referencedBlobs;
                             try
                             {
-                                IAsyncEnumerable<BlobId> referencedBlobsEnumerable = _referenceResolver.GetReferencedBlobs(ns, cb);
+                                IAsyncEnumerable<BlobId> referencedBlobsEnumerable = _referenceResolver.GetReferencedBlobsAsync(ns, cb);
                                 referencedBlobs = await referencedBlobsEnumerable.ToListAsync();
                             }
                             catch (PartialReferenceResolveException)
@@ -350,7 +350,7 @@ namespace Horde.Server.Ddc
                                 BlobId attachmentToSend = referencedBlobs.First();
                                 try
                                 {
-                                    BlobContents referencedBlobContents = await _blobStore.GetObject(ns, attachmentToSend);
+                                    BlobContents referencedBlobContents = await _blobStore.GetObjectAsync(ns, attachmentToSend);
                                     Response.Headers[CommonHeaders.InlinePayloadHash] = attachmentToSend.ToString();
 
                                     if (_nginxRedirectHelper.CanRedirect(Request, referencedBlobContents))
@@ -437,7 +437,7 @@ namespace Horde.Server.Ddc
         [FromRoute] [Required] RefId key,
         [FromQuery] string[] fields)
     {
-        ActionResult? accessResult = await _requestHelper.HasAccessToNamespace(User, Request, ns, new [] { DdcAclAction.ReadObject });
+        ActionResult? accessResult = await _requestHelper.HasAccessToNamespaceAsync(User, Request, ns, new [] { DdcAclAction.ReadObject });
         if (accessResult != null)
         {
             return accessResult;
@@ -445,7 +445,7 @@ namespace Horde.Server.Ddc
 
         try
         {
-            (RefRecord objectRecord, BlobContents? _) = await _refService.Get(ns, bucket, key, fields);
+            (RefRecord objectRecord, BlobContents? _) = await _refService.GetAsync(ns, bucket, key, fields);
 
             return Ok(new RefMetadataResponse(objectRecord));
         }
@@ -478,7 +478,7 @@ namespace Horde.Server.Ddc
             [FromRoute] [Required] BucketId bucket,
             [FromRoute] [Required] RefId key)
         {
-            ActionResult? accessResult = await _requestHelper.HasAccessToNamespace(User, Request, ns, new [] { DdcAclAction.ReadObject });
+            ActionResult? accessResult = await _requestHelper.HasAccessToNamespaceAsync(User, Request, ns, new [] { DdcAclAction.ReadObject });
             if (accessResult != null)
             {
                 return accessResult;
@@ -486,7 +486,7 @@ namespace Horde.Server.Ddc
 
             try
             {
-                (RefRecord record, BlobContents? blob) = await _refService.Get(ns, bucket, key, new string[] {"blobIdentifier", "IsFinalized"});
+                (RefRecord record, BlobContents? blob) = await _refService.GetAsync(ns, bucket, key, new string[] {"blobIdentifier", "IsFinalized"});
                 Response.Headers[CommonHeaders.HashHeaderName] = record.BlobIdentifier.ToString();
 
                 if (!record.IsFinalized)
@@ -494,19 +494,19 @@ namespace Horde.Server.Ddc
                     return NotFound(new ProblemDetails {Title = $"Object {bucket} {key} in namespace {ns} is not finalized."});
                 }
 
-                blob ??= await _blobStore.GetObject(ns, record.BlobIdentifier);
+                blob ??= await _blobStore.GetObjectAsync(ns, record.BlobIdentifier);
 
                 // we have to verify the blobs are available locally, as the record of the key is replicated a head of the content
                 // TODO: Once we support inline replication this step is not needed as at least one region as this blob, just maybe not this current one
                 byte[] blobContents = await blob.Stream.ToByteArray();
                 CbObject compactBinaryObject = new CbObject(blobContents);
                 // the reference resolver will throw if any blob is missing, so no need to do anything other then process each reference
-                IAsyncEnumerable<BlobId> references = _referenceResolver.GetReferencedBlobs(ns, compactBinaryObject);
+                IAsyncEnumerable<BlobId> references = _referenceResolver.GetReferencedBlobsAsync(ns, compactBinaryObject);
                 List<BlobId>? _ = await references.ToListAsync();
 
                 // we have to verify the blobs are available locally, as the record of the key is replicated a head of the content
                 // TODO: Once we support inline replication this step is not needed as at least one region as this blob, just maybe not this current one
-                BlobId[] unknownBlobs = await _blobStore.FilterOutKnownBlobs(ns, new BlobId[] { record.BlobIdentifier });
+                BlobId[] unknownBlobs = await _blobStore.FilterOutKnownBlobsAsync(ns, new BlobId[] { record.BlobIdentifier });
                 if (unknownBlobs.Length != 0)
                 {
                     return NotFound(new ProblemDetails {Title = $"Object {bucket} {key} in namespace {ns} had at least one missing blob."});
@@ -542,7 +542,7 @@ namespace Horde.Server.Ddc
             [FromRoute] [Required] NamespaceId ns,
             [FromQuery] [Required] List<string> names)
         {
-            ActionResult? accessResult = await _requestHelper.HasAccessToNamespace(User, Request, ns, new [] { DdcAclAction.ReadObject });
+            ActionResult? accessResult = await _requestHelper.HasAccessToNamespaceAsync(User, Request, ns, new [] { DdcAclAction.ReadObject });
             if (accessResult != null)
             {
                 return accessResult;
@@ -570,16 +570,16 @@ namespace Horde.Server.Ddc
                 try
                 {
                     (RefRecord record, BlobContents? blob) =
-                        await _refService.Get(ns, bucket, key, new string[] { "blobIdentifier" });
+                        await _refService.GetAsync(ns, bucket, key, new string[] { "blobIdentifier" });
 
-                    blob ??= await _blobStore.GetObject(ns, record.BlobIdentifier);
+                    blob ??= await _blobStore.GetObjectAsync(ns, record.BlobIdentifier);
 
                     // we have to verify the blobs are available locally, as the record of the key is replicated a head of the content
                     // TODO: Once we support inline replication this step is not needed as at least one region as this blob, just maybe not this current one
                     byte[] blobContents = await blob.Stream.ToByteArray();
                     CbObject cb = new CbObject(blobContents);
                     // the reference resolver will throw if any blob is missing, so no need to do anything other then process each reference
-                    IAsyncEnumerable<BlobId> references = _referenceResolver.GetReferencedBlobs(ns, cb);
+                    IAsyncEnumerable<BlobId> references = _referenceResolver.GetReferencedBlobsAsync(ns, cb);
                     List<BlobId>? _ = await references.ToListAsync();
                 }
                 catch (RefNotFoundException)
@@ -611,7 +611,7 @@ namespace Horde.Server.Ddc
             [FromRoute] [Required] BucketId bucket,
             [FromRoute] [Required] RefId key)
         {
-            ActionResult? accessResult = await _requestHelper.HasAccessToNamespace(User, Request, ns, new [] { DdcAclAction.WriteObject });
+            ActionResult? accessResult = await _requestHelper.HasAccessToNamespaceAsync(User, Request, ns, new [] { DdcAclAction.WriteObject });
             if (accessResult != null)
             {
                 return accessResult;
@@ -647,7 +647,7 @@ namespace Horde.Server.Ddc
                     {
                         // TODO: define a scheme for how a json object specifies references
 
-                        blobHeader = await _blobStore.PutObject(ns, payload, headerHash);
+                        blobHeader = await _blobStore.PutObjectAsync(ns, payload, headerHash);
 
                         // TODO: convert the json object into a compact binary instead
                         CbWriter writer = new CbWriter();
@@ -670,7 +670,7 @@ namespace Horde.Server.Ddc
                     }
                     case MediaTypeNames.Application.Octet:
                     {
-                        blobHeader = await _blobStore.PutObject(ns, payload, headerHash);
+                        blobHeader = await _blobStore.PutObjectAsync(ns, payload, headerHash);
 
                         CbWriter writer = new CbWriter();
                         writer.BeginObject();
@@ -699,7 +699,7 @@ namespace Horde.Server.Ddc
                 return Problem(e.Message, null, (int)HttpStatusCode.RequestTimeout);
             }
 
-            (ContentId[] missingReferences, BlobId[] missingBlobs) = await _refService.Put(ns, bucket, key, blobHeader, payloadObject);
+            (ContentId[] missingReferences, BlobId[] missingBlobs) = await _refService.PutAsync(ns, bucket, key, blobHeader, payloadObject);
 
             List<JupiterContentHash> missingHashes = new List<JupiterContentHash>(missingReferences);
             missingHashes.AddRange(missingBlobs);
@@ -714,7 +714,7 @@ namespace Horde.Server.Ddc
             [FromRoute][Required] BucketId bucket,
             [FromRoute][Required] RefId key)
         {
-            ActionResult? accessResult = await _requestHelper.HasAccessToNamespace(User, Request, ns, new [] { DdcAclAction.WriteObject });
+            ActionResult? accessResult = await _requestHelper.HasAccessToNamespaceAsync(User, Request, ns, new [] { DdcAclAction.WriteObject });
             if (accessResult != null)
             {
                 return accessResult;
@@ -743,7 +743,7 @@ namespace Horde.Server.Ddc
                     }
                     else
                     {
-                        await _blobStore.PutObject(ns, blob, BlobId.FromIoHash(entry.AttachmentHash));
+                        await _blobStore.PutObjectAsync(ns, blob, BlobId.FromIoHash(entry.AttachmentHash));
                     }
                 }
             }
@@ -758,7 +758,7 @@ namespace Horde.Server.Ddc
             CbObject rootObject = packageReader.RootObject;
             BlobId rootObjectHash = BlobId.FromIoHash(packageReader.RootHash);
 
-            (ContentId[] missingReferences, BlobId[] missingBlobs) = await _refService.Put(ns, bucket, key, rootObjectHash, rootObject);
+            (ContentId[] missingReferences, BlobId[] missingBlobs) = await _refService.PutAsync(ns, bucket, key, rootObjectHash, rootObject);
 
             List<JupiterContentHash> missingHashes = new List<JupiterContentHash>(missingReferences);
             missingHashes.AddRange(missingBlobs);
@@ -798,7 +798,7 @@ namespace Horde.Server.Ddc
             [FromRoute] [Required] RefId key,
             [FromRoute] [Required] BlobId hash)
         {
-            ActionResult? accessResult = await _requestHelper.HasAccessToNamespace(User, Request, ns, new [] { DdcAclAction.WriteObject });
+            ActionResult? accessResult = await _requestHelper.HasAccessToNamespaceAsync(User, Request, ns, new [] { DdcAclAction.WriteObject });
             if (accessResult != null)
             {
                 return accessResult;
@@ -806,7 +806,7 @@ namespace Horde.Server.Ddc
 
             try
             {
-                (ContentId[] missingReferences, BlobId[] missingBlobs) = await _refService.Finalize(ns, bucket, key, hash);
+                (ContentId[] missingReferences, BlobId[] missingBlobs) = await _refService.FinalizeAsync(ns, bucket, key, hash);
                 List<JupiterContentHash> missingHashes = new List<JupiterContentHash>(missingReferences);
                 missingHashes.AddRange(missingBlobs);
 
@@ -846,7 +846,7 @@ namespace Horde.Server.Ddc
 
             AclAction[] requiredActions = ops.Ops.Select(op => ActionForOp(op.Op)).ToArray();
 
-            ActionResult? accessResult = await _requestHelper.HasAccessToNamespace(User, Request, ns, requiredActions);
+            ActionResult? accessResult = await _requestHelper.HasAccessToNamespaceAsync(User, Request, ns, requiredActions);
             if (accessResult != null)
             {
                 return accessResult;
@@ -867,7 +867,7 @@ namespace Horde.Server.Ddc
             {
                 try
                 {
-                    (RefRecord objectRecord, BlobContents? blob) = await _refService.Get(ns, op.Bucket, op.Key, Array.Empty<string>());
+                    (RefRecord objectRecord, BlobContents? blob) = await _refService.GetAsync(ns, op.Bucket, op.Key, Array.Empty<string>());
 
                     if (!objectRecord.IsFinalized)
                     {
@@ -883,7 +883,7 @@ namespace Horde.Server.Ddc
 
                     if (op.ResolveAttachments ?? false)
                     {
-                        IAsyncEnumerable<BlobId> references = _referenceResolver.GetReferencedBlobs(ns, cb);
+                        IAsyncEnumerable<BlobId> references = _referenceResolver.GetReferencedBlobsAsync(ns, cb);
                         List<BlobId>? _ = await references.ToListAsync();
                     }
 
@@ -903,21 +903,21 @@ namespace Horde.Server.Ddc
             {
                 try
                 {
-                    (RefRecord record, BlobContents? blob) = await _refService.Get(ns, op.Bucket, op.Key, new string[] { "blobIdentifier" });
+                    (RefRecord record, BlobContents? blob) = await _refService.GetAsync(ns, op.Bucket, op.Key, new string[] { "blobIdentifier" });
 
                     if (!record.IsFinalized)
                     {
                         return (CbObject.Build(writer => writer.WriteBool("exists", false)), HttpStatusCode.NotFound);
                     }
 
-                    blob ??= await _blobStore.GetObject(ns, record.BlobIdentifier);
+                    blob ??= await _blobStore.GetObjectAsync(ns, record.BlobIdentifier);
 
                     if (op.ResolveAttachments ?? false)
                     {
                         byte[] blobContents = await blob.Stream.ToByteArray();
                         CbObject cb = new CbObject(blobContents);
                         // the reference resolver will throw if any blob is missing, so no need to do anything other then process each reference
-                        IAsyncEnumerable<BlobId> references = _referenceResolver.GetReferencedBlobs(ns, cb);
+                        IAsyncEnumerable<BlobId> references = _referenceResolver.GetReferencedBlobsAsync(ns, cb);
                         List<BlobId>? _ = await references.ToListAsync();
                     }
 
@@ -959,7 +959,7 @@ namespace Horde.Server.Ddc
                         throw new HashMismatchException(headerHash, objectHash);
                     }
 
-                    (ContentId[] missingReferences, BlobId[] missingBlobs) = await _refService.Put(ns, op.Bucket, op.Key, objectHash, op.Payload);
+                    (ContentId[] missingReferences, BlobId[] missingBlobs) = await _refService.PutAsync(ns, op.Bucket, op.Key, objectHash, op.Payload);
                     List<JupiterContentHash> missingHashes = new List<JupiterContentHash>(missingReferences);
 
                     return (CbSerializer.Serialize(new PutObjectResponse(missingHashes.ToArray())), HttpStatusCode.OK);
@@ -1030,7 +1030,7 @@ namespace Horde.Server.Ddc
             [FromRoute] [Required] NamespaceId ns
         )
         {
-            ActionResult? accessResult = await _requestHelper.HasAccessToNamespace(User, Request, ns, new [] { DdcAclAction.DeleteNamespace });
+            ActionResult? accessResult = await _requestHelper.HasAccessToNamespaceAsync(User, Request, ns, new [] { DdcAclAction.DeleteNamespace });
             if (accessResult != null)
             {
                 return accessResult;
@@ -1038,7 +1038,7 @@ namespace Horde.Server.Ddc
 
             try
             {
-                await _refService.DropNamespace(ns);
+                await _refService.DropNamespaceAsync(ns);
             }
             catch (NamespaceNotFoundException e)
             {
@@ -1059,7 +1059,7 @@ namespace Horde.Server.Ddc
             [FromRoute] [Required] NamespaceId ns,
             [FromRoute] [Required] BucketId bucket)
         {
-            ActionResult? accessResult = await _requestHelper.HasAccessToNamespace(User, Request, ns, new [] { DdcAclAction.DeleteBucket });
+            ActionResult? accessResult = await _requestHelper.HasAccessToNamespaceAsync(User, Request, ns, new [] { DdcAclAction.DeleteBucket });
             if (accessResult != null)
             {
                 return accessResult;
@@ -1068,7 +1068,7 @@ namespace Horde.Server.Ddc
             long countOfDeletedRecords;
             try
             {
-                countOfDeletedRecords = await _refService.DeleteBucket(ns, bucket);
+                countOfDeletedRecords = await _refService.DeleteBucketAsync(ns, bucket);
             }
             catch (NamespaceNotFoundException e)
             {
@@ -1092,7 +1092,7 @@ namespace Horde.Server.Ddc
             [FromRoute] [Required] BucketId bucket,
             [FromRoute] [Required] RefId key)
         {
-            ActionResult? accessResult = await _requestHelper.HasAccessToNamespace(User, Request, ns, new [] { DdcAclAction.DeleteObject });
+            ActionResult? accessResult = await _requestHelper.HasAccessToNamespaceAsync(User, Request, ns, new [] { DdcAclAction.DeleteObject });
             if (accessResult != null)
             {
                 return accessResult;
@@ -1100,7 +1100,7 @@ namespace Horde.Server.Ddc
 
             try
             {
-                bool deleted = await _refService.Delete(ns, bucket, key);
+                bool deleted = await _refService.DeleteAsync(ns, bucket, key);
                 if (!deleted)
                 {
                     return NotFound(new ProblemDetails { Title = $"Object {key} in bucket {bucket} and namespace {ns} did not exist" });
diff --git a/Engine/Source/Programs/Horde/Horde.Server/Ddc/IDdcBlobService.cs b/Engine/Source/Programs/Horde/Horde.Server/Ddc/IDdcBlobService.cs
--- a/Engine/Source/Programs/Horde/Horde.Server/Ddc/IDdcBlobService.cs
+++ b/Engine/Source/Programs/Horde/Horde.Server/Ddc/IDdcBlobService.cs
@@ -14,42 +14,31 @@ using OpenTelemetry.Trace;
 
 namespace Horde.Server.Ddc
 {
+	/// <summary>
+	/// Interface for the DDC blobs service
+	/// </summary>
 	public interface IDdcBlobService
 	{
-		Task<JupiterContentHash> VerifyContentMatchesHash(Stream content, JupiterContentHash identifier);
-		Task<BlobId> PutObjectKnownHash(NamespaceId ns, IBufferedPayload content, BlobId identifier);
-		Task<BlobId> PutObject(NamespaceId ns, IBufferedPayload payload, BlobId identifier);
-		Task<BlobId> PutObject(NamespaceId ns, byte[] payload, BlobId identifier);
-		Task<Uri?> MaybePutObjectWithRedirect(NamespaceId ns, BlobId identifier);
+		Task<JupiterContentHash> VerifyContentMatchesHashAsync(Stream content, JupiterContentHash identifier);
+		Task<BlobId> PutObjectKnownHashAsync(NamespaceId ns, IBufferedPayload content, BlobId identifier);
+		Task<BlobId> PutObjectAsync(NamespaceId ns, IBufferedPayload payload, BlobId identifier);
+		Task<BlobId> PutObjectAsync(NamespaceId ns, byte[] payload, BlobId identifier);
+		Task<Uri?> MaybePutObjectWithRedirectAsync(NamespaceId ns, BlobId identifier);
 
-		Task<BlobContents> GetObject(NamespaceId ns, BlobId blob, List<string>? storageLayers = null, bool supportsRedirectUri = false);
+		Task<BlobContents> GetObjectAsync(NamespaceId ns, BlobId blob, List<string>? storageLayers = null, bool supportsRedirectUri = false);
 
-		Task<Uri?> GetObjectWithRedirect(NamespaceId ns, BlobId blobIdentifier, List<string>? storageLayers = null);
+		Task<Uri?> GetObjectWithRedirectAsync(NamespaceId ns, BlobId blobIdentifier, List<string>? storageLayers = null);
 
-		Task<BlobContents> ReplicateObject(NamespaceId ns, BlobId blob, bool force = false);
-
-		Task<bool> Exists(NamespaceId ns, BlobId blob, List<string>? storageLayers = null);
-
-		/// <summary>
-		/// Checks that the blob exists in the root store, the store which is last in the list and thus is intended to have every blob in it
-		/// </summary>
-		/// <param name="ns">The namespace</param>
-		/// <param name="blob">The identifier of the blob</param>
-		/// <returns></returns>
-		Task<bool> ExistsInRootStore(NamespaceId ns, BlobId blob);
+		Task<bool> ExistsAsync(NamespaceId ns, BlobId blob, List<string>? storageLayers = null);
 
 		// Delete a object
-		Task DeleteObject(NamespaceId ns, BlobId blob);
+		Task DeleteObjectAsync(NamespaceId ns, BlobId blob);
 
 		// delete the whole namespace
-		Task DeleteNamespace(NamespaceId ns);
-
-		IAsyncEnumerable<(BlobId, DateTime)> ListObjects(NamespaceId ns);
-		Task<BlobId[]> FilterOutKnownBlobs(NamespaceId ns, IEnumerable<BlobId> blobs);
-		Task<BlobId[]> FilterOutKnownBlobs(NamespaceId ns, IAsyncEnumerable<BlobId> blobs);
-		Task<BlobContents> GetObjects(NamespaceId ns, BlobId[] refRequestBlobReferences);
+		Task DeleteNamespaceAsync(NamespaceId ns);
 
-		bool ShouldFetchBlobOnDemand(NamespaceId ns);
+		Task<BlobId[]> FilterOutKnownBlobsAsync(NamespaceId ns, IEnumerable<BlobId> blobs);
+		Task<BlobContents> GetObjectsAsync(NamespaceId ns, BlobId[] refRequestBlobReferences);
 	}
 
 	public class BlobNotFoundException : Exception
@@ -126,7 +115,7 @@ namespace Horde.Server.Ddc
 			ContentId identifierDecompressedPayload;
 			if (id != null)
 			{
-				identifierDecompressedPayload = ContentId.FromContentHash(await blobService.VerifyContentMatchesHash(decompressedStream, id));
+				identifierDecompressedPayload = ContentId.FromContentHash(await blobService.VerifyContentMatchesHashAsync(decompressedStream, id));
 			}
 			else
 			{
@@ -149,11 +138,11 @@ namespace Horde.Server.Ddc
 			// commit the mapping from the decompressed hash to the compressed hash, we run this in parallel with the blob store submit
 			// TODO: let users specify weight of the blob compared to previously submitted content ids
 			int contentIdWeight = (int)payload.Length;
-			Task contentIdStoreTask = contentIdStore.Put(ns, identifierDecompressedPayload, identifierCompressedPayload, contentIdWeight);
+			Task contentIdStoreTask = contentIdStore.PutAsync(ns, identifierDecompressedPayload, identifierCompressedPayload, contentIdWeight);
 
 			// we still commit the compressed buffer to the object store using the hash of the compressed content
 			{
-				await blobService.PutObjectKnownHash(ns, payload, identifierCompressedPayload);
+				await blobService.PutObjectKnownHashAsync(ns, payload, identifierCompressedPayload);
 			}
 
 			await contentIdStoreTask;
@@ -166,7 +155,7 @@ namespace Horde.Server.Ddc
 			IDdcContentIdService contentIdStore = provider.GetService<IDdcContentIdService>()!;
 			Tracer tracer = provider.GetService<Tracer>()!;
 
-			BlobId[]? chunks = await contentIdStore.Resolve(ns, contentId, mustBeContentId: false);
+			BlobId[]? chunks = await contentIdStore.ResolveAsync(ns, contentId, mustBeContentId: false);
 			if (chunks == null || chunks.Length == 0)
 			{
 				throw new ContentIdResolveException(contentId);
@@ -183,7 +172,7 @@ namespace Horde.Server.Ddc
 					mimeType = MediaTypeNames.Application.Octet;
 				}
 
-				return (await blobService.GetObject(ns, blobToReturn, supportsRedirectUri: supportsRedirectUri), mimeType);
+				return (await blobService.GetObjectAsync(ns, blobToReturn, supportsRedirectUri: supportsRedirectUri), mimeType);
 			}
 
 			// chunked content, combine the chunks into a single stream
@@ -192,7 +181,7 @@ namespace Horde.Server.Ddc
 			for (int i = 0; i < chunks.Length; i++)
 			{
 				// even if it was requested to support redirect, since we need to combine the chunks using redirects is not possible
-				tasks[i] = blobService.GetObject(ns, chunks[i], supportsRedirectUri: false);
+				tasks[i] = blobService.GetObjectAsync(ns, chunks[i], supportsRedirectUri: false);
 			}
 
 			MemoryStream ms = new MemoryStream();
diff --git a/Engine/Source/Programs/Horde/Horde.Server/Ddc/IDdcContentIdService.cs b/Engine/Source/Programs/Horde/Horde.Server/Ddc/IDdcContentIdService.cs
--- a/Engine/Source/Programs/Horde/Horde.Server/Ddc/IDdcContentIdService.cs
+++ b/Engine/Source/Programs/Horde/Horde.Server/Ddc/IDdcContentIdService.cs
@@ -17,7 +17,7 @@ namespace Horde.Server.Ddc
         /// <param name="contentId">The identifier for the content id</param>
         /// <param name="mustBeContentId"></param>
         /// <returns></returns>
-        Task<BlobIdentifier[]?> Resolve(NamespaceId ns, ContentId contentId, bool mustBeContentId = false);
+        Task<BlobId[]?> ResolveAsync(NamespaceId ns, ContentId contentId, bool mustBeContentId = false);
 
         /// <summary>
         /// Add a mapping from contentId to blobIdentifier
@@ -27,7 +27,7 @@ namespace Horde.Server.Ddc
         /// <param name="blobIdentifier">The blob the content id maps to</param>
         /// <param name="contentWeight">Weight of this identifier compared to previous mappings, used to determine which is more important, lower weight is considered a better fit</param>
         /// <returns></returns>
-        Task Put(NamespaceId ns, ContentId contentId, BlobIdentifier blobIdentifier, int contentWeight);
+        Task PutAsync(NamespaceId ns, ContentId contentId, BlobId blobIdentifier, int contentWeight);
     }
 
     public class InvalidContentIdException : Exception
diff --git a/Engine/Source/Programs/Horde/Horde.Server/Ddc/IDdcRefService.cs b/Engine/Source/Programs/Horde/Horde.Server/Ddc/IDdcRefService.cs
--- a/Engine/Source/Programs/Horde/Horde.Server/Ddc/IDdcRefService.cs
+++ b/Engine/Source/Programs/Horde/Horde.Server/Ddc/IDdcRefService.cs
@@ -13,23 +13,23 @@ namespace Horde.Server.Ddc
 {
 	public interface IDdcRefService
     {
-        Task<(RefRecord, BlobContents?)> Get(NamespaceId ns, BucketId bucket, RefId key, string[] fields, bool doLastAccessTracking = true);
-        Task<(ContentId[], BlobIdentifier[])> Put(NamespaceId ns, BucketId bucket, RefId key, BlobIdentifier blobHash, CbObject payload);
-        Task<(ContentId[], BlobIdentifier[])> Finalize(NamespaceId ns, BucketId bucket, RefId key, BlobIdentifier blobHash);
+        Task<(RefRecord, BlobContents?)> GetAsync(NamespaceId ns, BucketId bucket, RefId key, string[] fields, bool doLastAccessTracking = true);
+        Task<(ContentId[], BlobId[])> PutAsync(NamespaceId ns, BucketId bucket, RefId key, BlobId blobHash, CbObject payload);
+        Task<(ContentId[], BlobId[])> FinalizeAsync(NamespaceId ns, BucketId bucket, RefId key, BlobId blobHash);
 
-        IAsyncEnumerable<NamespaceId> GetNamespaces();
+        IAsyncEnumerable<NamespaceId> GetNamespacesAsync();
 
-        Task<bool> Delete(NamespaceId ns, BucketId bucket, RefId key);
-        Task<long> DropNamespace(NamespaceId ns);
-        Task<long> DeleteBucket(NamespaceId ns, BucketId bucket);
+        Task<bool> DeleteAsync(NamespaceId ns, BucketId bucket, RefId key);
+        Task<long> DropNamespaceAsync(NamespaceId ns);
+        Task<long> DeleteBucketAsync(NamespaceId ns, BucketId bucket);
 
-        Task<bool> Exists(NamespaceId ns, BucketId bucket, RefId key);
-        Task<List<BlobIdentifier>> GetReferencedBlobs(NamespaceId ns, BucketId bucket, RefId key);
+        Task<bool> ExistsAsync(NamespaceId ns, BucketId bucket, RefId key);
+        Task<List<BlobId>> GetReferencedBlobsAsync(NamespaceId ns, BucketId bucket, RefId key);
     }
 
 	public class RefRecord
 	{
-		public RefRecord(NamespaceId ns, BucketId bucket, RefId name, DateTime lastAccess, byte[]? inlinePayload, BlobIdentifier blobIdentifier, bool isFinalized)
+		public RefRecord(NamespaceId ns, BucketId bucket, RefId name, DateTime lastAccess, byte[]? inlinePayload, BlobId blobIdentifier, bool isFinalized)
 		{
 			Namespace = ns;
 			Bucket = bucket;
@@ -45,7 +45,7 @@ namespace Horde.Server.Ddc
 		public RefId Name { get; }
 		public DateTime LastAccess { get; }
 		public byte[]? InlinePayload { get; set; }
-		public BlobIdentifier BlobIdentifier { get; set; }
+		public BlobId BlobIdentifier { get; set; }
 		public bool IsFinalized { get; }
 	}
 
@@ -65,7 +65,7 @@ namespace Horde.Server.Ddc
 
 	public class ObjectHashMismatchException : Exception
     {
-        public ObjectHashMismatchException(NamespaceId ns, BucketId bucket, IoHash name, BlobIdentifier suppliedHash, BlobIdentifier actualHash) : base($"Object {name} in bucket {bucket} and namespace {ns} did not reference hash {suppliedHash} was referencing {actualHash}")
+        public ObjectHashMismatchException(NamespaceId ns, BucketId bucket, IoHash name, BlobId suppliedHash, BlobId actualHash) : base($"Object {name} in bucket {bucket} and namespace {ns} did not reference hash {suppliedHash} was referencing {actualHash}")
         {
         }
     }
diff --git a/Engine/Source/Programs/Horde/Horde.Server/Ddc/IReferenceResolver.cs b/Engine/Source/Programs/Horde/Horde.Server/Ddc/IReferenceResolver.cs
--- a/Engine/Source/Programs/Horde/Horde.Server/Ddc/IReferenceResolver.cs
+++ b/Engine/Source/Programs/Horde/Horde.Server/Ddc/IReferenceResolver.cs
@@ -75,7 +75,7 @@ namespace Horde.Server.Ddc
         /// <param name="ns">The namespace to check</param>
         /// <param name="cb">The compact binary object to resolve references for</param>
         /// <returns></returns>
-        IAsyncEnumerable<BlobId> GetReferencedBlobs(NamespaceId ns, CbObject cb);
+        IAsyncEnumerable<BlobId> GetReferencedBlobsAsync(NamespaceId ns, CbObject cb);
 
         /// <summary>
         /// Returns which attachments exist in the cb object or any children
@@ -84,7 +84,7 @@ namespace Horde.Server.Ddc
         /// <param name="ns">The namespace to check</param>
         /// <param name="cb">The compact binary object to resolve references for</param>
         /// <returns></returns>
-        IAsyncEnumerable<Attachment> GetAttachments(NamespaceId ns, CbObject cb);
+        IAsyncEnumerable<Attachment> GetAttachmentsAsync(NamespaceId ns, CbObject cb);
     }
 
     public class ReferenceResolver : IReferenceResolver
@@ -100,7 +100,7 @@ namespace Horde.Server.Ddc
             _tracer = tracer;
         }
 
-        public async IAsyncEnumerable<Attachment> GetAttachments(NamespaceId ns, CbObject cb)
+        public async IAsyncEnumerable<Attachment> GetAttachmentsAsync(NamespaceId ns, CbObject cb)
         {
             Queue<CbObject> objectsToVisit = new Queue<CbObject>();
             objectsToVisit.Enqueue(cb);
@@ -234,7 +234,7 @@ namespace Horde.Server.Ddc
             }
         }
 
-        public async IAsyncEnumerable<BlobId> GetReferencedBlobs(NamespaceId ns, CbObject cb)
+        public async IAsyncEnumerable<BlobId> GetReferencedBlobsAsync(NamespaceId ns, CbObject cb)
         {
             List<Task<(BlobId, bool)>> pendingBlobExistsChecks = new();
             List<Task<(ContentIdAttachment, bool)>> pendingContentIdChecks = new();
@@ -242,7 +242,7 @@ namespace Horde.Server.Ddc
             List<BlobId> unresolvedBlobReferences = new List<BlobId>();
 
             // Resolve all the attachments
-            await foreach (Attachment attachment in GetAttachments(ns, cb))
+            await foreach (Attachment attachment in GetAttachmentsAsync(ns, cb))
             {
                 if (attachment is BlobAttachment blobAttachment)
                 {
@@ -325,7 +325,7 @@ namespace Horde.Server.Ddc
 
         private async Task<CbObject> ParseCompactBinaryAttachment(NamespaceId ns, BlobId blobIdentifier)
         {
-            BlobContents contents = await _blobStore.GetObject(ns, blobIdentifier);
+            BlobContents contents = await _blobStore.GetObjectAsync(ns, blobIdentifier);
             byte[] data = await contents.Stream.ToByteArray();
             CbObject childBinaryObject = new CbObject(data);
 
@@ -337,13 +337,13 @@ namespace Horde.Server.Ddc
             using TelemetrySpan scope = _tracer.StartActiveSpan("ReferenceResolver.ResolveContentId")
                 .SetAttribute("operation.name", "ReferenceResolver.ResolveContentId")
                 .SetAttribute("resource.name", contentId.ToString());
-            BlobId[]? resolvedBlobs = await _contentIdStore.Resolve(ns, contentId);
+            BlobId[]? resolvedBlobs = await _contentIdStore.ResolveAsync(ns, contentId);
             return (contentId, resolvedBlobs);
         }
 
         private async Task<(BlobId, bool)> CheckBlobExists(NamespaceId ns, BlobId blob)
         {
-            return (blob, await _blobStore.Exists(ns, blob));
+            return (blob, await _blobStore.ExistsAsync(ns, blob));
         }
     }
 
diff --git a/Engine/Source/Programs/Horde/Horde.Server/Ddc/IRequestHelper.cs b/Engine/Source/Programs/Horde/Horde.Server/Ddc/IRequestHelper.cs
--- a/Engine/Source/Programs/Horde/Horde.Server/Ddc/IRequestHelper.cs
+++ b/Engine/Source/Programs/Horde/Horde.Server/Ddc/IRequestHelper.cs
@@ -14,8 +14,8 @@ namespace Horde.Server.Ddc
 {
 	public interface IRequestHelper
     {
-        public Task<ActionResult?> HasAccessToNamespace(ClaimsPrincipal user, HttpRequest request, NamespaceId ns, AclAction[] aclActions);
-        public Task<ActionResult?> HasAccessForGlobalOperations(ClaimsPrincipal user, AclAction[] aclActions);
+        public Task<ActionResult?> HasAccessToNamespaceAsync(ClaimsPrincipal user, HttpRequest request, NamespaceId ns, AclAction[] aclActions);
+        public Task<ActionResult?> HasAccessForGlobalOperationsAsync(ClaimsPrincipal user, AclAction[] aclActions);
     }
 
     public class AuthorizationException : Exception
diff --git a/Engine/Source/Programs/Horde/Horde.Server/Ddc/RequestHelper.cs b/Engine/Source/Programs/Horde/Horde.Server/Ddc/RequestHelper.cs
--- a/Engine/Source/Programs/Horde/Horde.Server/Ddc/RequestHelper.cs
+++ b/Engine/Source/Programs/Horde/Horde.Server/Ddc/RequestHelper.cs
@@ -12,12 +12,12 @@ namespace Horde.Server.Ddc
 {
 	class RequestHelper : IRequestHelper
 	{
-		public Task<ActionResult?> HasAccessForGlobalOperations(ClaimsPrincipal user, AclAction[] aclActions)
+		public Task<ActionResult?> HasAccessForGlobalOperationsAsync(ClaimsPrincipal user, AclAction[] aclActions)
 		{
 			throw new NotImplementedException();
 		}
 
-		public Task<ActionResult?> HasAccessToNamespace(ClaimsPrincipal user, HttpRequest request, NamespaceId ns, AclAction[] aclActions)
+		public Task<ActionResult?> HasAccessToNamespaceAsync(ClaimsPrincipal user, HttpRequest request, NamespaceId ns, AclAction[] aclActions)
 		{
 			throw new NotImplementedException();
 		}
