Horde: Rename StreamingBodyPayload -> BufferedPayload to match name of class, and add comments to suppress warnings.

diff --git a/Engine/Source/Programs/Horde/Horde.Server/Ddc/BufferedPayload.cs b/Engine/Source/Programs/Horde/Horde.Server/Ddc/BufferedPayload.cs
new file mode 100644
--- /dev/null
+++ b/Engine/Source/Programs/Horde/Horde.Server/Ddc/BufferedPayload.cs
@@ -0,0 +1,187 @@
+// Copyright Epic Games, Inc. All Rights Reserved.
+
+using System;
+using System.IO;
+using System.Threading.Tasks;
+using Microsoft.AspNetCore.Http;
+using Microsoft.Extensions.Options;
+using OpenTelemetry.Trace;
+
+namespace Horde.Server.Ddc
+{
+	/// <summary>
+	/// Base class for a payload stream that supports seeking
+	/// </summary>
+	public abstract class BufferedPayload : IDisposable
+    {
+		/// <summary>
+		/// Length of the payload
+		/// </summary>
+		public long Length { get; }
+
+		/// <summary>
+		/// Constructor
+		/// </summary>
+		protected BufferedPayload(long length)
+		{
+			Length = length;
+		}
+
+		/// <inheritdoc/>
+		public void Dispose()
+		{
+			Dispose(true);
+			GC.SuppressFinalize(this);
+		}
+
+		/// <summary>
+		/// Overridable implementation of <see cref="Dispose()"/>
+		/// </summary>
+		protected virtual void Dispose(bool disposing)
+		{
+		}
+
+		/// <summary>
+		/// Opens a stream to the payload
+		/// </summary>
+		public abstract Stream GetStream();
+    }
+
+    /// <summary>
+    /// Streaming request that is streamed into memory
+    /// </summary>
+    public sealed class MemoryBufferedPayload : BufferedPayload
+    {
+        private readonly byte[] _buffer;
+
+		/// <summary>
+		/// Constructor
+		/// </summary>
+        public MemoryBufferedPayload(byte[] source)
+			: base(source.LongLength)
+        {
+            _buffer = source;
+        }
+
+		/// <summary>
+		/// Create a buffered payload from a stream
+		/// </summary>
+		public static async Task<MemoryBufferedPayload> Create(Tracer tracer, Stream s)
+        {
+            using TelemetrySpan scope = tracer.StartActiveSpan("payload.buffer")
+                .SetAttribute("operation.name", "payload.buffer")
+                .SetAttribute("bufferType", "Memory");
+            MemoryBufferedPayload payload = new MemoryBufferedPayload(await s.ToByteArray());
+            return payload;
+        }
+
+		/// <inheritdoc/>
+		public override Stream GetStream() => new MemoryStream(_buffer);
+    }
+
+    /// <summary>
+    /// A streaming request backed by a temporary file on disk
+    /// </summary>
+    public sealed class FilesystemBufferedPayload : BufferedPayload
+    {
+        private readonly FileInfo _tempFile;
+
+		private FilesystemBufferedPayload(FileInfo tempFile)
+			: base(tempFile.Length)
+        {
+			_tempFile = tempFile;
+        }
+
+		/// <summary>
+		/// Create a new payload instance backed by the filesystem
+		/// </summary>
+        public static async Task<FilesystemBufferedPayload> Create(Tracer tracer, Stream s)
+        {
+			FileInfo tempFile = new FileInfo(Path.GetTempFileName());
+
+			{
+				using TelemetrySpan? scope = tracer.StartActiveSpan("payload.buffer")
+                    .SetAttribute("operation.name", "payload.buffer")
+                    .SetAttribute("bufferType", "Filesystem");
+                await using FileStream fs = tempFile.OpenWrite();
+                await s.CopyToAsync(fs);
+            }
+
+			tempFile.Refresh();
+
+			return new FilesystemBufferedPayload(tempFile);
+        }
+
+		/// <inheritdoc/>
+		protected override void Dispose(bool disposing)
+        {
+			base.Dispose(disposing);
+
+            if (_tempFile.Exists)
+            {
+                _tempFile.Delete();
+            }
+        }
+
+		/// <inheritdoc/>
+		public override Stream GetStream() => _tempFile.OpenRead();
+    }
+
+	/// <summary>
+	/// Options for creating <see cref="BufferedPayload"/> instances
+	/// </summary>
+    public class BufferedPayloadOptions
+    {
+		/// <summary>
+		/// If the request is smaller then MemoryBufferSize we buffer it in memory rather then as a file
+		/// </summary>
+		public long MemoryBufferSize { get; set; } = 128 * 1024 * 1024;
+    }
+
+	/// <summary>
+	/// Factory for creating <see cref="BufferedPayload"/> instances
+	/// </summary>
+    public class BufferedPayloadFactory
+    {
+        private readonly IOptionsMonitor<BufferedPayloadOptions> _options;
+        private readonly Tracer _tracer;
+
+		/// <summary>
+		/// Constructor
+		/// </summary>
+        public BufferedPayloadFactory(IOptionsMonitor<BufferedPayloadOptions> options, Tracer tracer)
+        {
+            _options = options;
+            _tracer = tracer;
+        }
+
+		/// <summary>
+		/// Create a new buffered payload from an HTTP request
+		/// </summary>
+        public Task<BufferedPayload> CreateFromRequest(HttpRequest request)
+        {
+            long? contentLength = request.ContentLength;
+
+            if (contentLength == null)
+            {
+                throw new Exception("Expected content-length on all requests");
+            }
+
+            return CreateFromStream(request.Body, contentLength.Value);
+        }
+
+		/// <summary>
+		/// Create a new buffered payload instance from a stream
+		/// </summary>
+        public async Task<BufferedPayload> CreateFromStream(Stream s, long contentLength)
+        {
+            // blob is small enough to fit into memory we just read it as is
+            if (contentLength < _options.CurrentValue.MemoryBufferSize)
+            {
+                return await MemoryBufferedPayload.Create(_tracer, s);
+            }
+
+            return await FilesystemBufferedPayload.Create(_tracer, s);
+        }
+    }
+}
diff --git a/Engine/Source/Programs/Horde/Horde.Server/Ddc/DdcBlobService.cs b/Engine/Source/Programs/Horde/Horde.Server/Ddc/DdcBlobService.cs
--- a/Engine/Source/Programs/Horde/Horde.Server/Ddc/DdcBlobService.cs
+++ b/Engine/Source/Programs/Horde/Horde.Server/Ddc/DdcBlobService.cs
@@ -65,7 +65,7 @@ namespace Horde.Server.Ddc
 			throw new NotImplementedException();
 		}
 
-		public Task<BlobId> PutObjectAsync(NamespaceId ns, IBufferedPayload payload, BlobId identifier)
+		public Task<BlobId> PutObjectAsync(NamespaceId ns, BufferedPayload payload, BlobId identifier)
 		{
 			throw new NotImplementedException();
 		}
@@ -75,7 +75,7 @@ namespace Horde.Server.Ddc
 			throw new NotImplementedException();
 		}
 
-		public Task<BlobId> PutObjectKnownHashAsync(NamespaceId ns, IBufferedPayload content, BlobId identifier)
+		public Task<BlobId> PutObjectKnownHashAsync(NamespaceId ns, BufferedPayload content, BlobId identifier)
 		{
 			throw new NotImplementedException();
 		}
diff --git a/Engine/Source/Programs/Horde/Horde.Server/Ddc/DdcBlobsController.cs b/Engine/Source/Programs/Horde/Horde.Server/Ddc/DdcBlobsController.cs
--- a/Engine/Source/Programs/Horde/Horde.Server/Ddc/DdcBlobsController.cs
+++ b/Engine/Source/Programs/Horde/Horde.Server/Ddc/DdcBlobsController.cs
@@ -187,7 +187,7 @@ namespace Horde.Server.Ddc
                         RedirectUri = uri,
                     });
                 }
-                using IBufferedPayload payload = await _bufferedPayloadFactory.CreateFromRequest(Request);
+                using BufferedPayload payload = await _bufferedPayloadFactory.CreateFromRequest(Request);
 
                 BlobId identifier = await _storage.PutObjectAsync(ns, payload, id);
                 return Ok(new
@@ -220,7 +220,7 @@ namespace Horde.Server.Ddc
             _diagnosticContext.Set("Content-Length", Request.ContentLength ?? -1);
             try
             {
-                using IBufferedPayload payload = await _bufferedPayloadFactory.CreateFromRequest(Request);
+                using BufferedPayload payload = await _bufferedPayloadFactory.CreateFromRequest(Request);
 
                 await using Stream stream = payload.GetStream();
 
diff --git a/Engine/Source/Programs/Horde/Horde.Server/Ddc/DdcObjectsController.cs b/Engine/Source/Programs/Horde/Horde.Server/Ddc/DdcObjectsController.cs
--- a/Engine/Source/Programs/Horde/Horde.Server/Ddc/DdcObjectsController.cs
+++ b/Engine/Source/Programs/Horde/Horde.Server/Ddc/DdcObjectsController.cs
@@ -158,7 +158,7 @@ namespace Horde.Server.Ddc
             _diagnosticContext.Set("Content-Length", Request.ContentLength ?? -1);
             try
             {
-                using IBufferedPayload payload = await _bufferedPayloadFactory.CreateFromRequest(Request);
+                using BufferedPayload payload = await _bufferedPayloadFactory.CreateFromRequest(Request);
 
                 BlobId identifier = await _storage.PutObjectAsync(ns, payload, id);
                 return Ok(new PutBlobResponse(identifier));
diff --git a/Engine/Source/Programs/Horde/Horde.Server/Ddc/DdcRefsController.cs b/Engine/Source/Programs/Horde/Horde.Server/Ddc/DdcRefsController.cs
--- a/Engine/Source/Programs/Horde/Horde.Server/Ddc/DdcRefsController.cs
+++ b/Engine/Source/Programs/Horde/Horde.Server/Ddc/DdcRefsController.cs
@@ -624,7 +624,7 @@ namespace Horde.Server.Ddc
 
             try
             {
-                using IBufferedPayload payload = await _bufferedPayloadFactory.CreateFromRequest(Request);
+                using BufferedPayload payload = await _bufferedPayloadFactory.CreateFromRequest(Request);
 
                 BlobId headerHash;
                 if (Request.Headers.ContainsKey(CommonHeaders.HashHeaderName))
diff --git a/Engine/Source/Programs/Horde/Horde.Server/Ddc/IDdcBlobService.cs b/Engine/Source/Programs/Horde/Horde.Server/Ddc/IDdcBlobService.cs
--- a/Engine/Source/Programs/Horde/Horde.Server/Ddc/IDdcBlobService.cs
+++ b/Engine/Source/Programs/Horde/Horde.Server/Ddc/IDdcBlobService.cs
@@ -20,8 +20,8 @@ namespace Horde.Server.Ddc
 	public interface IDdcBlobService
 	{
 		Task<JupiterContentHash> VerifyContentMatchesHashAsync(Stream content, JupiterContentHash identifier);
-		Task<BlobId> PutObjectKnownHashAsync(NamespaceId ns, IBufferedPayload content, BlobId identifier);
-		Task<BlobId> PutObjectAsync(NamespaceId ns, IBufferedPayload payload, BlobId identifier);
+		Task<BlobId> PutObjectKnownHashAsync(NamespaceId ns, BufferedPayload content, BlobId identifier);
+		Task<BlobId> PutObjectAsync(NamespaceId ns, BufferedPayload payload, BlobId identifier);
 		Task<BlobId> PutObjectAsync(NamespaceId ns, byte[] payload, BlobId identifier);
 		Task<Uri?> MaybePutObjectWithRedirectAsync(NamespaceId ns, BlobId identifier);
 
@@ -100,7 +100,7 @@ namespace Horde.Server.Ddc
 
 	public static class BlobServiceExtensions
 	{
-		public static async Task<ContentId> PutCompressedObject(this IDdcBlobService blobService, NamespaceId ns, IBufferedPayload payload, ContentId? id, IServiceProvider provider)
+		public static async Task<ContentId> PutCompressedObject(this IDdcBlobService blobService, NamespaceId ns, BufferedPayload payload, ContentId? id, IServiceProvider provider)
 		{
 			IDdcContentIdService contentIdStore = provider.GetService<IDdcContentIdService>()!;
 			CompressedBufferUtils compressedBufferUtils = provider.GetService<CompressedBufferUtils>()!;
diff --git a/Engine/Source/Programs/Horde/Horde.Server/Ddc/StreamingBodyPayload.cs b/Engine/Source/Programs/Horde/Horde.Server/Ddc/StreamingBodyPayload.cs
deleted file mode 100644
