Horde: Remove JupiterContentHash, and use struct types that directly wrap IoHash for BlobId and ContentId.

diff --git a/Engine/Source/Programs/Horde/Horde.Server/Ddc/BlobId.cs b/Engine/Source/Programs/Horde/Horde.Server/Ddc/BlobId.cs
--- a/Engine/Source/Programs/Horde/Horde.Server/Ddc/BlobId.cs
+++ b/Engine/Source/Programs/Horde/Horde.Server/Ddc/BlobId.cs
@@ -4,232 +4,139 @@ using System;
 using System.ComponentModel;
 using System.Globalization;
 using System.IO;
-using System.Text;
 using System.Text.Json;
 using System.Text.Json.Serialization;
-using System.Threading.Tasks;
-using Blake3;
 using EpicGames.Core;
-using EpicGames.Horde.Storage;
-using EpicGames.Horde.Storage.Bundles;
 using EpicGames.Serialization;
 
-#pragma warning disable CS1591
-
 namespace Horde.Server.Ddc
 {
+	/// <summary>
+	/// Identifier for a blob in the store
+	/// </summary>
 	[TypeConverter(typeof(BlobIdTypeConverter))]
-    [JsonConverter(typeof(BlobIdJsonConverter))]
-    [CbConverter(typeof(BlobIdCbConverter))]
-    public class BlobId : JupiterContentHash,  IEquatable<BlobId>
-    {
-        // multi thread the hashing for blobs larger then this size
-        private const int MultiThreadedSize = 1_000_000;
-        private string? _stringIdentifier;
-
-        public BlobId(byte[] identifier) : base(identifier)
-        {
-        }
-
-        [JsonConstructor]
-        public BlobId(string identifier) : base(identifier)
-        {
-
-        }
-
-        public override int GetHashCode()
-        {
-			HashCode hashCode = new HashCode();
-			hashCode.AddBytes(Identifier);
-			return hashCode.ToHashCode();
-        }
-
-		public bool Equals(BlobId? other) => other != null && other.Identifier.AsSpan().SequenceEqual(other.Identifier);
-
-        public override bool Equals(object? obj)
-        {
-            if (ReferenceEquals(null, obj))
-            {
-                return false;
-            }
-
-            if (ReferenceEquals(this, obj))
-            {
-                return true;
-            }
-
-            if (obj.GetType() != GetType())
-            {
-                return false;
-            }
-
-            return Equals((BlobId) obj);
-        }
-
-        public override string ToString()
-        {
-            if (_stringIdentifier == null)
-            {
-                _stringIdentifier = StringUtils.FormatHexString(Identifier);
-            }
-
-            return _stringIdentifier;
-        }
-
-        public static new BlobId FromBlob(byte[] blobMemory)
-        {
-            Hash blake3Hash;
-            if (blobMemory.Length < MultiThreadedSize)
-            {
-                using Hasher hasher = Hasher.New();
-                hasher.Update(blobMemory);
-                blake3Hash = hasher.Finalize();
-            }
-            else
-            {
-                using Hasher hasher = Hasher.New();
-                hasher.UpdateWithJoin(blobMemory);
-                blake3Hash = hasher.Finalize();
-            }
-            
-            // we only keep the first 20 bytes of the Blake3 hash
-            Span<byte> hash = blake3Hash.AsSpanUnsafe().Slice(0, 20);
-            return new BlobId(hash.ToArray());
-        }
-
-        public static BlobId FromBlob(in Memory<byte> blobMemory)
-        {
-            Hash blake3Hash;
-            if (blobMemory.Length < MultiThreadedSize)
-            {
-                using Hasher hasher = Hasher.New();
-                hasher.Update(blobMemory.Span);
-                blake3Hash = hasher.Finalize();
-            }
-            else
-            {
-                using Hasher hasher = Hasher.New();
-                hasher.UpdateWithJoin(blobMemory.Span);
-                blake3Hash = hasher.Finalize();
-            }
-
-            // we only keep the first 20 bytes of the Blake3 hash
-            Span<byte> hash = blake3Hash.AsSpanUnsafe().Slice(0, 20);
-            return new BlobId(hash.ToArray());
-        }
-
-        public static BlobId FromContentHash(JupiterContentHash testObjectHash)
-        {
-            return new BlobId(testObjectHash.HashData);
-        }
-
-        public static async Task<BlobId> FromStream(Stream stream)
-        {
-            using Hasher hasher = Hasher.New();
-            const int bufferSize = 1024 * 1024 * 5;
-            byte[] buffer = new byte[bufferSize];
-            int read = await stream.ReadAsync(buffer, 0, buffer.Length);
-            while (read > 0)
-            {
-                hasher.UpdateWithJoin(new ReadOnlySpan<byte>(buffer, 0, read));
-                read = await stream.ReadAsync(buffer, 0, buffer.Length);
-            }
-            Hash blake3Hash = hasher.Finalize();
-
-            // we only keep the first 20 bytes of the Blake3 hash
-            byte[] hash = blake3Hash.AsSpanUnsafe().Slice(0, 20).ToArray();
-            return new BlobId(hash);
-        }
-
-        public static BlobId FromIoHash(IoHash blobIdentifier)
-        {
-            return new BlobId(blobIdentifier.ToByteArray());
-        }
-
-        public IoHash AsIoHash()
-        {
-            return new IoHash(HashData);
-        }
-
-        public static BlobId FromBlobLocator(BundleLocator locator)
-        {
-            return new BlobId(Encoding.UTF8.GetBytes(locator.ToString()));
-        }
-
-        public BundleLocator AsBlobLocator()
-        {
-            return new BundleLocator(Encoding.UTF8.GetString(HashData));
-        }
-    }
-
-    public class BlobIdTypeConverter : TypeConverter
-    {
-        public override bool CanConvertFrom(ITypeDescriptorContext? context, Type sourceType)
-        {  
-            if (sourceType == typeof(string))  
-            {  
-                return true;
-            }  
-            return base.CanConvertFrom(context, sourceType);
-        }  
-  
-        public override object? ConvertFrom(ITypeDescriptorContext? context, CultureInfo? culture, object value)  
-        {
-            if (value is string s)
-            {
-                return new BlobId(s);
-            }
-
-            return base.ConvertFrom(context, culture, value);  
-        }
-
-        public override bool CanConvertTo(ITypeDescriptorContext? context, Type? destinationType)
-        {
-            if (destinationType == typeof(string))
-            {
-                return true;
-            }
-            return base.CanConvertTo(context, destinationType);
-        }
-
-        public override object? ConvertTo(ITypeDescriptorContext? context, CultureInfo? culture, object? value, Type destinationType)
-        {
-            if (destinationType == typeof(string))
-            {
-                BlobId? identifier = (BlobId?)value;
-                return identifier?.ToString();
-            }
-            return base.ConvertTo(context, culture, value, destinationType);
-        }
-    }
-
-    public class BlobIdJsonConverter : JsonConverter<BlobId>
-    {
-        public override BlobId? Read(ref Utf8JsonReader reader, Type typeToConvert, JsonSerializerOptions options)
-        {
-            string? str = reader.GetString();
-            if (str == null)
-            {
-                throw new InvalidDataException("Unable to parse blob identifier");
-            }
-
-            return new BlobId(str);
-        }
-
-        public override void Write(Utf8JsonWriter writer, BlobId value, JsonSerializerOptions options)
-        {
-            writer.WriteStringValue(value.ToString());
-        }
-    }
-
-    public class BlobIdCbConverter : CbConverterBase<BlobId>
-    {
-        public override BlobId Read(CbField field) => new BlobId(field.AsHash().ToByteArray());
-
-        /// <inheritdoc/>
-        public override void Write(CbWriter writer, BlobId value) => writer.WriteBinaryAttachmentValue(new IoHash(value.HashData));
-
-        /// <inheritdoc/>
-        public override void WriteNamed(CbWriter writer, Utf8String name, BlobId value) => writer.WriteBinaryAttachment(name, new IoHash(value.HashData));
-    }
+	[JsonConverter(typeof(BlobIdJsonConverter))]
+	[CbConverter(typeof(BlobIdCbConverter))]
+	public readonly struct BlobId : IEquatable<BlobId>
+	{
+		/// <summary>
+		/// Hash of the blob
+		/// </summary>
+		public IoHash Hash { get; }
+
+		/// <summary>
+		/// Constructor
+		/// </summary>
+		public BlobId(IoHash hash) => Hash = hash;
+
+		/// <summary>
+		/// Parses a BlobId from a string
+		/// </summary>
+		public static BlobId Parse(string text) => new BlobId(IoHash.Parse(text));
+
+		/// <inheritdoc/>
+		public override readonly int GetHashCode() => Hash.GetHashCode();
+
+		/// <inheritdoc/>
+		public readonly bool Equals(BlobId other) => Hash.Equals(other.Hash);
+
+		/// <inheritdoc/>
+		public override readonly bool Equals(object? obj) => obj is BlobId blobId && Equals(blobId);
+
+		/// <inheritdoc/>
+		public override readonly string ToString() => Hash.ToString();
+
+		/// <inheritdoc cref="IoHash.op_Equality"/>
+		public static bool operator ==(BlobId left, BlobId right) => left.Hash == right.Hash;
+
+		/// <inheritdoc cref="IoHash.op_Inequality"/>
+		public static bool operator !=(BlobId left, BlobId right) => !(left == right);
+	}
+
+	/// <summary>
+	/// Converts from <see cref="BlobId"/> instances to other types
+	/// </summary>
+	public class BlobIdTypeConverter : TypeConverter
+	{
+		/// <inheritdoc/>
+		public override bool CanConvertFrom(ITypeDescriptorContext? context, Type sourceType)
+		{
+			if (sourceType == typeof(string))
+			{
+				return true;
+			}
+			return base.CanConvertFrom(context, sourceType);
+		}
+
+		/// <inheritdoc/>
+		public override object? ConvertFrom(ITypeDescriptorContext? context, CultureInfo? culture, object value)
+		{
+			if (value is string s)
+			{
+				return BlobId.Parse(s);
+			}
+
+			return base.ConvertFrom(context, culture, value);
+		}
+
+		/// <inheritdoc/>
+		public override bool CanConvertTo(ITypeDescriptorContext? context, Type? destinationType)
+		{
+			if (destinationType == typeof(string))
+			{
+				return true;
+			}
+			return base.CanConvertTo(context, destinationType);
+		}
+
+		/// <inheritdoc/>
+		public override object? ConvertTo(ITypeDescriptorContext? context, CultureInfo? culture, object? value, Type destinationType)
+		{
+			if (destinationType == typeof(string))
+			{
+				return value?.ToString();
+			}
+
+			return base.ConvertTo(context, culture, value, destinationType);
+		}
+	}
+
+	/// <summary>
+	/// Converts from <see cref="BlobId"/> instances to other types
+	/// </summary>
+	public class BlobIdJsonConverter : JsonConverter<BlobId>
+	{
+		/// <inheritdoc/>
+		public override BlobId Read(ref Utf8JsonReader reader, Type typeToConvert, JsonSerializerOptions options)
+		{
+			string? str = reader.GetString();
+			if (str == null)
+			{
+				throw new InvalidDataException("Unable to parse blob identifier");
+			}
+
+			return BlobId.Parse(str);
+		}
+
+		/// <inheritdoc/>
+		public override void Write(Utf8JsonWriter writer, BlobId value, JsonSerializerOptions options)
+		{
+			writer.WriteStringValue(value.ToString());
+		}
+	}
+
+	/// <summary>
+	/// Serializes <see cref="BlobId"/> instances to compact binary
+	/// </summary>
+	public class BlobIdCbConverter : CbConverterBase<BlobId>
+	{
+		/// <inheritdoc/>
+		public override BlobId Read(CbField field) => new BlobId(field.AsHash());
+
+		/// <inheritdoc/>
+		public override void Write(CbWriter writer, BlobId value) => writer.WriteBinaryAttachmentValue(value.Hash);
+
+		/// <inheritdoc/>
+		public override void WriteNamed(CbWriter writer, Utf8String name, BlobId value) => writer.WriteBinaryAttachment(name, value.Hash);
+	}
 }
diff --git a/Engine/Source/Programs/Horde/Horde.Server/Ddc/CompressedBufferUtils.cs b/Engine/Source/Programs/Horde/Horde.Server/Ddc/CompressedBufferUtils.cs
--- a/Engine/Source/Programs/Horde/Horde.Server/Ddc/CompressedBufferUtils.cs
+++ b/Engine/Source/Programs/Horde/Horde.Server/Ddc/CompressedBufferUtils.cs
@@ -231,16 +231,13 @@ namespace Horde.Server.Ddc
             {
                 using TelemetrySpan _ = _tracer.StartActiveSpan("web.hash").SetAttribute("operation.name", "web.hash");
 
-                // only read the first 20 bytes of the hash field as IoHashes are 20 bytes and not 32 bytes
-                byte[] slicedHash = new byte[20];
-                Array.Copy(header.RawHash, 0, slicedHash, 0, 20);
+				// only read the first 20 bytes of the hash field as IoHashes are 20 bytes and not 32 bytes
+				IoHash slicedHash = new IoHash(header.RawHash.AsSpan(0, 20));
+                IoHash contentHash = IoHash.Compute(decompressedPayload);
 
-                BlobId headerIdentifier = new BlobId(slicedHash);
-                BlobId contentHash = BlobId.FromBlob(decompressedPayload);
-
-                if (!headerIdentifier.Equals(contentHash))
+                if (!slicedHash.Equals(contentHash))
                 {
-                    throw new Exception($"Payload was expected to be {headerIdentifier} but was {contentHash}");
+                    throw new Exception($"Payload was expected to be {slicedHash} but was {contentHash}");
                 }
             }
 
diff --git a/Engine/Source/Programs/Horde/Horde.Server/Ddc/ContentId.cs b/Engine/Source/Programs/Horde/Horde.Server/Ddc/ContentId.cs
--- a/Engine/Source/Programs/Horde/Horde.Server/Ddc/ContentId.cs
+++ b/Engine/Source/Programs/Horde/Horde.Server/Ddc/ContentId.cs
@@ -3,100 +3,142 @@
 using System;
 using System.ComponentModel;
 using System.Globalization;
+using System.IO;
+using System.Text.Json;
 using System.Text.Json.Serialization;
 using EpicGames.Core;
+using EpicGames.Serialization;
 
 #pragma warning disable CS1591
 
 namespace Horde.Server.Ddc
 {
+	/// <summary>
+	/// Identifier for a compressed buffer in the store
+	/// </summary>
 	[TypeConverter(typeof(ContentIdTypeConverter))]
-    public class ContentId : JupiterContentHash, IEquatable<ContentId>
-    {
-        public ContentId(byte[] identifier) : base(identifier)
-        {
-        }
-
-        [JsonConstructor]
-        public ContentId(string identifier) : base(identifier)
-        {
-
-        }
-
-        public override int GetHashCode()
-        {
-			HashCode hashCode = new HashCode();
-			hashCode.AddBytes(Identifier);
-			return hashCode.ToHashCode();
-        }
-
-        public bool Equals(ContentId? other) => other != null && other.Identifier.AsSpan().SequenceEqual(Identifier);
-
-        public override bool Equals(object? obj)
-        {
-            if (ReferenceEquals(null, obj))
-            {
-                return false;
-            }
-
-            if (ReferenceEquals(this, obj))
-            {
-                return true;
-            }
-
-            if (obj.GetType() != GetType())
-            {
-                return false;
-            }
-
-            return Equals((ContentId) obj);
-        }
-
-        public static ContentId FromContentHash(JupiterContentHash contentHash)
-        {
-            return new ContentId(contentHash.HashData);
-        }
-
-        public static ContentId FromBlobIdentifier(BlobId blobIdentifier)
-        {
-            return new ContentId(blobIdentifier.HashData);
-        }
-
-        public BlobId AsBlobIdentifier()
-        {
-            return new BlobId(HashData);
-        }
-
-        public static ContentId FromIoHash(IoHash ioHash)
-        {
-            return new ContentId(ioHash.ToByteArray());
-        }
-
-        public IoHash AsIoHash()
-        {
-            return new IoHash(HashData);
-        }
-    }
-
-    public class ContentIdTypeConverter : TypeConverter
-    {
-        public override bool CanConvertFrom(ITypeDescriptorContext? context, Type sourceType)
-        {  
-            if (sourceType == typeof(string))  
-            {  
-                return true;
-            }  
-            return base.CanConvertFrom(context, sourceType);
-        }  
-  
-        public override object? ConvertFrom(ITypeDescriptorContext? context, CultureInfo? culture, object value)  
-        {
-            if (value is string s)
-            {
-                return new ContentId(s);
-            }
-
-            return base.ConvertFrom(context, culture, value);  
-        }  
-    }
+	[JsonConverter(typeof(ContentIdJsonConverter))]
+	[CbConverter(typeof(ContentIdCbConverter))]
+	public readonly struct ContentId : IEquatable<ContentId>
+	{
+		/// <summary>
+		/// Hash of the blob
+		/// </summary>
+		public IoHash Hash { get; }
+
+		/// <summary>
+		/// Constructor
+		/// </summary>
+		public ContentId(IoHash hash) => Hash = hash;
+
+		/// <summary>
+		/// Parses a <see cref="ContentId"/> from a string
+		/// </summary>
+		public static ContentId Parse(string text) => new ContentId(IoHash.Parse(text));
+
+		/// <inheritdoc/>
+		public override readonly int GetHashCode() => Hash.GetHashCode();
+
+		/// <inheritdoc/>
+		public readonly bool Equals(ContentId other) => Hash.Equals(other.Hash);
+
+		/// <inheritdoc/>
+		public override readonly bool Equals(object? obj) => obj is ContentId contentId && Equals(contentId);
+
+		/// <inheritdoc/>
+		public override readonly string ToString() => Hash.ToString();
+
+		/// <inheritdoc cref="IoHash.op_Equality"/>
+		public static bool operator ==(ContentId left, ContentId right) => left.Hash == right.Hash;
+
+		/// <inheritdoc cref="IoHash.op_Inequality"/>
+		public static bool operator !=(ContentId left, ContentId right) => !(left == right);
+	}
+
+	/// <summary>
+	/// Converts from <see cref="ContentId"/> instances to other types
+	/// </summary>
+	public class ContentIdTypeConverter : TypeConverter
+	{
+		/// <inheritdoc/>
+		public override bool CanConvertFrom(ITypeDescriptorContext? context, Type sourceType)
+		{
+			if (sourceType == typeof(string))
+			{
+				return true;
+			}
+			return base.CanConvertFrom(context, sourceType);
+		}
+
+		/// <inheritdoc/>
+		public override object? ConvertFrom(ITypeDescriptorContext? context, CultureInfo? culture, object value)
+		{
+			if (value is string s)
+			{
+				return ContentId.Parse(s);
+			}
+
+			return base.ConvertFrom(context, culture, value);
+		}
+
+		/// <inheritdoc/>
+		public override bool CanConvertTo(ITypeDescriptorContext? context, Type? destinationType)
+		{
+			if (destinationType == typeof(string))
+			{
+				return true;
+			}
+			return base.CanConvertTo(context, destinationType);
+		}
+
+		/// <inheritdoc/>
+		public override object? ConvertTo(ITypeDescriptorContext? context, CultureInfo? culture, object? value, Type destinationType)
+		{
+			if (destinationType == typeof(string))
+			{
+				return value?.ToString();
+			}
+
+			return base.ConvertTo(context, culture, value, destinationType);
+		}
+	}
+
+	/// <summary>
+	/// Converts from <see cref="ContentId"/> instances to other types
+	/// </summary>
+	public class ContentIdJsonConverter : JsonConverter<ContentId>
+	{
+		/// <inheritdoc/>
+		public override ContentId Read(ref Utf8JsonReader reader, Type typeToConvert, JsonSerializerOptions options)
+		{
+			string? str = reader.GetString();
+			if (str == null)
+			{
+				throw new InvalidDataException("Unable to parse content id");
+			}
+
+			return ContentId.Parse(str);
+		}
+
+		/// <inheritdoc/>
+		public override void Write(Utf8JsonWriter writer, ContentId value, JsonSerializerOptions options)
+		{
+			writer.WriteStringValue(value.ToString());
+		}
+	}
+
+	/// <summary>
+	/// Serializes <see cref="ContentId"/> instances to compact binary
+	/// </summary>
+	public class ContentIdCbConverter : CbConverterBase<ContentId>
+	{
+		/// <inheritdoc/>
+		public override ContentId Read(CbField field) => new ContentId(field.AsHash());
+
+		/// <inheritdoc/>
+		public override void Write(CbWriter writer, ContentId value) => writer.WriteBinaryAttachmentValue(value.Hash);
+
+		/// <inheritdoc/>
+		public override void WriteNamed(CbWriter writer, Utf8String name, ContentId value) => writer.WriteBinaryAttachment(name, value.Hash);
+	}
 }
diff --git a/Engine/Source/Programs/Horde/Horde.Server/Ddc/DdcBlobService.cs b/Engine/Source/Programs/Horde/Horde.Server/Ddc/DdcBlobService.cs
--- a/Engine/Source/Programs/Horde/Horde.Server/Ddc/DdcBlobService.cs
+++ b/Engine/Source/Programs/Horde/Horde.Server/Ddc/DdcBlobService.cs
@@ -4,6 +4,7 @@ using System;
 using System.Collections.Generic;
 using System.IO;
 using System.Threading.Tasks;
+using EpicGames.Core;
 using EpicGames.Horde.Storage;
 
 namespace Horde.Server.Ddc
@@ -90,7 +91,7 @@ namespace Horde.Server.Ddc
 			throw new NotImplementedException();
 		}
 
-		public Task<JupiterContentHash> VerifyContentMatchesHashAsync(Stream content, JupiterContentHash identifier)
+		public Task<IoHash> VerifyContentMatchesHashAsync(Stream content, IoHash identifier)
 		{
 			throw new NotImplementedException();
 		}
diff --git a/Engine/Source/Programs/Horde/Horde.Server/Ddc/DdcBlobsController.cs b/Engine/Source/Programs/Horde/Horde.Server/Ddc/DdcBlobsController.cs
--- a/Engine/Source/Programs/Horde/Horde.Server/Ddc/DdcBlobsController.cs
+++ b/Engine/Source/Programs/Horde/Horde.Server/Ddc/DdcBlobsController.cs
@@ -11,6 +11,7 @@ using System.Net.Mime;
 using System.Text.Json.Serialization;
 using System.Threading.Tasks;
 using EpicGames.AspNet;
+using EpicGames.Core;
 using EpicGames.Horde.Storage;
 using EpicGames.Serialization;
 using Horde.Server.Acls;
@@ -224,7 +225,7 @@ namespace Horde.Server.Ddc
 
                 await using Stream stream = payload.GetStream();
 
-                BlobId id = await BlobId.FromStream(stream);
+				BlobId id = new BlobId(await IoHash.ComputeAsync(stream));
                 await _storage.PutObjectKnownHashAsync(ns, payload, id);
                 
                 return Ok(new
@@ -356,7 +357,7 @@ namespace Horde.Server.Ddc
                             return BadRequest();
                         }
 
-                        tasks[index] = GetImpl(op.Namespace.Value, op.Id).ContinueWith((t, _) =>
+                        tasks[index] = GetImpl(op.Namespace.Value, op.Id.Value).ContinueWith((t, _) =>
                         {
                             // TODO: This is very allocation heavy but given that the end result is a json object we can not really stream this anyway
                             using BlobContents blobContents = t.Result;
@@ -374,7 +375,7 @@ namespace Horde.Server.Ddc
                             return BadRequest();
                         }
 
-                        tasks[index] = _storage.ExistsAsync(op.Namespace.Value, op.Id)
+                        tasks[index] = _storage.ExistsAsync(op.Namespace.Value, op.Id.Value)
                             .ContinueWith((t,_) => t.Result ? (object?) null : op.Id, null, TaskScheduler.Current);
                         break;
                     case BatchOp.Operation.PUT:
@@ -390,7 +391,7 @@ namespace Horde.Server.Ddc
                             }
 
                             using MemoryBufferedPayload payload = new MemoryBufferedPayload(op.Content);
-                            tasks[index] = _storage.PutObjectAsync(op.Namespace.Value, payload, op.Id).ContinueWith((t, _) => (object?) t.Result, null, TaskScheduler.Current);
+                            tasks[index] = _storage.PutObjectAsync(op.Namespace.Value, payload, op.Id.Value).ContinueWith((t, _) => (object?) t.Result, null, TaskScheduler.Current);
                             break;
                         }
                     case BatchOp.Operation.DELETE:
@@ -399,7 +400,7 @@ namespace Horde.Server.Ddc
                             return BadRequest();
                         }
 
-                        tasks[index] = DeleteImpl(op.Namespace.Value, op.Id).ContinueWith((t, _) => (object?) null, null, TaskScheduler.Current);
+                        tasks[index] = DeleteImpl(op.Namespace.Value, op.Id.Value).ContinueWith((t, _) => (object?) null, null, TaskScheduler.Current);
                         break;
                     default:
                         throw new NotImplementedException($"{op.Op} is not a support op type");
diff --git a/Engine/Source/Programs/Horde/Horde.Server/Ddc/DdcObjectsController.cs b/Engine/Source/Programs/Horde/Horde.Server/Ddc/DdcObjectsController.cs
--- a/Engine/Source/Programs/Horde/Horde.Server/Ddc/DdcObjectsController.cs
+++ b/Engine/Source/Programs/Horde/Horde.Server/Ddc/DdcObjectsController.cs
@@ -260,7 +260,6 @@ namespace Horde.Server.Ddc
     {
         public PutBlobResponse()
         {
-            Identifier = null!;
         }
 
         public PutBlobResponse(BlobId identifier)
diff --git a/Engine/Source/Programs/Horde/Horde.Server/Ddc/DdcRefsController.cs b/Engine/Source/Programs/Horde/Horde.Server/Ddc/DdcRefsController.cs
--- a/Engine/Source/Programs/Horde/Horde.Server/Ddc/DdcRefsController.cs
+++ b/Engine/Source/Programs/Horde/Horde.Server/Ddc/DdcRefsController.cs
@@ -198,7 +198,7 @@ namespace Horde.Server.Ddc
 
                             IoHash hash = binaryAttachmentField.AsBinaryAttachment();
 
-                            BlobContents referencedBlobContents = await _blobStore.GetObjectAsync(ns, BlobId.FromIoHash(hash));
+                            BlobContents referencedBlobContents = await _blobStore.GetObjectAsync(ns, new BlobId(hash));
 
                             if (_nginxRedirectHelper.CanRedirect(Request, referencedBlobContents))
                             {
@@ -236,7 +236,7 @@ namespace Horde.Server.Ddc
                         IAsyncEnumerable<Attachment> attachments = _referenceResolver.GetAttachmentsAsync(ns, cb);
 
                         using CbPackageBuilder writer = new CbPackageBuilder();
-                        writer.AddAttachment(objectRecord.BlobIdentifier.AsIoHash(), CbPackageAttachmentFlags.IsObject, blobMemory);
+                        writer.AddAttachment(objectRecord.BlobIdentifier.Hash, CbPackageAttachmentFlags.IsObject, blobMemory);
 
                         await Parallel.ForEachAsync(attachments, async (attachment, token) =>
                         {
@@ -629,7 +629,7 @@ namespace Horde.Server.Ddc
                 BlobId headerHash;
                 if (Request.Headers.ContainsKey(CommonHeaders.HashHeaderName))
                 {
-                    headerHash = new BlobId(Request.Headers[CommonHeaders.HashHeaderName]);
+                    headerHash = BlobId.Parse(Request.Headers[CommonHeaders.HashHeaderName]);
                 }
                 else
                 {
@@ -652,12 +652,12 @@ namespace Horde.Server.Ddc
                         // TODO: convert the json object into a compact binary instead
                         CbWriter writer = new CbWriter();
                         writer.BeginObject();
-                        writer.WriteBinaryAttachmentValue(blobHeader.AsIoHash());
+						writer.WriteBinaryAttachmentValue(blobHeader.Hash);
                         writer.EndObject();
 
                         byte[] blob = writer.ToByteArray();
                         payloadObject = new CbObject(blob);
-                        blobHeader = BlobId.FromBlob(blob);
+						blobHeader = new BlobId(IoHash.Compute(blob));
                         break;
                     }
                     case CustomMediaTypeNames.UnrealCompactBinary:
@@ -674,13 +674,13 @@ namespace Horde.Server.Ddc
 
                         CbWriter writer = new CbWriter();
                         writer.BeginObject();
-                        writer.WriteBinaryAttachment("RawHash", blobHeader.AsIoHash());
+						writer.WriteBinaryAttachment("RawHash", blobHeader.Hash);
                         writer.WriteInteger("RawSize", payload.Length);
                         writer.EndObject();
 
                         byte[] blob = writer.ToByteArray();
                         payloadObject = new CbObject(blob);
-                        blobHeader = BlobId.FromBlob(blob);
+						blobHeader = new BlobId(IoHash.Compute(blob));
                         break;
                     }
                     default:
@@ -701,8 +701,9 @@ namespace Horde.Server.Ddc
 
             (ContentId[] missingReferences, BlobId[] missingBlobs) = await _refService.PutAsync(ns, bucket, key, blobHeader, payloadObject);
 
-            List<JupiterContentHash> missingHashes = new List<JupiterContentHash>(missingReferences);
-            missingHashes.AddRange(missingBlobs);
+			List<IoHash> missingHashes = new List<IoHash>();
+			missingHashes.AddRange(missingReferences.Select(x => x.Hash));
+            missingHashes.AddRange(missingBlobs.Select(x => x.Hash));
             return Ok(new PutObjectResponse(missingHashes.ToArray()));
         }
 
@@ -739,11 +740,11 @@ namespace Horde.Server.Ddc
                     if (entry.Flags.HasFlag(CbPackageAttachmentFlags.IsCompressed))
                     {
                         using MemoryBufferedPayload payload = new MemoryBufferedPayload(blob);
-                        await _blobStore.PutCompressedObject(ns, payload, ContentId.FromIoHash(entry.AttachmentHash), HttpContext.RequestServices);
+                        await _blobStore.PutCompressedObject(ns, payload, new ContentId(entry.AttachmentHash), HttpContext.RequestServices);
                     }
                     else
                     {
-                        await _blobStore.PutObjectAsync(ns, blob, BlobId.FromIoHash(entry.AttachmentHash));
+                        await _blobStore.PutObjectAsync(ns, blob, new BlobId(entry.AttachmentHash));
                     }
                 }
             }
@@ -756,13 +757,10 @@ namespace Horde.Server.Ddc
             }
             
             CbObject rootObject = packageReader.RootObject;
-            BlobId rootObjectHash = BlobId.FromIoHash(packageReader.RootHash);
+            BlobId rootObjectHash = new BlobId(packageReader.RootHash);
 
             (ContentId[] missingReferences, BlobId[] missingBlobs) = await _refService.PutAsync(ns, bucket, key, rootObjectHash, rootObject);
-
-            List<JupiterContentHash> missingHashes = new List<JupiterContentHash>(missingReferences);
-            missingHashes.AddRange(missingBlobs);
-            return Ok(new PutObjectResponse(missingHashes.ToArray()));
+			return Ok(new PutObjectResponse(missingReferences, missingBlobs));
         }
 
 		static async Task<byte[]> ReadRawBodyAsync(HttpRequest request)
@@ -807,10 +805,7 @@ namespace Horde.Server.Ddc
             try
             {
                 (ContentId[] missingReferences, BlobId[] missingBlobs) = await _refService.FinalizeAsync(ns, bucket, key, hash);
-                List<JupiterContentHash> missingHashes = new List<JupiterContentHash>(missingReferences);
-                missingHashes.AddRange(missingBlobs);
-
-                return Ok(new PutObjectResponse(missingHashes.ToArray()));
+				return Ok(new PutObjectResponse(missingReferences, missingBlobs));
             }
             catch (ObjectHashMismatchException e)
             {
@@ -951,18 +946,17 @@ namespace Horde.Server.Ddc
                     {
                         throw new Exception($"Missing payload hash for operation: {op.OpId}");
                     }
-                    BlobId headerHash = BlobId.FromContentHash(op.PayloadHash);
-                    BlobId objectHash = BlobId.FromBlob(op.Payload.GetView().ToArray());
 
+                    IoHash headerHash = op.PayloadHash.Value;
+					IoHash objectHash = IoHash.Compute(op.Payload.GetView().Span);
+				
                     if (!headerHash.Equals(objectHash))
                     {
-                        throw new HashMismatchException(headerHash, objectHash);
+                        throw new HashMismatchException(op.PayloadHash.Value, objectHash);
                     }
 
-                    (ContentId[] missingReferences, BlobId[] missingBlobs) = await _refService.PutAsync(ns, op.Bucket, op.Key, objectHash, op.Payload);
-                    List<JupiterContentHash> missingHashes = new List<JupiterContentHash>(missingReferences);
-
-                    return (CbSerializer.Serialize(new PutObjectResponse(missingHashes.ToArray())), HttpStatusCode.OK);
+                    (ContentId[] missingReferences, BlobId[] missingBlobs) = await _refService.PutAsync(ns, op.Bucket, op.Key, new BlobId(objectHash), op.Payload);
+                    return (CbSerializer.Serialize(new PutObjectResponse(missingReferences, missingBlobs)), HttpStatusCode.OK);
                 }
                 catch (Exception e)
                 {
@@ -1205,7 +1199,7 @@ namespace Horde.Server.Ddc
             public CbObject? Payload { get; set; } = null;
 
             [CbField("payloadHash")] 
-            public JupiterContentHash? PayloadHash { get; set; } = null;
+            public IoHash? PayloadHash { get; set; } = null;
         }
 
         [CbField("ops")]
@@ -1245,7 +1239,6 @@ namespace Horde.Server.Ddc
     {
         public RefMetadataResponse()
         {
-            PayloadIdentifier = null!;
             InlinePayload = null!;
         }
 
@@ -1301,13 +1294,18 @@ namespace Horde.Server.Ddc
             Needs = null!;
         }
 
-        public PutObjectResponse(JupiterContentHash[] missingReferences)
+        public PutObjectResponse(IoHash[] missingReferences)
         {
             Needs = missingReferences;
         }
 
+		public PutObjectResponse(IEnumerable<ContentId> missingContentIds, IEnumerable<BlobId> missingBlobIds)
+			: this(Enumerable.Concat(missingContentIds.Select(x => x.Hash), missingBlobIds.Select(x => x.Hash)).ToArray())
+		{
+		}
+
         [CbField("needs")]
-        public JupiterContentHash[] Needs { get; set; }
+        public IoHash[] Needs { get; set; }
     }
 
     public class ExistCheckMultipleRefsResponse
@@ -1359,10 +1357,10 @@ namespace Horde.Server.Ddc
 
     public class HashMismatchException : Exception
     {
-        public JupiterContentHash SuppliedHash { get; }
-        public JupiterContentHash ContentHash { get; }
+        public IoHash SuppliedHash { get; }
+        public IoHash ContentHash { get; }
 
-        public HashMismatchException(JupiterContentHash suppliedHash, JupiterContentHash contentHash) : base($"ID was not a hash of the content uploaded. Supplied hash was: {suppliedHash} but hash of content was {contentHash}")
+        public HashMismatchException(IoHash suppliedHash, IoHash contentHash) : base($"ID was not a hash of the content uploaded. Supplied hash was: {suppliedHash} but hash of content was {contentHash}")
         {
             SuppliedHash = suppliedHash;
             ContentHash = contentHash;
diff --git a/Engine/Source/Programs/Horde/Horde.Server/Ddc/IDdcBlobService.cs b/Engine/Source/Programs/Horde/Horde.Server/Ddc/IDdcBlobService.cs
--- a/Engine/Source/Programs/Horde/Horde.Server/Ddc/IDdcBlobService.cs
+++ b/Engine/Source/Programs/Horde/Horde.Server/Ddc/IDdcBlobService.cs
@@ -6,9 +6,11 @@ using System.IO;
 using System.Net.Mime;
 using System.Threading.Tasks;
 using EpicGames.AspNet;
+using EpicGames.Core;
 using EpicGames.Horde.Storage;
 using Microsoft.Extensions.DependencyInjection;
 using OpenTelemetry.Trace;
+using OpenTracing.Util;
 
 #pragma warning disable CS1591
 
@@ -19,7 +21,6 @@ namespace Horde.Server.Ddc
 	/// </summary>
 	public interface IDdcBlobService
 	{
-		Task<JupiterContentHash> VerifyContentMatchesHashAsync(Stream content, JupiterContentHash identifier);
 		Task<BlobId> PutObjectKnownHashAsync(NamespaceId ns, BufferedPayload content, BlobId identifier);
 		Task<BlobId> PutObjectAsync(NamespaceId ns, BufferedPayload payload, BlobId identifier);
 		Task<BlobId> PutObjectAsync(NamespaceId ns, byte[] payload, BlobId identifier);
@@ -100,6 +101,19 @@ namespace Horde.Server.Ddc
 
 	public static class BlobServiceExtensions
 	{
+		public static async Task VerifyContentMatchesHash(Stream stream, IoHash expectedHash, Tracer tracer)
+		{
+			IoHash hash;
+			using (TelemetrySpan _ = tracer.StartActiveSpan("web.hash").SetAttribute("operation.name", "web.hash"))
+			{
+				hash = await IoHash.ComputeAsync(stream);
+			}
+			if (hash != expectedHash)
+			{
+				throw new HashMismatchException(hash, expectedHash);
+			}
+		}
+
 		public static async Task<ContentId> PutCompressedObject(this IDdcBlobService blobService, NamespaceId ns, BufferedPayload payload, ContentId? id, IServiceProvider provider)
 		{
 			IDdcContentIdService contentIdStore = provider.GetService<IDdcContentIdService>()!;
@@ -111,28 +125,23 @@ namespace Horde.Server.Ddc
 			// TODO: we should add a overload for decompress content that can work on streams, otherwise we are still limited to 2GB compressed blobs
 			byte[] decompressedContent = compressedBufferUtils.DecompressContent(await decompressStream.ToByteArray());
 
-			await using MemoryStream decompressedStream = new MemoryStream(decompressedContent);
-			ContentId identifierDecompressedPayload;
-			if (id != null)
+			IoHash decompressedHash;
+			using (TelemetrySpan _ = tracer.StartActiveSpan("web.hash").SetAttribute("operation.name", "web.hash"))
 			{
-				identifierDecompressedPayload = ContentId.FromContentHash(await blobService.VerifyContentMatchesHashAsync(decompressedStream, id));
+				decompressedHash = IoHash.Compute(decompressedContent);
 			}
-			else
+			if (id != null && id.Value.Hash != decompressedHash)
 			{
-				JupiterContentHash blobHash;
-				{
-					using TelemetrySpan _ = tracer.StartActiveSpan("web.hash").SetAttribute("operation.name", "web.hash");
-					blobHash = await BlobId.FromStream(decompressedStream);
-				}
-
-				identifierDecompressedPayload = ContentId.FromContentHash(blobHash);
+				throw new HashMismatchException(decompressedHash, id.Value.Hash);
 			}
 
+			ContentId identifierDecompressedPayload = new ContentId(decompressedHash);
+
 			BlobId identifierCompressedPayload;
 			{
 				using TelemetrySpan _ = tracer.StartActiveSpan("web.hash").SetAttribute("operation.name", "web.hash");
 				await using Stream hashStream = payload.GetStream();
-				identifierCompressedPayload = await BlobId.FromStream(hashStream);
+				identifierCompressedPayload = new BlobId(await IoHash.ComputeAsync(hashStream));
 			}
 
 			// commit the mapping from the decompressed hash to the compressed hash, we run this in parallel with the blob store submit
diff --git a/Engine/Source/Programs/Horde/Horde.Server/Ddc/IReferenceResolver.cs b/Engine/Source/Programs/Horde/Horde.Server/Ddc/IReferenceResolver.cs
--- a/Engine/Source/Programs/Horde/Horde.Server/Ddc/IReferenceResolver.cs
+++ b/Engine/Source/Programs/Horde/Horde.Server/Ddc/IReferenceResolver.cs
@@ -29,7 +29,7 @@ namespace Horde.Server.Ddc
 
         public override IoHash AsIoHash()
         {
-            return Identifier.AsIoHash();
+			return Identifier.Hash;
         }
     }
 
@@ -44,7 +44,7 @@ namespace Horde.Server.Ddc
 
         public override IoHash AsIoHash()
         {
-            return Identifier.AsIoHash();
+			return Identifier.Hash;
         }
     }
 
@@ -61,7 +61,7 @@ namespace Horde.Server.Ddc
 
         public override IoHash AsIoHash()
         {
-            return Identifier.AsIoHash();
+			return Identifier.Hash;
         }
     }
 
@@ -119,8 +119,8 @@ namespace Horde.Server.Ddc
                     {
                         IoHash attachmentHash = field.AsAttachment();
 
-                        BlobId blobIdentifier = BlobId.FromIoHash(attachmentHash);
-                        ContentId contentId = ContentId.FromIoHash(attachmentHash);
+                        BlobId blobIdentifier = new BlobId(attachmentHash);
+                        ContentId contentId = new ContentId(attachmentHash);
 
                         if (field.IsBinaryAttachment())
                         {
@@ -145,14 +145,14 @@ namespace Horde.Server.Ddc
                 {
                     if (pendingContentIdResolve.IsCompleted)
                     {
-                        ContentId? contentId = null;
+                        ContentId contentId = default;
                         BlobId[]? resolvedBlobs = null;
-                        BlobId? blobIdentifier = null;
+                        BlobId blobIdentifier = default;
                         bool wasContentId = false;
                         try
                         {
                             (contentId, resolvedBlobs) = await pendingContentIdResolve;
-                            blobIdentifier = contentId.AsBlobIdentifier();
+							blobIdentifier = new BlobId(contentId.Hash);
                             wasContentId = !(resolvedBlobs is { Length: 1 } && resolvedBlobs[0].Equals(blobIdentifier));
                         }
                         catch (InvalidContentIdException)
@@ -166,11 +166,11 @@ namespace Horde.Server.Ddc
 
                         if (wasContentId && resolvedBlobs != null)
                         {
-                            attachments.Add(new ContentIdAttachment(contentId!, resolvedBlobs));
+                            attachments.Add(new ContentIdAttachment(contentId, resolvedBlobs));
                         }
                         else
                         {
-                            attachments.Add(new BlobAttachment(blobIdentifier!));
+                            attachments.Add(new BlobAttachment(blobIdentifier));
                         }
 
                         finishedContentIdResolves.Add(pendingContentIdResolve);
diff --git a/Engine/Source/Programs/Horde/Horde.Server/Ddc/JupiterContentHash.cs b/Engine/Source/Programs/Horde/Horde.Server/Ddc/JupiterContentHash.cs
deleted file mode 100644
