Horde: Remove support for manually deleting blobs and namespaces.

diff --git a/Engine/Source/Programs/Horde/Horde.Server/Ddc/DdcBlobsController.cs b/Engine/Source/Programs/Horde/Horde.Server/Ddc/DdcBlobsController.cs
--- a/Engine/Source/Programs/Horde/Horde.Server/Ddc/DdcBlobsController.cs
+++ b/Engine/Source/Programs/Horde/Horde.Server/Ddc/DdcBlobsController.cs
@@ -239,43 +239,6 @@ namespace Horde.Server.Ddc
             }
         }
 
-        [HttpDelete("{ns}/{id}")]
-        public async Task<IActionResult> Delete(
-            [Required] NamespaceId ns,
-            [Required] BlobId id)
-        {
-            ActionResult? result = await _requestHelper.HasAccessToNamespaceAsync(User, Request, ns, new [] { DdcAclAction.DeleteObject });
-            if (result != null)
-            {
-                return result;
-            }
-
-            await DeleteImpl(ns, id);
-
-            return NoContent();
-        }
-
-        
-        [HttpDelete("{ns}")]
-        public async Task<IActionResult> DeleteNamespace(
-            [Required] NamespaceId ns)
-        {
-            ActionResult? result = await _requestHelper.HasAccessToNamespaceAsync(User, Request, ns, new [] { DdcAclAction.DeleteNamespace });
-            if (result != null)
-            {
-                return result;
-            }
-
-            await  _storage.DeleteNamespaceAsync(ns);
-
-            return NoContent();
-        }
-
-        private async Task DeleteImpl(NamespaceId ns, BlobId id)
-        {
-            await _storage.DeleteObjectAsync(ns, id);
-        }
-
         // ReSharper disable UnusedAutoPropertyAccessor.Global
         // ReSharper disable once ClassNeverInstantiated.Global
         [System.Diagnostics.CodeAnalysis.SuppressMessage("Design", "CA1034:Nested types should not be visible", Justification = "Used by serialization")]
@@ -287,7 +250,6 @@ namespace Horde.Server.Ddc
                 INVALID,
                 GET,
                 PUT,
-                DELETE,
                 HEAD
             }
 
@@ -319,8 +281,6 @@ namespace Horde.Server.Ddc
                         return DdcAclAction.ReadObject;
                     case BatchOp.Operation.PUT:
                         return DdcAclAction.WriteObject;
-                    case BatchOp.Operation.DELETE:
-                        return DdcAclAction.DeleteObject;
                     default:
                         throw new ArgumentOutOfRangeException(nameof(op), op, null);
                 }
@@ -394,14 +354,6 @@ namespace Horde.Server.Ddc
                             tasks[index] = _storage.PutObjectAsync(op.Namespace.Value, payload, op.Id.Value).ContinueWith((t, _) => (object?) t.Result, null, TaskScheduler.Current);
                             break;
                         }
-                    case BatchOp.Operation.DELETE:
-                        if (op.Id == null)
-                        {
-                            return BadRequest();
-                        }
-
-                        tasks[index] = DeleteImpl(op.Namespace.Value, op.Id.Value).ContinueWith((t, _) => (object?) null, null, TaskScheduler.Current);
-                        break;
                     default:
                         throw new NotImplementedException($"{op.Op} is not a support op type");
                 }
diff --git a/Engine/Source/Programs/Horde/Horde.Server/Ddc/DdcObjectsController.cs b/Engine/Source/Programs/Horde/Horde.Server/Ddc/DdcObjectsController.cs
--- a/Engine/Source/Programs/Horde/Horde.Server/Ddc/DdcObjectsController.cs
+++ b/Engine/Source/Programs/Horde/Horde.Server/Ddc/DdcObjectsController.cs
@@ -220,40 +220,6 @@ namespace Horde.Server.Ddc
                 return BadRequest(new ValidationProblemDetails {Title = $"Object {id} is missing blobs", Detail = $"Following blobs are missing: {string.Join(",", e.MissingBlobs)}"});
             }
         }
-
-        [HttpDelete("{ns}/{id}")]
-        public async Task<IActionResult> Delete(
-            [Required] NamespaceId ns,
-            [Required] BlobId id)
-        {
-            ActionResult? result = await _requestHelper.HasAccessToNamespaceAsync(User, Request, ns, new [] { DdcAclAction.DeleteObject });
-            if (result != null)
-            {
-                return result;
-            }
-
-            await _storage.DeleteObjectAsync(ns, id);
-
-            return Ok( new DeletedResponse
-            {
-                DeletedCount = 1
-            });
-        }
-
-        [HttpDelete("{ns}")]
-        public async Task<IActionResult> DeleteNamespace(
-            [Required] NamespaceId ns)
-        {
-            ActionResult? result = await _requestHelper.HasAccessToNamespaceAsync(User, Request, ns, new [] { DdcAclAction.DeleteNamespace });
-            if (result != null)
-            {
-                return result;
-            }
-
-            await _storage.DeleteNamespaceAsync(ns);
-
-            return Ok();
-        }
     }
 
     public class PutBlobResponse
diff --git a/Engine/Source/Programs/Horde/Horde.Server/Ddc/IDdcBlobService.cs b/Engine/Source/Programs/Horde/Horde.Server/Ddc/IDdcBlobService.cs
--- a/Engine/Source/Programs/Horde/Horde.Server/Ddc/IDdcBlobService.cs
+++ b/Engine/Source/Programs/Horde/Horde.Server/Ddc/IDdcBlobService.cs
@@ -32,12 +32,6 @@ namespace Horde.Server.Ddc
 
 		Task<bool> ExistsAsync(NamespaceId ns, BlobId blob, List<string>? storageLayers = null);
 
-		// Delete a object
-		Task DeleteObjectAsync(NamespaceId ns, BlobId blob);
-
-		// delete the whole namespace
-		Task DeleteNamespaceAsync(NamespaceId ns);
-
 		Task<BlobId[]> FilterOutKnownBlobsAsync(NamespaceId ns, IEnumerable<BlobId> blobs);
 		Task<BlobContents> GetObjectsAsync(NamespaceId ns, BlobId[] refRequestBlobReferences);
 	}
