Horde: Change IDdcBlobService.PutObjectAsync() to extension methods.

diff --git a/Engine/Source/Programs/Horde/Horde.Server/Ddc/DdcBlobsController.cs b/Engine/Source/Programs/Horde/Horde.Server/Ddc/DdcBlobsController.cs
--- a/Engine/Source/Programs/Horde/Horde.Server/Ddc/DdcBlobsController.cs
+++ b/Engine/Source/Programs/Horde/Horde.Server/Ddc/DdcBlobsController.cs
@@ -18,6 +18,7 @@ using Horde.Server.Acls;
 using Microsoft.AspNetCore.Authorization;
 using Microsoft.AspNetCore.Http;
 using Microsoft.AspNetCore.Mvc;
+using OpenTelemetry.Trace;
 
 #pragma warning disable CS1591
 
@@ -36,14 +37,16 @@ namespace Horde.Server.Ddc
         private readonly IRequestHelper _requestHelper;
         private readonly BufferedPayloadFactory _bufferedPayloadFactory;
         private readonly NginxRedirectHelper _nginxRedirectHelper;
+		private readonly Tracer _tracer;
 
-        public DdcBlobsController(IDdcBlobService storage, IDiagnosticContext diagnosticContext, IRequestHelper requestHelper, BufferedPayloadFactory bufferedPayloadFactory, NginxRedirectHelper nginxRedirectHelper)
+        public DdcBlobsController(IDdcBlobService storage, IDiagnosticContext diagnosticContext, IRequestHelper requestHelper, BufferedPayloadFactory bufferedPayloadFactory, NginxRedirectHelper nginxRedirectHelper, Tracer tracer)
         {
             _storage = storage;
             _diagnosticContext = diagnosticContext;
             _requestHelper = requestHelper;
             _bufferedPayloadFactory = bufferedPayloadFactory;
             _nginxRedirectHelper = nginxRedirectHelper;
+			_tracer = tracer;
         }
 
         [HttpGet("{ns}/{id}")]
@@ -190,7 +193,7 @@ namespace Horde.Server.Ddc
                 }
                 using BufferedPayload payload = await _bufferedPayloadFactory.CreateFromRequest(Request);
 
-                BlobId identifier = await _storage.PutObjectAsync(ns, payload, id);
+                BlobId identifier = await _storage.PutObjectAsync(ns, payload, id, _tracer);
                 return Ok(new
                 {
                     Identifier = identifier.ToString()
@@ -351,7 +354,7 @@ namespace Horde.Server.Ddc
                             }
 
                             using MemoryBufferedPayload payload = new MemoryBufferedPayload(op.Content);
-                            tasks[index] = _storage.PutObjectAsync(op.Namespace.Value, payload, op.Id.Value).ContinueWith((t, _) => (object?) t.Result, null, TaskScheduler.Current);
+                            tasks[index] = _storage.PutObjectAsync(op.Namespace.Value, payload, op.Id.Value, _tracer).ContinueWith((t, _) => (object?) t.Result, null, TaskScheduler.Current);
                             break;
                         }
                     default:
diff --git a/Engine/Source/Programs/Horde/Horde.Server/Ddc/DdcObjectsController.cs b/Engine/Source/Programs/Horde/Horde.Server/Ddc/DdcObjectsController.cs
--- a/Engine/Source/Programs/Horde/Horde.Server/Ddc/DdcObjectsController.cs
+++ b/Engine/Source/Programs/Horde/Horde.Server/Ddc/DdcObjectsController.cs
@@ -14,6 +14,7 @@ using EpicGames.Serialization;
 using Microsoft.AspNetCore.Authorization;
 using Microsoft.AspNetCore.Mvc;
 using Microsoft.Extensions.Logging;
+using OpenTelemetry.Trace;
 
 #pragma warning disable CS1591
 
@@ -32,16 +33,17 @@ namespace Horde.Server.Ddc
         private readonly IRequestHelper _requestHelper;
         private readonly IReferenceResolver _referenceResolver;
         private readonly BufferedPayloadFactory _bufferedPayloadFactory;
-
+		private readonly Tracer _tracer;
         private readonly ILogger _logger;
 
-        public DdcObjectsController(IDdcBlobService storage, IDiagnosticContext diagnosticContext, IRequestHelper requestHelper, IReferenceResolver referenceResolver, BufferedPayloadFactory bufferedPayloadFactory, ILogger<DdcObjectsController> logger)
+        public DdcObjectsController(IDdcBlobService storage, IDiagnosticContext diagnosticContext, IRequestHelper requestHelper, IReferenceResolver referenceResolver, BufferedPayloadFactory bufferedPayloadFactory, Tracer tracer, ILogger<DdcObjectsController> logger)
         {
             _storage = storage;
             _diagnosticContext = diagnosticContext;
             _requestHelper = requestHelper;
             _referenceResolver = referenceResolver;
             _bufferedPayloadFactory = bufferedPayloadFactory;
+			_tracer = tracer;
             _logger = logger;
         }
 
@@ -160,7 +162,7 @@ namespace Horde.Server.Ddc
             {
                 using BufferedPayload payload = await _bufferedPayloadFactory.CreateFromRequest(Request);
 
-                BlobId identifier = await _storage.PutObjectAsync(ns, payload, id);
+                BlobId identifier = await _storage.PutObjectAsync(ns, payload, id, _tracer);
                 return Ok(new PutBlobResponse(identifier));
             }
             catch (ClientSendSlowException e)
diff --git a/Engine/Source/Programs/Horde/Horde.Server/Ddc/DdcRefsController.cs b/Engine/Source/Programs/Horde/Horde.Server/Ddc/DdcRefsController.cs
--- a/Engine/Source/Programs/Horde/Horde.Server/Ddc/DdcRefsController.cs
+++ b/Engine/Source/Programs/Horde/Horde.Server/Ddc/DdcRefsController.cs
@@ -647,7 +647,7 @@ namespace Horde.Server.Ddc
                     {
                         // TODO: define a scheme for how a json object specifies references
 
-                        blobHeader = await _blobStore.PutObjectAsync(ns, payload, headerHash);
+                        blobHeader = await _blobStore.PutObjectAsync(ns, payload, headerHash, _tracer);
 
                         // TODO: convert the json object into a compact binary instead
                         CbWriter writer = new CbWriter();
@@ -670,7 +670,7 @@ namespace Horde.Server.Ddc
                     }
                     case MediaTypeNames.Application.Octet:
                     {
-                        blobHeader = await _blobStore.PutObjectAsync(ns, payload, headerHash);
+                        blobHeader = await _blobStore.PutObjectAsync(ns, payload, headerHash, _tracer);
 
                         CbWriter writer = new CbWriter();
                         writer.BeginObject();
@@ -744,7 +744,7 @@ namespace Horde.Server.Ddc
                     }
                     else
                     {
-                        await _blobStore.PutObjectAsync(ns, blob, new BlobId(entry.AttachmentHash));
+                        await _blobStore.PutObjectAsync(ns, blob, new BlobId(entry.AttachmentHash), _tracer);
                     }
                 }
             }
diff --git a/Engine/Source/Programs/Horde/Horde.Server/Ddc/IDdcBlobService.cs b/Engine/Source/Programs/Horde/Horde.Server/Ddc/IDdcBlobService.cs
--- a/Engine/Source/Programs/Horde/Horde.Server/Ddc/IDdcBlobService.cs
+++ b/Engine/Source/Programs/Horde/Horde.Server/Ddc/IDdcBlobService.cs
@@ -22,8 +22,6 @@ namespace Horde.Server.Ddc
 	public interface IDdcBlobService
 	{
 		Task<BlobId> PutObjectKnownHashAsync(NamespaceId ns, BufferedPayload content, BlobId identifier);
-		Task<BlobId> PutObjectAsync(NamespaceId ns, BufferedPayload payload, BlobId identifier);
-		Task<BlobId> PutObjectAsync(NamespaceId ns, byte[] payload, BlobId identifier);
 		Task<Uri?> MaybePutObjectWithRedirectAsync(NamespaceId ns, BlobId identifier);
 
 		Task<BlobContents> GetObjectAsync(NamespaceId ns, BlobId blob, List<string>? storageLayers = null, bool supportsRedirectUri = false);
@@ -95,7 +93,7 @@ namespace Horde.Server.Ddc
 
 	public static class BlobServiceExtensions
 	{
-		public static async Task VerifyContentMatchesHash(Stream stream, IoHash expectedHash, Tracer tracer)
+		public static async Task VerifyContentMatchesHashAsync(Stream stream, IoHash expectedHash, Tracer tracer)
 		{
 			IoHash hash;
 			using (TelemetrySpan _ = tracer.StartActiveSpan("web.hash").SetAttribute("operation.name", "web.hash"))
@@ -108,6 +106,26 @@ namespace Horde.Server.Ddc
 			}
 		}
 
+		public static async Task<BlobId> PutObjectAsync(this IDdcBlobService blobService, NamespaceId ns, BufferedPayload payload, BlobId identifier, Tracer tracer)
+		{
+			using TelemetrySpan scope = tracer.StartActiveSpan("put_blob")
+				.SetAttribute("operation.name", "put_blob")
+				.SetAttribute("resource.name", identifier.ToString())
+				.SetAttribute("Content-Length", payload.Length.ToString());
+
+			await using Stream hashStream = payload.GetStream();
+			await VerifyContentMatchesHashAsync(hashStream, identifier.Hash, tracer);
+
+			await blobService.PutObjectKnownHashAsync(ns, payload, identifier);
+			return identifier;
+		}
+
+		public static Task<BlobId> PutObjectAsync(this IDdcBlobService blobService, NamespaceId ns, byte[] payload, BlobId identifier, Tracer tracer)
+		{
+			using MemoryBufferedPayload bufferedPayload = new MemoryBufferedPayload(payload);
+			return PutObjectAsync(blobService, ns, bufferedPayload, identifier, tracer);
+		}
+
 		public static async Task<ContentId> PutCompressedObject(this IDdcBlobService blobService, NamespaceId ns, BufferedPayload payload, ContentId? id, IServiceProvider provider)
 		{
 			IDdcContentIdService contentIdStore = provider.GetService<IDdcContentIdService>()!;
