Horde: Remove API to delete namespaces/buckets.

diff --git a/Engine/Source/Programs/Horde/Horde.Server/Ddc/DdcRefsController.cs b/Engine/Source/Programs/Horde/Horde.Server/Ddc/DdcRefsController.cs
--- a/Engine/Source/Programs/Horde/Horde.Server/Ddc/DdcRefsController.cs
+++ b/Engine/Source/Programs/Horde/Horde.Server/Ddc/DdcRefsController.cs
@@ -66,31 +66,6 @@ namespace Horde.Server.Ddc
             _logger = logger;
         }
 
-        /// <summary>
-        /// Returns all the known namespace the token has access to
-        /// </summary>
-        /// <returns></returns>
-        [HttpGet("")]
-        [ProducesDefaultResponseType]
-        [ProducesResponseType(type: typeof(ProblemDetails), 400)]
-        public async Task<IActionResult> GetNamespaces()
-        {
-            NamespaceId[] namespaces = await _refService.GetNamespacesAsync().ToArrayAsync();
-
-            // filter namespaces down to only the namespaces the user has access to
-            List<NamespaceId> namespacesWithAccess = new();
-            foreach (NamespaceId ns in namespaces)
-            {
-                ActionResult? accessResult = await _requestHelper.HasAccessToNamespaceAsync(User, Request, ns, new [] { DdcAclAction.ReadObject });
-                if (accessResult == null)
-                {
-                    namespacesWithAccess.Add(ns);
-                }
-            }
-
-            return Ok(new GetNamespacesResponse(namespacesWithAccess.ToArray()));
-        }
-
         /// <summary>
         /// Returns a refs key
         /// </summary>
@@ -1014,64 +989,6 @@ namespace Horde.Server.Ddc
             return (writer.ToObject(), statusCode);
         }
 
-        /// <summary>
-        /// Drop all refs records in the namespace
-        /// </summary>
-        /// <param name="ns">Namespace. Each namespace is completely separated from each other. Use for different types of data that is never expected to be similar (between two different games for instance)</param>
-        [HttpDelete("{ns}", Order = 500)]
-        [ProducesResponseType(204)]
-        public async Task<IActionResult> DeleteNamespace(
-            [FromRoute] [Required] NamespaceId ns
-        )
-        {
-            ActionResult? accessResult = await _requestHelper.HasAccessToNamespaceAsync(User, Request, ns, new [] { DdcAclAction.DeleteNamespace });
-            if (accessResult != null)
-            {
-                return accessResult;
-            }
-
-            try
-            {
-                await _refService.DropNamespaceAsync(ns);
-            }
-            catch (NamespaceNotFoundException e)
-            {
-                return NotFound(new ProblemDetails {Title = $"Namespace {e.Namespace} did not exist"});
-            }
-
-            return NoContent();
-        }
-
-        /// <summary>
-        /// Drop all refs records in the bucket
-        /// </summary>
-        /// <param name="ns">Namespace. Each namespace is completely separated from each other. Use for different types of data that is never expected to be similar (between two different games for instance)</param>
-        /// <param name="bucket">The category/type of record you are caching. Is a clustered key together with the actual key, but all records in the same bucket can be dropped easily.</param>
-        [HttpDelete("{ns}/{bucket}", Order = 500)]
-        [ProducesResponseType(200)]
-        public async Task<IActionResult> DeleteBucket(
-            [FromRoute] [Required] NamespaceId ns,
-            [FromRoute] [Required] BucketId bucket)
-        {
-            ActionResult? accessResult = await _requestHelper.HasAccessToNamespaceAsync(User, Request, ns, new [] { DdcAclAction.DeleteBucket });
-            if (accessResult != null)
-            {
-                return accessResult;
-            }
-
-            long countOfDeletedRecords;
-            try
-            {
-                countOfDeletedRecords = await _refService.DeleteBucketAsync(ns, bucket);
-            }
-            catch (NamespaceNotFoundException e)
-            {
-                return NotFound(new ProblemDetails {Title = $"Namespace {e.Namespace} did not exist"});
-            }
-
-            return Ok(new BucketDeletedResponse(countOfDeletedRecords));
-        }
-
         /// <summary>
         /// Delete a individual refs key
         /// </summary>
diff --git a/Engine/Source/Programs/Horde/Horde.Server/Ddc/IDdcRefService.cs b/Engine/Source/Programs/Horde/Horde.Server/Ddc/IDdcRefService.cs
--- a/Engine/Source/Programs/Horde/Horde.Server/Ddc/IDdcRefService.cs
+++ b/Engine/Source/Programs/Horde/Horde.Server/Ddc/IDdcRefService.cs
@@ -17,11 +17,7 @@ namespace Horde.Server.Ddc
         Task<(ContentId[], BlobId[])> PutAsync(NamespaceId ns, BucketId bucket, RefId key, BlobId blobHash, CbObject payload);
         Task<(ContentId[], BlobId[])> FinalizeAsync(NamespaceId ns, BucketId bucket, RefId key, BlobId blobHash);
 
-        IAsyncEnumerable<NamespaceId> GetNamespacesAsync();
-
         Task<bool> DeleteAsync(NamespaceId ns, BucketId bucket, RefId key);
-        Task<long> DropNamespaceAsync(NamespaceId ns);
-        Task<long> DeleteBucketAsync(NamespaceId ns, BucketId bucket);
 
         Task<bool> ExistsAsync(NamespaceId ns, BucketId bucket, RefId key);
         Task<List<BlobId>> GetReferencedBlobsAsync(NamespaceId ns, BucketId bucket, RefId key);
diff --git a/Engine/Source/Programs/Horde/Horde.Server/Ddc/IRequestHelper.cs b/Engine/Source/Programs/Horde/Horde.Server/Ddc/IRequestHelper.cs
--- a/Engine/Source/Programs/Horde/Horde.Server/Ddc/IRequestHelper.cs
+++ b/Engine/Source/Programs/Horde/Horde.Server/Ddc/IRequestHelper.cs
@@ -15,7 +15,6 @@ namespace Horde.Server.Ddc
 	public interface IRequestHelper
     {
         public Task<ActionResult?> HasAccessToNamespaceAsync(ClaimsPrincipal user, HttpRequest request, NamespaceId ns, AclAction[] aclActions);
-        public Task<ActionResult?> HasAccessForGlobalOperationsAsync(ClaimsPrincipal user, AclAction[] aclActions);
     }
 
     public class AuthorizationException : Exception
