Horde: Rename IStorageClient.FindNodesAsync() to FindAliasAsync().

diff --git a/Engine/Source/Programs/Horde/Horde.Server.Tests/StorageServiceTests.cs b/Engine/Source/Programs/Horde/Horde.Server.Tests/StorageServiceTests.cs
--- a/Engine/Source/Programs/Horde/Horde.Server.Tests/StorageServiceTests.cs
+++ b/Engine/Source/Programs/Horde/Horde.Server.Tests/StorageServiceTests.cs
@@ -51,12 +51,12 @@ namespace Horde.Server.Tests
 
 			List<BundleNodeHandle> handles;
 			
-			handles = await client.FindNodesAsync("foo").ToListAsync();
+			handles = await client.FindAliasAsync("foo").ToListAsync();
 			Assert.AreEqual(2, handles.Count);
 			Assert.AreEqual(new BundleNodeLocator(hash1, locator, 0), handles[0].GetLocator());
 			Assert.AreEqual(new BundleNodeLocator(hash1, locator, 1), handles[1].GetLocator());
 
-			handles = await client.FindNodesAsync("bar").ToListAsync();
+			handles = await client.FindAliasAsync("bar").ToListAsync();
 			Assert.AreEqual(1, handles.Count);
 			Assert.AreEqual(new BundleNodeLocator(hash2, locator, 2), handles[0].GetLocator());
 		}
diff --git a/Engine/Source/Programs/Horde/Horde.Server/Commands/Bundles/BundleWriteTest.cs b/Engine/Source/Programs/Horde/Horde.Server/Commands/Bundles/BundleWriteTest.cs
--- a/Engine/Source/Programs/Horde/Horde.Server/Commands/Bundles/BundleWriteTest.cs
+++ b/Engine/Source/Programs/Horde/Horde.Server/Commands/Bundles/BundleWriteTest.cs
@@ -30,7 +30,7 @@ namespace Horde.Server.Commands.Bundles
 			public Task<ReadOnlyMemory<byte>> ReadBundleRangeAsync(BundleLocator locator, int offset, int length, CancellationToken cancellationToken = default) => throw new NotImplementedException();
 			public Task AddAliasAsync(Utf8String name, BlobHandle handle, CancellationToken cancellationToken = default) => throw new NotImplementedException();
 			public Task RemoveAliasAsync(Utf8String name, BlobHandle handle, CancellationToken cancellationToken = default) => throw new NotImplementedException();
-			public IAsyncEnumerable<BlobHandle> FindNodesAsync(Utf8String alias, CancellationToken cancellationToken = default) => throw new NotImplementedException();
+			public IAsyncEnumerable<BlobHandle> FindAliasAsync(Utf8String alias, CancellationToken cancellationToken = default) => throw new NotImplementedException();
 			public Task<BlobHandle?> TryReadRefTargetAsync(RefName name, RefCacheTime cacheTime = default, CancellationToken cancellationToken = default) => throw new NotImplementedException();
 			public Task<BundleLocator> WriteBundleAsync(Bundle bundle, Utf8String prefix = default, CancellationToken cancellationToken = default) => Task.FromResult(BundleLocator.CreateUnique(Utf8String.Empty));
 			public Task WriteRefTargetAsync(RefName name, BlobHandle target, RefOptions? options = null, CancellationToken cancellationToken = default) => Task.CompletedTask;
diff --git a/Engine/Source/Programs/Horde/Horde.Server/Storage/StorageController.cs b/Engine/Source/Programs/Horde/Horde.Server/Storage/StorageController.cs
--- a/Engine/Source/Programs/Horde/Horde.Server/Storage/StorageController.cs
+++ b/Engine/Source/Programs/Horde/Horde.Server/Storage/StorageController.cs
@@ -337,7 +337,7 @@ namespace Horde.Server.Storage
 			StorageClient client = await _storageService.GetClientAsync(namespaceId, cancellationToken);
 
 			FindNodesResponse response = new FindNodesResponse();
-			await foreach (BundleNodeHandle handle in client.FindNodesAsync(alias, cancellationToken))
+			await foreach (BundleNodeHandle handle in client.FindAliasAsync(alias, cancellationToken))
 			{
 				response.Nodes.Add(new FindNodeResponse(handle));
 			}
diff --git a/Engine/Source/Programs/Horde/Horde.Server/Storage/StorageService.cs b/Engine/Source/Programs/Horde/Horde.Server/Storage/StorageService.cs
--- a/Engine/Source/Programs/Horde/Horde.Server/Storage/StorageService.cs
+++ b/Engine/Source/Programs/Horde/Horde.Server/Storage/StorageService.cs
@@ -226,7 +226,7 @@ namespace Horde.Server.Storage
 			public override Task RemoveAliasAsync(Utf8String name, BundleNodeLocator locator, CancellationToken cancellationToken = default) => _outer.RemoveAliasAsync(NamespaceId, name, locator, cancellationToken);
 
 			/// <inheritdoc/>
-			public override async IAsyncEnumerable<BundleNodeHandle> FindNodesAsync(Utf8String alias, [EnumeratorCancellation] CancellationToken cancellationToken = default)
+			public override async IAsyncEnumerable<BundleNodeHandle> FindAliasAsync(Utf8String alias, [EnumeratorCancellation] CancellationToken cancellationToken = default)
 			{
 				await foreach (BundleNodeLocator locator in _outer.FindNodesAsync(NamespaceId, alias, cancellationToken))
 				{
diff --git a/Engine/Source/Programs/Shared/EpicGames.Horde/Compute/AgentStorageClient.cs b/Engine/Source/Programs/Shared/EpicGames.Horde/Compute/AgentStorageClient.cs
--- a/Engine/Source/Programs/Shared/EpicGames.Horde/Compute/AgentStorageClient.cs
+++ b/Engine/Source/Programs/Shared/EpicGames.Horde/Compute/AgentStorageClient.cs
@@ -83,7 +83,7 @@ namespace EpicGames.Horde.Compute
 		}
 
 		/// <inheritdoc/>
-		public override IAsyncEnumerable<BundleNodeHandle> FindNodesAsync(Utf8String name, CancellationToken cancellationToken = default)
+		public override IAsyncEnumerable<BundleNodeHandle> FindAliasAsync(Utf8String name, CancellationToken cancellationToken = default)
 		{
 			throw new NotSupportedException();
 		}
diff --git a/Engine/Source/Programs/Shared/EpicGames.Horde/Storage/Backends/FileStorageClient.cs b/Engine/Source/Programs/Shared/EpicGames.Horde/Storage/Backends/FileStorageClient.cs
--- a/Engine/Source/Programs/Shared/EpicGames.Horde/Storage/Backends/FileStorageClient.cs
+++ b/Engine/Source/Programs/Shared/EpicGames.Horde/Storage/Backends/FileStorageClient.cs
@@ -113,7 +113,7 @@ namespace EpicGames.Horde.Storage.Backends
 		}
 
 		/// <inheritdoc/>
-		public override IAsyncEnumerable<BundleNodeHandle> FindNodesAsync(Utf8String alias, CancellationToken cancellationToken = default)
+		public override IAsyncEnumerable<BundleNodeHandle> FindAliasAsync(Utf8String alias, CancellationToken cancellationToken = default)
 		{
 			throw new NotSupportedException("File storage client does not currently support aliases.");
 		}
diff --git a/Engine/Source/Programs/Shared/EpicGames.Horde/Storage/Backends/HttpStorageClient.cs b/Engine/Source/Programs/Shared/EpicGames.Horde/Storage/Backends/HttpStorageClient.cs
--- a/Engine/Source/Programs/Shared/EpicGames.Horde/Storage/Backends/HttpStorageClient.cs
+++ b/Engine/Source/Programs/Shared/EpicGames.Horde/Storage/Backends/HttpStorageClient.cs
@@ -209,7 +209,7 @@ namespace EpicGames.Horde.Storage.Backends
 		}
 
 		/// <inheritdoc/>
-		public override async IAsyncEnumerable<BundleNodeHandle> FindNodesAsync(Utf8String alias, [EnumeratorCancellation] CancellationToken cancellationToken = default)
+		public override async IAsyncEnumerable<BundleNodeHandle> FindAliasAsync(Utf8String alias, [EnumeratorCancellation] CancellationToken cancellationToken = default)
 		{
 			_logger.LogDebug("Finding nodes with alias {Alias}", alias);
 			using (HttpClient httpClient = _createClient())
diff --git a/Engine/Source/Programs/Shared/EpicGames.Horde/Storage/Backends/MemoryStorageClient.cs b/Engine/Source/Programs/Shared/EpicGames.Horde/Storage/Backends/MemoryStorageClient.cs
--- a/Engine/Source/Programs/Shared/EpicGames.Horde/Storage/Backends/MemoryStorageClient.cs
+++ b/Engine/Source/Programs/Shared/EpicGames.Horde/Storage/Backends/MemoryStorageClient.cs
@@ -98,7 +98,7 @@ namespace EpicGames.Horde.Storage.Backends
 		}
 
 		/// <inheritdoc/>
-		public override async IAsyncEnumerable<BundleNodeHandle> FindNodesAsync(Utf8String alias, [EnumeratorCancellation] CancellationToken cancellationToken = default)
+		public override async IAsyncEnumerable<BundleNodeHandle> FindAliasAsync(Utf8String alias, [EnumeratorCancellation] CancellationToken cancellationToken = default)
 		{
 			if(_exports.TryGetValue(alias, out ExportEntry? entry))
 			{
diff --git a/Engine/Source/Programs/Shared/EpicGames.Horde/Storage/Bundles/BundleStorageClient.cs b/Engine/Source/Programs/Shared/EpicGames.Horde/Storage/Bundles/BundleStorageClient.cs
--- a/Engine/Source/Programs/Shared/EpicGames.Horde/Storage/Bundles/BundleStorageClient.cs
+++ b/Engine/Source/Programs/Shared/EpicGames.Horde/Storage/Bundles/BundleStorageClient.cs
@@ -86,16 +86,16 @@ namespace EpicGames.Horde.Storage.Bundles
 		public abstract Task RemoveAliasAsync(Utf8String name, BundleNodeHandle handle, CancellationToken cancellationToken = default);
 
 		/// <inheritdoc/>
-		async IAsyncEnumerable<BlobHandle> IStorageClient.FindNodesAsync(Utf8String name, [EnumeratorCancellation] CancellationToken cancellationToken)
+		async IAsyncEnumerable<BlobHandle> IStorageClient.FindAliasAsync(Utf8String name, [EnumeratorCancellation] CancellationToken cancellationToken)
 		{
-			await foreach (BundleNodeHandle handle in FindNodesAsync(name, cancellationToken))
+			await foreach (BundleNodeHandle handle in FindAliasAsync(name, cancellationToken))
 			{
 				yield return handle;
 			}
 		}
 
 		/// <inheritdoc/>
-		public abstract IAsyncEnumerable<BundleNodeHandle> FindNodesAsync(Utf8String name, CancellationToken cancellationToken = default);
+		public abstract IAsyncEnumerable<BundleNodeHandle> FindAliasAsync(Utf8String name, CancellationToken cancellationToken = default);
 
 		#endregion
 
diff --git a/Engine/Source/Programs/Shared/EpicGames.Horde/Storage/IStorageClient.cs b/Engine/Source/Programs/Shared/EpicGames.Horde/Storage/IStorageClient.cs
--- a/Engine/Source/Programs/Shared/EpicGames.Horde/Storage/IStorageClient.cs
+++ b/Engine/Source/Programs/Shared/EpicGames.Horde/Storage/IStorageClient.cs
@@ -108,7 +108,7 @@ namespace EpicGames.Horde.Storage
 		/// <param name="name">Alias for the node</param>
 		/// <param name="cancellationToken">Cancellation token for the operation</param>
 		/// <returns>Nodes matching the given handle</returns>
-		IAsyncEnumerable<BlobHandle> FindNodesAsync(Utf8String name, CancellationToken cancellationToken = default);
+		IAsyncEnumerable<BlobHandle> FindAliasAsync(Utf8String name, CancellationToken cancellationToken = default);
 
 		#endregion
 
diff --git a/Engine/Source/Programs/UnrealCloudDDC/Jupiter/Controllers/BundlesController.cs b/Engine/Source/Programs/UnrealCloudDDC/Jupiter/Controllers/BundlesController.cs
--- a/Engine/Source/Programs/UnrealCloudDDC/Jupiter/Controllers/BundlesController.cs
+++ b/Engine/Source/Programs/UnrealCloudDDC/Jupiter/Controllers/BundlesController.cs
@@ -289,7 +289,7 @@ namespace Jupiter.Controllers
         /// <param name="cancellationToken">Cancellation token for the operation</param>
         [HttpGet]
         [Route("{namespaceId}/nodes")]
-        public async Task<ActionResult<FindNodesResponse>> FindNodesAsync(NamespaceId namespaceId, string alias, CancellationToken cancellationToken = default)
+        public async Task<ActionResult<FindNodesResponse>> FindAliasAsync(NamespaceId namespaceId, string alias, CancellationToken cancellationToken = default)
         {
             ActionResult? result = await _requestHelper.HasAccessToNamespace(User, Request, namespaceId, new [] { JupiterAclAction.ReadObject });
             if (result != null)
@@ -300,7 +300,7 @@ namespace Jupiter.Controllers
             StorageClient client = await _storageService.GetClientAsync(namespaceId, cancellationToken);
 
             FindNodesResponse response = new FindNodesResponse();
-            await foreach (BundleNodeHandle handle in client.FindNodesAsync(alias, cancellationToken))
+            await foreach (BundleNodeHandle handle in client.FindAliasAsync(alias, cancellationToken))
             {
                 response.Nodes.Add(new FindNodeResponse(handle));
             }
diff --git a/Engine/Source/Programs/UnrealCloudDDC/Jupiter/Implementation/Bundles/StorageService.cs b/Engine/Source/Programs/UnrealCloudDDC/Jupiter/Implementation/Bundles/StorageService.cs
--- a/Engine/Source/Programs/UnrealCloudDDC/Jupiter/Implementation/Bundles/StorageService.cs
+++ b/Engine/Source/Programs/UnrealCloudDDC/Jupiter/Implementation/Bundles/StorageService.cs
@@ -104,7 +104,7 @@ public class StorageClient : BundleStorageClient
         await Task.CompletedTask;
     }
 
-    public override async IAsyncEnumerable<BundleNodeHandle> FindNodesAsync(Utf8String name, [EnumeratorCancellation] CancellationToken cancellationToken = default)
+    public override async IAsyncEnumerable<BundleNodeHandle> FindAliasAsync(Utf8String name, [EnumeratorCancellation] CancellationToken cancellationToken = default)
     {
         // TODO: Implement aliases
         await Task.CompletedTask;
