Horde: Initial implementation of DDC services in Horde.

diff --git a/Engine/Source/Programs/Horde/Horde.Server/Artifacts/ArtifactsController.cs b/Engine/Source/Programs/Horde/Horde.Server/Artifacts/ArtifactsController.cs
--- a/Engine/Source/Programs/Horde/Horde.Server/Artifacts/ArtifactsController.cs
+++ b/Engine/Source/Programs/Horde/Horde.Server/Artifacts/ArtifactsController.cs
@@ -191,7 +191,8 @@ namespace Horde.Server.Artifacts
 				return BadRequest("Invalid blob id for artifact");
 			}
 
-			return StorageController.ReadBlobInternalAsync(_storageService, artifact.NamespaceId, locator, offset, length, cancellationToken);
+			StorageClient storageClient = await _storageService.GetClientAsync(artifact.NamespaceId, cancellationToken);
+			return StorageController.ReadBlobInternalAsync(storageClient, locator, offset, length, cancellationToken);
 		}
 
 		/// <summary>
diff --git a/Engine/Source/Programs/Horde/Horde.Server/Ddc/BufferedPayload.cs b/Engine/Source/Programs/Horde/Horde.Server/Ddc/BufferedPayload.cs
--- a/Engine/Source/Programs/Horde/Horde.Server/Ddc/BufferedPayload.cs
+++ b/Engine/Source/Programs/Horde/Horde.Server/Ddc/BufferedPayload.cs
@@ -3,6 +3,7 @@
 using System;
 using System.IO;
 using System.Threading.Tasks;
+using EpicGames.Core;
 using Microsoft.AspNetCore.Http;
 using Microsoft.Extensions.Options;
 using OpenTelemetry.Trace;
@@ -52,13 +53,13 @@ namespace Horde.Server.Ddc
     /// </summary>
     public sealed class MemoryBufferedPayload : BufferedPayload
     {
-        private readonly byte[] _buffer;
+        private readonly ReadOnlyMemory<byte> _buffer;
 
 		/// <summary>
 		/// Constructor
 		/// </summary>
-        public MemoryBufferedPayload(byte[] source)
-			: base(source.LongLength)
+        public MemoryBufferedPayload(ReadOnlyMemory<byte> source)
+			: base(source.Length)
         {
             _buffer = source;
         }
@@ -76,7 +77,7 @@ namespace Horde.Server.Ddc
         }
 
 		/// <inheritdoc/>
-		public override Stream GetStream() => new MemoryStream(_buffer);
+		public override Stream GetStream() => new ReadOnlyMemoryStream(_buffer);
     }
 
     /// <summary>
diff --git a/Engine/Source/Programs/Horde/Horde.Server/Ddc/DdcBlobService.cs b/Engine/Source/Programs/Horde/Horde.Server/Ddc/DdcBlobService.cs
--- a/Engine/Source/Programs/Horde/Horde.Server/Ddc/DdcBlobService.cs
+++ b/Engine/Source/Programs/Horde/Horde.Server/Ddc/DdcBlobService.cs
@@ -3,97 +3,150 @@
 using System;
 using System.Collections.Generic;
 using System.IO;
+using System.Linq;
+using System.Threading;
 using System.Threading.Tasks;
 using EpicGames.Core;
 using EpicGames.Horde.Storage;
+using Horde.Server.Storage;
+using OpenTelemetry.Trace;
 
 namespace Horde.Server.Ddc
 {
 	class DdcBlobService : IDdcBlobService
 	{
-		public Task DeleteNamespaceAsync(NamespaceId ns)
-		{
-			throw new NotImplementedException();
-		}
+		BlobType s_rawBlobType = new BlobType(new Guid("{03E6C37B-33C1-491F-8541-D3C401B8B8EF}"), 1);
 
-		public Task DeleteObjectAsync(NamespaceId ns, BlobId blob)
-		{
-			throw new NotImplementedException();
-		}
+		readonly StorageService _storageService;
+		readonly Tracer _tracer;
 
-		public Task<bool> ExistsAsync(NamespaceId ns, BlobId blob, List<string>? storageLayers = null)
+		public DdcBlobService(StorageService storageService, Tracer tracer)
 		{
-			throw new NotImplementedException();
+			_storageService = storageService;
+			_tracer = tracer;
 		}
 
-		public Task<bool> ExistsInRootStore(NamespaceId ns, BlobId blob)
-		{
-			throw new NotImplementedException();
-		}
+		public static Utf8String GetAlias(BlobId blobId) => $"ddc:{blobId}";
 
-		public Task<BlobId[]> FilterOutKnownBlobsAsync(NamespaceId ns, IEnumerable<BlobId> blobs)
+		public async Task VerifyContentMatchesHashAsync(Stream stream, IoHash expectedHash, CancellationToken cancellationToken)
 		{
-			throw new NotImplementedException();
+			IoHash hash;
+			using (TelemetrySpan _ = _tracer.StartActiveSpan("web.hash").SetAttribute("operation.name", "web.hash"))
+			{
+				hash = await IoHash.ComputeAsync(stream, cancellationToken);
+			}
+			if (hash != expectedHash)
+			{
+				throw new HashMismatchException(hash, expectedHash);
+			}
 		}
 
-		public Task<BlobId[]> FilterOutKnownBlobsAsync(NamespaceId ns, IAsyncEnumerable<BlobId> blobs)
+		public async Task<BlobId> PutObjectAsync(NamespaceId ns, BufferedPayload payload, BlobId identifier, CancellationToken cancellationToken)
 		{
-			throw new NotImplementedException();
-		}
+			using TelemetrySpan scope = _tracer.StartActiveSpan("put_blob")
+				.SetAttribute("operation.name", "put_blob")
+				.SetAttribute("resource.name", identifier.ToString())
+				.SetAttribute("Content-Length", payload.Length.ToString());
 
-		public Task<BlobContents> GetObjectAsync(NamespaceId ns, BlobId blob, List<string>? storageLayers = null, bool supportsRedirectUri = false)
-		{
-			throw new NotImplementedException();
-		}
+			await using Stream hashStream = payload.GetStream();
+			await VerifyContentMatchesHashAsync(hashStream, identifier.Hash, cancellationToken);
 
-		public Task<BlobContents> GetObjectsAsync(NamespaceId ns, BlobId[] refRequestBlobReferences)
-		{
-			throw new NotImplementedException();
+			await PutObjectKnownHashAsync(ns, payload, identifier, cancellationToken);
+			return identifier;
 		}
 
-		public Task<Uri?> GetObjectWithRedirectAsync(NamespaceId ns, BlobId blobIdentifier, List<string>? storageLayers = null)
+		public Task<BlobId> PutObjectAsync(NamespaceId ns, ReadOnlyMemory<byte> payload, BlobId identifier, CancellationToken cancellationToken)
 		{
-			throw new NotImplementedException();
+			using MemoryBufferedPayload bufferedPayload = new MemoryBufferedPayload(payload);
+			return PutObjectAsync(ns, bufferedPayload, identifier, cancellationToken);
 		}
 
-		public IAsyncEnumerable<(BlobId, DateTime)> ListObjects(NamespaceId ns)
+		public async Task<bool> ExistsAsync(NamespaceId ns, BlobId blob, List<string>? storageLayers, CancellationToken cancellationToken)
 		{
-			throw new NotImplementedException();
+			IStorageClient storageClient = await _storageService.GetClientAsync(ns, cancellationToken);
+			return await storageClient.FindAliasAsync(GetAlias(blob), cancellationToken).AnyAsync(cancellationToken);
 		}
 
-		public Task<Uri?> MaybePutObjectWithRedirectAsync(NamespaceId ns, BlobId identifier)
+		public async Task<BlobId[]> FilterOutKnownBlobsAsync(NamespaceId ns, IEnumerable<BlobId> blobIds, CancellationToken cancellationToken)
 		{
-			throw new NotImplementedException();
-		}
+			IStorageClient storageClient = await _storageService.GetClientAsync(ns, cancellationToken);
 
-		public Task<BlobId> PutObjectAsync(NamespaceId ns, BufferedPayload payload, BlobId identifier)
-		{
-			throw new NotImplementedException();
+			List<BlobId> unknownBlobIds = new List<BlobId>();
+			foreach (BlobId blobId in blobIds)
+			{
+				if (!await storageClient.FindAliasAsync(GetAlias(blobId), cancellationToken).AnyAsync())
+				{
+					unknownBlobIds.Add(blobId);
+				}
+			}
+
+			return unknownBlobIds.ToArray();
 		}
 
-		public Task<BlobId> PutObjectAsync(NamespaceId ns, byte[] payload, BlobId identifier)
+		public async Task<BlobContents> GetObjectAsync(NamespaceId ns, BlobId blob, List<string>? storageLayers, bool supportsRedirectUri, CancellationToken cancellationToken)
 		{
-			throw new NotImplementedException();
+			IStorageClient storageClient = await _storageService.GetClientAsync(ns, cancellationToken);
+
+			BlobHandle? handle = await storageClient.FindAliasAsync(GetAlias(blob), cancellationToken).FirstOrDefaultAsync(cancellationToken);
+			if (handle == null)
+			{
+				throw new BlobNotFoundException(ns, blob);
+			}
+
+			BlobData data = await handle.ReadAsync(cancellationToken);
+			return new BlobContents(data.Data.ToArray());
 		}
 
-		public Task<BlobId> PutObjectKnownHashAsync(NamespaceId ns, BufferedPayload content, BlobId identifier)
+		public async Task<BlobContents> GetObjectsAsync(NamespaceId ns, BlobId[] blobs, CancellationToken cancellationToken)
 		{
-			throw new NotImplementedException();
+			using TelemetrySpan _ = _tracer.StartActiveSpan("blob.combine").SetAttribute("operation.name", "blob.combine");
+			Task<BlobContents>[] tasks = new Task<BlobContents>[blobs.Length];
+			for (int i = 0; i < blobs.Length; i++)
+			{
+				tasks[i] = GetObjectAsync(ns, blobs[i], storageLayers: null, supportsRedirectUri: false, cancellationToken);
+			}
+
+			MemoryStream ms = new MemoryStream();
+			foreach (Task<BlobContents> task in tasks)
+			{
+				BlobContents blob = await task;
+				await using Stream s = blob.Stream;
+				await s.CopyToAsync(ms, cancellationToken);
+			}
+
+			ms.Seek(0, SeekOrigin.Begin);
+
+			return new BlobContents(ms, ms.Length);
 		}
 
-		public Task<BlobContents> ReplicateObjectAsync(NamespaceId ns, BlobId blob, bool force = false)
+		public Task<Uri?> GetObjectWithRedirectAsync(NamespaceId ns, BlobId blobIdentifier, List<string>? storageLayers, CancellationToken cancellationToken)
 		{
-			throw new NotImplementedException();
+			return Task.FromResult<Uri?>(null);
 		}
 
-		public bool ShouldFetchBlobOnDemand(NamespaceId ns)
+		public Task<Uri?> MaybePutObjectWithRedirectAsync(NamespaceId ns, BlobId identifier, CancellationToken cancellationToken)
 		{
-			throw new NotImplementedException();
+			return Task.FromResult<Uri?>(null);
 		}
 
-		public Task<IoHash> VerifyContentMatchesHashAsync(Stream content, IoHash identifier)
+		public async Task<BlobId> PutObjectKnownHashAsync(NamespaceId ns, BufferedPayload content, BlobId identifier, CancellationToken cancellationToken)
 		{
-			throw new NotImplementedException();
+			IStorageClient storageClient = await _storageService.GetClientAsync(ns, cancellationToken);
+
+			BlobHandle blobHandle;
+			await using (IStorageWriter writer = storageClient.CreateWriter())
+			{
+				Memory<byte> memory = writer.GetOutputBuffer(0, (int)content.Length);
+
+				using Stream stream = content.GetStream();
+				await stream.ReadAsync(memory, cancellationToken);
+
+				blobHandle = await writer.WriteNodeAsync(memory.Length, Array.Empty<BlobHandle>(), s_rawBlobType, cancellationToken);
+				await writer.FlushAsync(cancellationToken);
+			}
+
+			await storageClient.AddAliasAsync(GetAlias(identifier), blobHandle, cancellationToken);
+			return identifier;
 		}
 	}
 }
diff --git a/Engine/Source/Programs/Horde/Horde.Server/Ddc/DdcBlobsController.cs b/Engine/Source/Programs/Horde/Horde.Server/Ddc/DdcBlobsController.cs
--- a/Engine/Source/Programs/Horde/Horde.Server/Ddc/DdcBlobsController.cs
+++ b/Engine/Source/Programs/Horde/Horde.Server/Ddc/DdcBlobsController.cs
@@ -9,6 +9,7 @@ using System.Linq;
 using System.Net;
 using System.Net.Mime;
 using System.Text.Json.Serialization;
+using System.Threading;
 using System.Threading.Tasks;
 using EpicGames.AspNet;
 using EpicGames.Core;
@@ -170,7 +171,8 @@ namespace Horde.Server.Ddc
         [DisableRequestSizeLimit]
         public async Task<IActionResult> Put(
             [Required] NamespaceId ns,
-            [Required] BlobId id)
+            [Required] BlobId id,
+			CancellationToken cancellationToken)
         {
             ActionResult? result = await _requestHelper.HasAccessToNamespaceAsync(User, Request, ns, new [] { DdcAclAction.WriteObject });
             if (result != null)
@@ -193,7 +195,7 @@ namespace Horde.Server.Ddc
                 }
                 using BufferedPayload payload = await _bufferedPayloadFactory.CreateFromRequest(Request);
 
-                BlobId identifier = await _storage.PutObjectAsync(ns, payload, id, _tracer);
+                BlobId identifier = await _storage.PutObjectAsync(ns, payload, id, cancellationToken);
                 return Ok(new
                 {
                     Identifier = identifier.ToString()
@@ -273,7 +275,7 @@ namespace Horde.Server.Ddc
         // ReSharper restore UnusedAutoPropertyAccessor.Global
 
         [HttpPost("")]
-        public async Task<IActionResult> Post([FromBody] BatchCall batch)
+        public async Task<IActionResult> Post([FromBody] BatchCall batch, CancellationToken cancellationToken)
         {
             static AclAction MapToAclAction(BatchOp.Operation op)
             {
@@ -354,7 +356,7 @@ namespace Horde.Server.Ddc
                             }
 
                             using MemoryBufferedPayload payload = new MemoryBufferedPayload(op.Content);
-                            tasks[index] = _storage.PutObjectAsync(op.Namespace.Value, payload, op.Id.Value, _tracer).ContinueWith((t, _) => (object?) t.Result, null, TaskScheduler.Current);
+                            tasks[index] = _storage.PutObjectAsync(op.Namespace.Value, payload, op.Id.Value, cancellationToken).ContinueWith((t, _) => (object?) t.Result, null, TaskScheduler.Current);
                             break;
                         }
                     default:
diff --git a/Engine/Source/Programs/Horde/Horde.Server/Ddc/DdcObjectsController.cs b/Engine/Source/Programs/Horde/Horde.Server/Ddc/DdcObjectsController.cs
--- a/Engine/Source/Programs/Horde/Horde.Server/Ddc/DdcObjectsController.cs
+++ b/Engine/Source/Programs/Horde/Horde.Server/Ddc/DdcObjectsController.cs
@@ -7,6 +7,7 @@ using System.ComponentModel.DataAnnotations;
 using System.Linq;
 using System.Net;
 using System.Net.Mime;
+using System.Threading;
 using System.Threading.Tasks;
 using EpicGames.AspNet;
 using EpicGames.Horde.Storage;
@@ -149,7 +150,8 @@ namespace Horde.Server.Ddc
         [RequiredContentType(CustomMediaTypeNames.UnrealCompactBinary)]
         public async Task<IActionResult> Put(
             [Required] NamespaceId ns,
-            [Required] BlobId id)
+            [Required] BlobId id,
+			CancellationToken cancellationToken)
         {
             ActionResult? result = await _requestHelper.HasAccessToNamespaceAsync(User, Request, ns, new [] { DdcAclAction.WriteObject });
             if (result != null)
@@ -162,7 +164,7 @@ namespace Horde.Server.Ddc
             {
                 using BufferedPayload payload = await _bufferedPayloadFactory.CreateFromRequest(Request);
 
-                BlobId identifier = await _storage.PutObjectAsync(ns, payload, id, _tracer);
+                BlobId identifier = await _storage.PutObjectAsync(ns, payload, id, cancellationToken);
                 return Ok(new PutBlobResponse(identifier));
             }
             catch (ClientSendSlowException e)
diff --git a/Engine/Source/Programs/Horde/Horde.Server/Ddc/DdcRefService.cs b/Engine/Source/Programs/Horde/Horde.Server/Ddc/DdcRefService.cs
--- a/Engine/Source/Programs/Horde/Horde.Server/Ddc/DdcRefService.cs
+++ b/Engine/Source/Programs/Horde/Horde.Server/Ddc/DdcRefService.cs
@@ -1,58 +1,158 @@
 // Copyright Epic Games, Inc. All Rights Reserved.
 
+using System;
 using System.Collections.Generic;
+using System.Linq;
+using System.Threading;
 using System.Threading.Tasks;
 using EpicGames.Core;
 using EpicGames.Horde.Storage;
 using EpicGames.Serialization;
+using Horde.Server.Storage;
 
 namespace Horde.Server.Ddc
 {
-	class DdcRefService : IDdcRefService
+	/// <summary>
+	/// A node containing arbitrary compact binary data
+	/// </summary>
+	[NodeType("{0C7E5F25-454B-4B55-9B4A-F4635106D074}", 1)]
+	class DdcRefNode : Node
 	{
-		public Task<bool> DeleteAsync(NamespaceId ns, BucketId bucket, RefId key)
+		public IoHash RootHash { get; }
+		public List<(IoHash Hash, BlobHandle Handle)> References { get; }
+
+		public DdcRefNode(IoHash rootHash)
 		{
-			throw new System.NotImplementedException();
+			RootHash = rootHash;
+			References = new List<(IoHash, BlobHandle)>();
 		}
 
-		public Task<long> DeleteBucketAsync(NamespaceId ns, BucketId bucket)
+		public DdcRefNode(NodeReader reader)
 		{
-			throw new System.NotImplementedException();
+			RootHash = reader.ReadIoHash();
+			References = reader.ReadList(x => (reader.ReadIoHash(), reader.ReadNodeHandle()));
 		}
 
-		public Task<long> DropNamespaceAsync(NamespaceId ns)
+		/// <inheritdoc/>
+		public override void Serialize(NodeWriter writer)
 		{
-			throw new System.NotImplementedException();
+			writer.WriteIoHash(RootHash);
+			writer.WriteList(References, x => WriteReference(writer, x.Hash, x.Handle));
 		}
 
-		public Task<bool> ExistsAsync(NamespaceId ns, BucketId bucket, RefId key)
+		static void WriteReference(NodeWriter writer, IoHash hash, BlobHandle handle)
 		{
-			throw new System.NotImplementedException();
+			writer.WriteIoHash(hash);
+			writer.WriteNodeHandle(handle);
 		}
+	}
+
+	class DdcRefService : IDdcRefService
+	{
+		readonly StorageService _storageService;
+		readonly DdcBlobService _blobService;
 
-		public Task<(ContentId[], BlobId[])> FinalizeAsync(NamespaceId ns, BucketId bucket, RefId key, BlobId blobHash)
+		public DdcRefService(StorageService storageService, DdcBlobService blobService)
 		{
-			throw new System.NotImplementedException();
+			_storageService = storageService;
+			_blobService = blobService;
 		}
 
-		public Task<(RefRecord, BlobContents?)> GetAsync(NamespaceId ns, BucketId bucket, RefId key, string[] fields, bool doLastAccessTracking = true)
+		static RefName GetRefName(BucketId bucketId, RefId refId) => $"{bucketId}/{refId}";
+
+		public async Task<bool> DeleteAsync(NamespaceId ns, BucketId bucket, RefId key, CancellationToken cancellationToken)
 		{
-			throw new System.NotImplementedException();
+			IStorageClient storageClient = await _storageService.GetClientAsync(ns, cancellationToken);
+			await storageClient.DeleteRefAsync(GetRefName(bucket, key), cancellationToken);
+			return true;
 		}
 
-		public IAsyncEnumerable<NamespaceId> GetNamespacesAsync()
+		public async Task<bool> ExistsAsync(NamespaceId ns, BucketId bucket, RefId key, CancellationToken cancellationToken)
 		{
-			throw new System.NotImplementedException();
+			IStorageClient storageClient = await _storageService.GetClientAsync(ns, cancellationToken);
+			BlobHandle? handle = await storageClient.TryReadRefTargetAsync(GetRefName(bucket, key), cancellationToken: cancellationToken);
+			return handle != null;
 		}
 
-		public Task<List<BlobId>> GetReferencedBlobsAsync(NamespaceId ns, BucketId bucket, RefId key)
+		public async Task<(ContentId[], BlobId[])> FinalizeAsync(NamespaceId ns, BucketId bucket, RefId key, BlobId blobHash, CancellationToken cancellationToken)
 		{
-			throw new System.NotImplementedException();
+			HashSet<BlobId> missingBlobIds = new HashSet<BlobId>();
+
+			// Read all the blobs back in
+			List<BlobId> blobIds = new List<BlobId>();
+			blobIds.Add(blobHash);
+
+			HashSet<BlobId> addedBlobIds = new HashSet<BlobId>();
+			for (int idx = 0; idx < blobIds.Count; idx++)
+			{
+				try
+				{
+					using BlobContents contents = await _blobService.GetObjectAsync(ns, blobIds[idx], storageLayers: null, supportsRedirectUri: false, cancellationToken: cancellationToken);
+					byte[] data = await contents.Stream.ReadAllBytesAsync(cancellationToken);
+
+					CbObject obj = new CbObject(data);
+					obj.IterateAttachments(x =>
+					{
+						BlobId blobId = new BlobId(x.AsAttachment());
+						if (addedBlobIds.Add(blobId))
+						{
+							blobIds.Add(blobId);
+						}
+					});
+				}
+				catch (BlobNotFoundException)
+				{
+					missingBlobIds.Add(blobIds[idx]);
+				}
+			}
+
+			if (missingBlobIds.Count == 0)
+			{
+				IStorageClient storageClient = await _storageService.GetClientAsync(ns, cancellationToken);
+
+				BlobHandle? handle = await storageClient.FindAliasAsync(DdcBlobService.GetAlias(blobHash), cancellationToken).FirstOrDefaultAsync(cancellationToken);
+				if (handle == null)
+				{
+					throw new BlobNotFoundException(ns, blobHash);
+				}
+
+				await storageClient.WriteRefTargetAsync(GetRefName(bucket, key), handle, cancellationToken: cancellationToken);
+			}
+
+			return (Array.Empty<ContentId>(), missingBlobIds.ToArray());
+		}
+
+		public async Task<(RefRecord, BlobContents?)> GetAsync(NamespaceId ns, BucketId bucket, RefId key, string[] fields, bool doLastAccessTracking, CancellationToken cancellationToken)
+		{
+			IStorageClient storageClient = await _storageService.GetClientAsync(ns, cancellationToken);
+			
+			DdcRefNode? node = await storageClient.TryReadNodeAsync<DdcRefNode>(GetRefName(bucket, key), cancellationToken: cancellationToken);
+			if (node == null)
+			{
+				throw new RefNotFoundException(ns, bucket, key);
+			}
+
+			RefRecord record = new RefRecord(ns, bucket, key, DateTime.UtcNow, null, new BlobId(node.RootHash), false);
+			return (record, null);
+		}
+
+		public async Task<List<BlobId>> GetReferencedBlobsAsync(NamespaceId ns, BucketId bucket, RefId key, CancellationToken cancellationToken)
+		{
+			IStorageClient storageClient = await _storageService.GetClientAsync(ns, cancellationToken);
+
+			DdcRefNode? node = await storageClient.TryReadNodeAsync<DdcRefNode>(GetRefName(bucket, key), cancellationToken: cancellationToken);
+			if (node == null)
+			{
+				throw new RefNotFoundException(ns, bucket, key);
+			}
+
+			return node.References.Select(x => new BlobId(x.Hash)).ToList();
 		}
 
-		public Task<(ContentId[], BlobId[])> PutAsync(NamespaceId ns, BucketId bucket, RefId key, BlobId blobHash, CbObject payload)
+		public async Task<(ContentId[], BlobId[])> PutAsync(NamespaceId ns, BucketId bucket, RefId key, BlobId blobHash, CbObject payload, CancellationToken cancellationToken)
 		{
-			throw new System.NotImplementedException();
+			await _blobService.PutObjectAsync(ns, payload.GetView(), blobHash, cancellationToken);
+			return await FinalizeAsync(ns, bucket, key, blobHash, cancellationToken);
 		}
 	}
 }
diff --git a/Engine/Source/Programs/Horde/Horde.Server/Ddc/DdcRefsController.cs b/Engine/Source/Programs/Horde/Horde.Server/Ddc/DdcRefsController.cs
--- a/Engine/Source/Programs/Horde/Horde.Server/Ddc/DdcRefsController.cs
+++ b/Engine/Source/Programs/Horde/Horde.Server/Ddc/DdcRefsController.cs
@@ -584,7 +584,8 @@ namespace Horde.Server.Ddc
         public async Task<IActionResult> PutObject(
             [FromRoute] [Required] NamespaceId ns,
             [FromRoute] [Required] BucketId bucket,
-            [FromRoute] [Required] RefId key)
+            [FromRoute] [Required] RefId key,
+			CancellationToken cancellationToken)
         {
             ActionResult? accessResult = await _requestHelper.HasAccessToNamespaceAsync(User, Request, ns, new [] { DdcAclAction.WriteObject });
             if (accessResult != null)
@@ -622,7 +623,7 @@ namespace Horde.Server.Ddc
                     {
                         // TODO: define a scheme for how a json object specifies references
 
-                        blobHeader = await _blobStore.PutObjectAsync(ns, payload, headerHash, _tracer);
+                        blobHeader = await _blobStore.PutObjectAsync(ns, payload, headerHash, cancellationToken);
 
                         // TODO: convert the json object into a compact binary instead
                         CbWriter writer = new CbWriter();
@@ -639,13 +640,13 @@ namespace Horde.Server.Ddc
                     {
                         await using MemoryStream ms = new MemoryStream();
                         await using Stream payloadStream = payload.GetStream();
-                        await payloadStream.CopyToAsync(ms);
+                        await payloadStream.CopyToAsync(ms, cancellationToken);
                         payloadObject = new CbObject(ms.ToArray());
                         break;
                     }
                     case MediaTypeNames.Application.Octet:
                     {
-                        blobHeader = await _blobStore.PutObjectAsync(ns, payload, headerHash, _tracer);
+                        blobHeader = await _blobStore.PutObjectAsync(ns, payload, headerHash, cancellationToken);
 
                         CbWriter writer = new CbWriter();
                         writer.BeginObject();
@@ -688,7 +689,8 @@ namespace Horde.Server.Ddc
         public async Task<IActionResult> PutPackage(
             [FromRoute][Required] NamespaceId ns,
             [FromRoute][Required] BucketId bucket,
-            [FromRoute][Required] RefId key)
+            [FromRoute][Required] RefId key,
+			CancellationToken cancellationToken)
         {
             ActionResult? accessResult = await _requestHelper.HasAccessToNamespaceAsync(User, Request, ns, new [] { DdcAclAction.WriteObject });
             if (accessResult != null)
@@ -719,7 +721,7 @@ namespace Horde.Server.Ddc
                     }
                     else
                     {
-                        await _blobStore.PutObjectAsync(ns, blob, new BlobId(entry.AttachmentHash), _tracer);
+                        await _blobStore.PutObjectAsync(ns, blob, new BlobId(entry.AttachmentHash), cancellationToken);
                     }
                 }
             }
diff --git a/Engine/Source/Programs/Horde/Horde.Server/Ddc/IDdcBlobService.cs b/Engine/Source/Programs/Horde/Horde.Server/Ddc/IDdcBlobService.cs
--- a/Engine/Source/Programs/Horde/Horde.Server/Ddc/IDdcBlobService.cs
+++ b/Engine/Source/Programs/Horde/Horde.Server/Ddc/IDdcBlobService.cs
@@ -4,6 +4,7 @@ using System;
 using System.Collections.Generic;
 using System.IO;
 using System.Net.Mime;
+using System.Threading;
 using System.Threading.Tasks;
 using EpicGames.AspNet;
 using EpicGames.Core;
@@ -21,17 +22,20 @@ namespace Horde.Server.Ddc
 	/// </summary>
 	public interface IDdcBlobService
 	{
-		Task<BlobId> PutObjectKnownHashAsync(NamespaceId ns, BufferedPayload content, BlobId identifier);
-		Task<Uri?> MaybePutObjectWithRedirectAsync(NamespaceId ns, BlobId identifier);
+		Task VerifyContentMatchesHashAsync(Stream stream, IoHash expectedHash, CancellationToken cancellationToken = default);
+		Task<BlobId> PutObjectAsync(NamespaceId ns, BufferedPayload payload, BlobId identifier, CancellationToken cancellationToken = default);
+		Task<BlobId> PutObjectAsync(NamespaceId ns, ReadOnlyMemory<byte> payload, BlobId identifier, CancellationToken cancellationToken = default);
+		Task<BlobId> PutObjectKnownHashAsync(NamespaceId ns, BufferedPayload content, BlobId identifier, CancellationToken cancellationToken = default);
+		Task<Uri?> MaybePutObjectWithRedirectAsync(NamespaceId ns, BlobId identifier, CancellationToken cancellationToken = default);
 
-		Task<BlobContents> GetObjectAsync(NamespaceId ns, BlobId blob, List<string>? storageLayers = null, bool supportsRedirectUri = false);
+		Task<BlobContents> GetObjectAsync(NamespaceId ns, BlobId blob, List<string>? storageLayers = null, bool supportsRedirectUri = false, CancellationToken cancellationToken = default);
 
-		Task<Uri?> GetObjectWithRedirectAsync(NamespaceId ns, BlobId blobIdentifier, List<string>? storageLayers = null);
+		Task<Uri?> GetObjectWithRedirectAsync(NamespaceId ns, BlobId blobIdentifier, List<string>? storageLayers = null, CancellationToken cancellationToken = default);
 
-		Task<bool> ExistsAsync(NamespaceId ns, BlobId blob, List<string>? storageLayers = null);
+		Task<bool> ExistsAsync(NamespaceId ns, BlobId blob, List<string>? storageLayers = null, CancellationToken cancellationToken = default);
 
-		Task<BlobId[]> FilterOutKnownBlobsAsync(NamespaceId ns, IEnumerable<BlobId> blobs);
-		Task<BlobContents> GetObjectsAsync(NamespaceId ns, BlobId[] refRequestBlobReferences);
+		Task<BlobId[]> FilterOutKnownBlobsAsync(NamespaceId ns, IEnumerable<BlobId> blobs, CancellationToken cancellationToken = default);
+		Task<BlobContents> GetObjectsAsync(NamespaceId ns, BlobId[] refRequestBlobReferences, CancellationToken cancellationToken = default);
 	}
 
 	public class BlobNotFoundException : Exception
@@ -93,39 +97,6 @@ namespace Horde.Server.Ddc
 
 	public static class BlobServiceExtensions
 	{
-		public static async Task VerifyContentMatchesHashAsync(Stream stream, IoHash expectedHash, Tracer tracer)
-		{
-			IoHash hash;
-			using (TelemetrySpan _ = tracer.StartActiveSpan("web.hash").SetAttribute("operation.name", "web.hash"))
-			{
-				hash = await IoHash.ComputeAsync(stream);
-			}
-			if (hash != expectedHash)
-			{
-				throw new HashMismatchException(hash, expectedHash);
-			}
-		}
-
-		public static async Task<BlobId> PutObjectAsync(this IDdcBlobService blobService, NamespaceId ns, BufferedPayload payload, BlobId identifier, Tracer tracer)
-		{
-			using TelemetrySpan scope = tracer.StartActiveSpan("put_blob")
-				.SetAttribute("operation.name", "put_blob")
-				.SetAttribute("resource.name", identifier.ToString())
-				.SetAttribute("Content-Length", payload.Length.ToString());
-
-			await using Stream hashStream = payload.GetStream();
-			await VerifyContentMatchesHashAsync(hashStream, identifier.Hash, tracer);
-
-			await blobService.PutObjectKnownHashAsync(ns, payload, identifier);
-			return identifier;
-		}
-
-		public static Task<BlobId> PutObjectAsync(this IDdcBlobService blobService, NamespaceId ns, byte[] payload, BlobId identifier, Tracer tracer)
-		{
-			using MemoryBufferedPayload bufferedPayload = new MemoryBufferedPayload(payload);
-			return PutObjectAsync(blobService, ns, bufferedPayload, identifier, tracer);
-		}
-
 		public static async Task<ContentId> PutCompressedObject(this IDdcBlobService blobService, NamespaceId ns, BufferedPayload payload, ContentId? id, IServiceProvider provider)
 		{
 			IDdcContentIdService contentIdStore = provider.GetService<IDdcContentIdService>()!;
diff --git a/Engine/Source/Programs/Horde/Horde.Server/Ddc/IDdcRefService.cs b/Engine/Source/Programs/Horde/Horde.Server/Ddc/IDdcRefService.cs
--- a/Engine/Source/Programs/Horde/Horde.Server/Ddc/IDdcRefService.cs
+++ b/Engine/Source/Programs/Horde/Horde.Server/Ddc/IDdcRefService.cs
@@ -2,6 +2,7 @@
 
 using System;
 using System.Collections.Generic;
+using System.Threading;
 using System.Threading.Tasks;
 using EpicGames.Core;
 using EpicGames.Horde.Storage;
@@ -13,14 +14,14 @@ namespace Horde.Server.Ddc
 {
 	public interface IDdcRefService
     {
-        Task<(RefRecord, BlobContents?)> GetAsync(NamespaceId ns, BucketId bucket, RefId key, string[] fields, bool doLastAccessTracking = true);
-        Task<(ContentId[], BlobId[])> PutAsync(NamespaceId ns, BucketId bucket, RefId key, BlobId blobHash, CbObject payload);
-        Task<(ContentId[], BlobId[])> FinalizeAsync(NamespaceId ns, BucketId bucket, RefId key, BlobId blobHash);
+        Task<(RefRecord, BlobContents?)> GetAsync(NamespaceId ns, BucketId bucket, RefId key, string[] fields, bool doLastAccessTracking = true, CancellationToken cancellationToken = default);
+        Task<(ContentId[], BlobId[])> PutAsync(NamespaceId ns, BucketId bucket, RefId key, BlobId blobHash, CbObject payload, CancellationToken cancellationToken = default);
+        Task<(ContentId[], BlobId[])> FinalizeAsync(NamespaceId ns, BucketId bucket, RefId key, BlobId blobHash, CancellationToken cancellationToken = default);
 
-        Task<bool> DeleteAsync(NamespaceId ns, BucketId bucket, RefId key);
+        Task<bool> DeleteAsync(NamespaceId ns, BucketId bucket, RefId key, CancellationToken cancellationToken = default);
 
-        Task<bool> ExistsAsync(NamespaceId ns, BucketId bucket, RefId key);
-        Task<List<BlobId>> GetReferencedBlobsAsync(NamespaceId ns, BucketId bucket, RefId key);
+        Task<bool> ExistsAsync(NamespaceId ns, BucketId bucket, RefId key, CancellationToken cancellationToken = default);
+        Task<List<BlobId>> GetReferencedBlobsAsync(NamespaceId ns, BucketId bucket, RefId key, CancellationToken cancellationToken = default);
     }
 
 	public class RefRecord
diff --git a/Engine/Source/Programs/Horde/Horde.Server/Ddc/RequestHelper.cs b/Engine/Source/Programs/Horde/Horde.Server/Ddc/RequestHelper.cs
--- a/Engine/Source/Programs/Horde/Horde.Server/Ddc/RequestHelper.cs
+++ b/Engine/Source/Programs/Horde/Horde.Server/Ddc/RequestHelper.cs
@@ -1,10 +1,13 @@
 // Copyright Epic Games, Inc. All Rights Reserved.
 
 using System;
+using System.Linq;
 using System.Security.Claims;
+using System.Threading;
 using System.Threading.Tasks;
 using EpicGames.Horde.Storage;
 using Horde.Server.Acls;
+using Horde.Server.Storage;
 using Microsoft.AspNetCore.Http;
 using Microsoft.AspNetCore.Mvc;
 
@@ -12,14 +15,25 @@ namespace Horde.Server.Ddc
 {
 	class RequestHelper : IRequestHelper
 	{
-		public Task<ActionResult?> HasAccessForGlobalOperationsAsync(ClaimsPrincipal user, AclAction[] aclActions)
+		readonly StorageService _storageService;
+
+		public RequestHelper(StorageService storageService)
 		{
-			throw new NotImplementedException();
+			_storageService = storageService;
 		}
 
-		public Task<ActionResult?> HasAccessToNamespaceAsync(ClaimsPrincipal user, HttpRequest request, NamespaceId ns, AclAction[] aclActions)
+		public async Task<ActionResult?> HasAccessToNamespaceAsync(ClaimsPrincipal user, HttpRequest request, NamespaceId ns, AclAction[] aclActions)
 		{
-			throw new NotImplementedException();
+			StorageClient? storageClient = await _storageService.TryGetClientAsync(ns, CancellationToken.None);
+			if (storageClient == null)
+			{
+				return new ForbidResult();
+			}
+			if (!aclActions.Any(x => !storageClient.Config.Authorize(x, user)))
+			{
+				return new ForbidResult();
+			}
+			return null;
 		}
 	}
 }
diff --git a/Engine/Source/Programs/Horde/Horde.Server/Storage/StorageController.cs b/Engine/Source/Programs/Horde/Horde.Server/Storage/StorageController.cs
--- a/Engine/Source/Programs/Horde/Horde.Server/Storage/StorageController.cs
+++ b/Engine/Source/Programs/Horde/Horde.Server/Storage/StorageController.cs
@@ -191,17 +191,16 @@ namespace Horde.Server.Storage
 		[Route("/api/v1/storage/{namespaceId}/bundles")]
 		public async Task<ActionResult<WriteBlobResponse>> WriteBlobAsync(NamespaceId namespaceId, IFormFile? file, [FromForm] string? prefix = default, CancellationToken cancellationToken = default)
 		{
-			NamespaceConfig? namespaceConfig;
-			if (!_globalConfig.Value.Storage.TryGetNamespace(namespaceId, out namespaceConfig))
+			StorageClient? storageClient = await _storageService.TryGetClientAsync(namespaceId, cancellationToken);
+			if (storageClient == null)
 			{
 				return NotFound(namespaceId);
 			}
-			if (!namespaceConfig.Authorize(StorageAclAction.WriteBlobs, User) && !HasPathClaim(User, HordeClaimTypes.WriteNamespace, namespaceId, prefix ?? String.Empty))
+			if (!storageClient.Config.Authorize(StorageAclAction.WriteBlobs, User) && !HasPathClaim(User, HordeClaimTypes.WriteNamespace, namespaceId, prefix ?? String.Empty))
 			{
 				return Forbid(StorageAclAction.WriteBlobs, namespaceId);
 			}
 
-			StorageClient storageClient = await _storageService.GetClientAsync(namespaceId, cancellationToken);
 			return await WriteBlobAsync(storageClient, file, prefix, cancellationToken);
 		}
 
@@ -257,25 +256,16 @@ namespace Horde.Server.Storage
 		[Route("/api/v1/storage/{namespaceId}/bundles/{*locator}")]
 		public async Task<ActionResult> ReadBlobAsync(NamespaceId namespaceId, BundleLocator locator, [FromQuery] int? offset = null, [FromQuery] int? length = null, CancellationToken cancellationToken = default)
 		{
-			NamespaceConfig? namespaceConfig;
-			if (!_globalConfig.Value.Storage.TryGetNamespace(namespaceId, out namespaceConfig))
+			StorageClient? client = await _storageService.TryGetClientAsync(namespaceId, cancellationToken);
+			if (client == null)
 			{
 				return NotFound(namespaceId);
 			}
-			if (!namespaceConfig.Authorize(StorageAclAction.ReadBlobs, User) && !HasPathClaim(User, HordeClaimTypes.ReadNamespace, namespaceId, locator.Path.ToString()))
+			if (!client.Config.Authorize(StorageAclAction.ReadBlobs, User) && !HasPathClaim(User, HordeClaimTypes.ReadNamespace, namespaceId, locator.Path.ToString()))
 			{
 				return Forbid(StorageAclAction.ReadBlobs, namespaceId);
 			}
 
-			return await ReadBlobInternalAsync(_storageService, namespaceId, locator, offset, length, cancellationToken);
-		}
-
-		/// <summary>
-		/// Reads a blob from storage, without performing namespace access checks.
-		/// </summary>
-		internal static async Task<ActionResult> ReadBlobInternalAsync(StorageService storageService, NamespaceId namespaceId, BundleLocator locator, int? offset, int? length, CancellationToken cancellationToken)
-		{
-			StorageClient client = await storageService.GetClientAsync(namespaceId, cancellationToken);
 			return await ReadBlobInternalAsync(client, locator, offset, length, cancellationToken);
 		}
 
@@ -324,18 +314,16 @@ namespace Horde.Server.Storage
 		[Route("/api/v1/storage/{namespaceId}/nodes")]
 		public async Task<ActionResult<FindNodesResponse>> FindNodesAsync(NamespaceId namespaceId, string alias, CancellationToken cancellationToken = default)
 		{
-			NamespaceConfig? namespaceConfig;
-			if (!_globalConfig.Value.Storage.TryGetNamespace(namespaceId, out namespaceConfig))
+			StorageClient? client = await _storageService.TryGetClientAsync(namespaceId, cancellationToken);
+			if (client == null)
 			{
 				return NotFound(namespaceId);
 			}
-			if (!namespaceConfig.Authorize(StorageAclAction.ReadBlobs, User))
+			if (!client.Config.Authorize(StorageAclAction.ReadBlobs, User))
 			{
 				return Forbid(StorageAclAction.ReadBlobs, namespaceId);
 			}
 
-			StorageClient client = await _storageService.GetClientAsync(namespaceId, cancellationToken);
-
 			FindNodesResponse response = new FindNodesResponse();
 			await foreach (BundleNodeHandle handle in client.FindAliasAsync(alias, cancellationToken))
 			{
@@ -361,17 +349,16 @@ namespace Horde.Server.Storage
 		[Route("/api/v1/storage/{namespaceId}/refs/{*refName}")]
 		public async Task<ActionResult> WriteRefAsync(NamespaceId namespaceId, RefName refName, [FromBody] WriteRefRequest request, CancellationToken cancellationToken)
 		{
-			NamespaceConfig? namespaceConfig;
-			if (!_globalConfig.Value.Storage.TryGetNamespace(namespaceId, out namespaceConfig))
+			StorageClient? client = await _storageService.TryGetClientAsync(namespaceId, cancellationToken);
+			if (client == null)
 			{
 				return NotFound(namespaceId);
 			}
-			if (!namespaceConfig.Authorize(StorageAclAction.WriteRefs, User) && !HasPathClaim(User, HordeClaimTypes.WriteNamespace, namespaceId, refName.ToString()))
+			if (!client.Config.Authorize(StorageAclAction.WriteRefs, User) && !HasPathClaim(User, HordeClaimTypes.WriteNamespace, namespaceId, refName.ToString()))
 			{
 				return Forbid(StorageAclAction.WriteRefs, namespaceId);
 			}
 
-			StorageClient client = await _storageService.GetClientAsync(namespaceId, cancellationToken);
 			BundleNodeLocator target = new BundleNodeLocator(request.Hash, request.Blob, request.ExportIdx);
 			await client.WriteRefTargetAsync(refName, target, request.Options, cancellationToken);
 
@@ -481,17 +468,16 @@ namespace Horde.Server.Storage
 		[Route("/api/v1/storage/{namespaceId}/bundleinfo/{*locator}")]
 		public async Task<ActionResult<object>> GetBundleAsync(NamespaceId namespaceId, BundleLocator locator, [FromQuery(Name = "imports")] bool includeImports = false, [FromQuery(Name = "exports")] bool includeExports = true, [FromQuery(Name = "packets")] bool includePackets = false, CancellationToken cancellationToken = default)
 		{
-			NamespaceConfig? namespaceConfig;
-			if (!_globalConfig.Value.Storage.TryGetNamespace(namespaceId, out namespaceConfig))
+			StorageClient? storageClient = await _storageService.TryGetClientAsync(namespaceId, cancellationToken);
+			if (storageClient == null)
 			{
 				return NotFound(namespaceId);
 			}
-			if (!namespaceConfig.Authorize(StorageAclAction.ReadBlobs, User))
+			if (!storageClient.Config.Authorize(StorageAclAction.ReadBlobs, User))
 			{
 				return Forbid(StorageAclAction.ReadBlobs, namespaceId);
 			}
 
-			StorageClient storageClient = await _storageService.GetClientAsync(namespaceId, cancellationToken);
 			BundleReader reader = new BundleReader(storageClient, _memoryCache, _logger);
 
 			BundleHeader header = await reader.ReadBundleHeaderAsync(locator, cancellationToken);
diff --git a/Engine/Source/Programs/Horde/Horde.Server/Storage/StorageService.cs b/Engine/Source/Programs/Horde/Horde.Server/Storage/StorageService.cs
--- a/Engine/Source/Programs/Horde/Horde.Server/Storage/StorageService.cs
+++ b/Engine/Source/Programs/Horde/Horde.Server/Storage/StorageService.cs
@@ -556,6 +556,18 @@ namespace Horde.Server.Storage
 			return namespaceInfo.Client;
 		}
 
+		/// <summary>
+		/// Attempts to gets a storage client for the given namespace
+		/// </summary>
+		/// <param name="namespaceId">Namespace identifier</param>
+		/// <param name="cancellationToken">Cancellation token for the operation</param>
+		/// <returns></returns>
+		public async ValueTask<StorageClient?> TryGetClientAsync(NamespaceId namespaceId, CancellationToken cancellationToken)
+		{
+			NamespaceInfo? namespaceInfo = await TryGetNamespaceInfoAsync(namespaceId, cancellationToken);
+			return namespaceInfo?.Client;
+		}
+
 		/// <inheritdoc/>
 		async ValueTask<IStorageClient> IStorageClientFactory.GetClientAsync(NamespaceId namespaceId, CancellationToken cancellationToken) => await GetClientAsync(namespaceId, cancellationToken);
 
