Horde: Use standard ACL scopes for DDC access.

- Also remove Ddc prefix from class names; makes the code more confusing to read.

diff --git a/Engine/Source/Programs/Horde/Horde.Server/Ddc/BlobService.cs b/Engine/Source/Programs/Horde/Horde.Server/Ddc/BlobService.cs
new file mode 100644
--- /dev/null
+++ b/Engine/Source/Programs/Horde/Horde.Server/Ddc/BlobService.cs
@@ -0,0 +1,152 @@
+// Copyright Epic Games, Inc. All Rights Reserved.
+
+using System;
+using System.Collections.Generic;
+using System.IO;
+using System.Linq;
+using System.Threading;
+using System.Threading.Tasks;
+using EpicGames.Core;
+using EpicGames.Horde.Storage;
+using Horde.Server.Storage;
+using OpenTelemetry.Trace;
+
+namespace Horde.Server.Ddc
+{
+	class BlobService : IBlobService
+	{
+		BlobType s_rawBlobType = new BlobType(new Guid("{03E6C37B-33C1-491F-8541-D3C401B8B8EF}"), 1);
+
+		readonly StorageService _storageService;
+		readonly Tracer _tracer;
+
+		public BlobService(StorageService storageService, Tracer tracer)
+		{
+			_storageService = storageService;
+			_tracer = tracer;
+		}
+
+		public static Utf8String GetAlias(BlobId blobId) => $"ddc:{blobId}";
+
+		public async Task VerifyContentMatchesHashAsync(Stream stream, IoHash expectedHash, CancellationToken cancellationToken)
+		{
+			IoHash hash;
+			using (TelemetrySpan _ = _tracer.StartActiveSpan("web.hash").SetAttribute("operation.name", "web.hash"))
+			{
+				hash = await IoHash.ComputeAsync(stream, cancellationToken);
+			}
+			if (hash != expectedHash)
+			{
+				throw new HashMismatchException(hash, expectedHash);
+			}
+		}
+
+		public async Task<BlobId> PutObjectAsync(NamespaceId ns, BufferedPayload payload, BlobId identifier, CancellationToken cancellationToken)
+		{
+			using TelemetrySpan scope = _tracer.StartActiveSpan("put_blob")
+				.SetAttribute("operation.name", "put_blob")
+				.SetAttribute("resource.name", identifier.ToString())
+				.SetAttribute("Content-Length", payload.Length.ToString());
+
+			await using Stream hashStream = payload.GetStream();
+			await VerifyContentMatchesHashAsync(hashStream, identifier.Hash, cancellationToken);
+
+			await PutObjectKnownHashAsync(ns, payload, identifier, cancellationToken);
+			return identifier;
+		}
+
+		public Task<BlobId> PutObjectAsync(NamespaceId ns, ReadOnlyMemory<byte> payload, BlobId identifier, CancellationToken cancellationToken)
+		{
+			using MemoryBufferedPayload bufferedPayload = new MemoryBufferedPayload(payload);
+			return PutObjectAsync(ns, bufferedPayload, identifier, cancellationToken);
+		}
+
+		public async Task<bool> ExistsAsync(NamespaceId ns, BlobId blob, List<string>? storageLayers, CancellationToken cancellationToken)
+		{
+			IStorageClient storageClient = await _storageService.GetClientAsync(ns, cancellationToken);
+			return await storageClient.FindAliasAsync(GetAlias(blob), cancellationToken).AnyAsync(cancellationToken);
+		}
+
+		public async Task<BlobId[]> FilterOutKnownBlobsAsync(NamespaceId ns, IEnumerable<BlobId> blobIds, CancellationToken cancellationToken)
+		{
+			IStorageClient storageClient = await _storageService.GetClientAsync(ns, cancellationToken);
+
+			List<BlobId> unknownBlobIds = new List<BlobId>();
+			foreach (BlobId blobId in blobIds)
+			{
+				if (!await storageClient.FindAliasAsync(GetAlias(blobId), cancellationToken).AnyAsync())
+				{
+					unknownBlobIds.Add(blobId);
+				}
+			}
+
+			return unknownBlobIds.ToArray();
+		}
+
+		public async Task<BlobContents> GetObjectAsync(NamespaceId ns, BlobId blob, List<string>? storageLayers, bool supportsRedirectUri, CancellationToken cancellationToken)
+		{
+			IStorageClient storageClient = await _storageService.GetClientAsync(ns, cancellationToken);
+
+			BlobHandle? handle = await storageClient.FindAliasAsync(GetAlias(blob), cancellationToken).FirstOrDefaultAsync(cancellationToken);
+			if (handle == null)
+			{
+				throw new BlobNotFoundException(ns, blob);
+			}
+
+			BlobData data = await handle.ReadAsync(cancellationToken);
+			return new BlobContents(data.Data.ToArray());
+		}
+
+		public async Task<BlobContents> GetObjectsAsync(NamespaceId ns, BlobId[] blobs, CancellationToken cancellationToken)
+		{
+			using TelemetrySpan _ = _tracer.StartActiveSpan("blob.combine").SetAttribute("operation.name", "blob.combine");
+			Task<BlobContents>[] tasks = new Task<BlobContents>[blobs.Length];
+			for (int i = 0; i < blobs.Length; i++)
+			{
+				tasks[i] = GetObjectAsync(ns, blobs[i], storageLayers: null, supportsRedirectUri: false, cancellationToken);
+			}
+
+			MemoryStream ms = new MemoryStream();
+			foreach (Task<BlobContents> task in tasks)
+			{
+				BlobContents blob = await task;
+				await using Stream s = blob.Stream;
+				await s.CopyToAsync(ms, cancellationToken);
+			}
+
+			ms.Seek(0, SeekOrigin.Begin);
+
+			return new BlobContents(ms, ms.Length);
+		}
+
+		public Task<Uri?> GetObjectWithRedirectAsync(NamespaceId ns, BlobId blobIdentifier, List<string>? storageLayers, CancellationToken cancellationToken)
+		{
+			return Task.FromResult<Uri?>(null);
+		}
+
+		public Task<Uri?> MaybePutObjectWithRedirectAsync(NamespaceId ns, BlobId identifier, CancellationToken cancellationToken)
+		{
+			return Task.FromResult<Uri?>(null);
+		}
+
+		public async Task<BlobId> PutObjectKnownHashAsync(NamespaceId ns, BufferedPayload content, BlobId identifier, CancellationToken cancellationToken)
+		{
+			IStorageClient storageClient = await _storageService.GetClientAsync(ns, cancellationToken);
+
+			BlobHandle blobHandle;
+			await using (IStorageWriter writer = storageClient.CreateWriter())
+			{
+				Memory<byte> memory = writer.GetOutputBuffer(0, (int)content.Length);
+
+				using Stream stream = content.GetStream();
+				await stream.ReadAsync(memory, cancellationToken);
+
+				blobHandle = await writer.WriteNodeAsync(memory.Length, Array.Empty<BlobHandle>(), s_rawBlobType, cancellationToken);
+				await writer.FlushAsync(cancellationToken);
+			}
+
+			await storageClient.AddAliasAsync(GetAlias(identifier), blobHandle, cancellationToken);
+			return identifier;
+		}
+	}
+}
diff --git a/Engine/Source/Programs/Horde/Horde.Server/Ddc/BlobsController.cs b/Engine/Source/Programs/Horde/Horde.Server/Ddc/BlobsController.cs
new file mode 100644
--- /dev/null
+++ b/Engine/Source/Programs/Horde/Horde.Server/Ddc/BlobsController.cs
@@ -0,0 +1,381 @@
+// Copyright Epic Games, Inc. All Rights Reserved.
+
+using System;
+using System.Collections.Concurrent;
+using System.Collections.Generic;
+using System.ComponentModel.DataAnnotations;
+using System.IO;
+using System.Linq;
+using System.Net;
+using System.Net.Mime;
+using System.Text.Json.Serialization;
+using System.Threading;
+using System.Threading.Tasks;
+using EpicGames.AspNet;
+using EpicGames.Core;
+using EpicGames.Horde.Storage;
+using EpicGames.Serialization;
+using Horde.Server.Acls;
+using Horde.Server.Storage;
+using Microsoft.AspNetCore.Authorization;
+using Microsoft.AspNetCore.Http;
+using Microsoft.AspNetCore.Mvc;
+using OpenTelemetry.Trace;
+
+#pragma warning disable CS1591
+
+namespace Horde.Server.Ddc
+{
+	using IDiagnosticContext = Serilog.IDiagnosticContext;
+
+    [ApiController]
+    [Route("api/v1/s", Order = 1)]
+    [Route("api/v1/blobs", Order = 0)]
+    [Authorize]
+    public class BlobsController : ControllerBase
+    {
+        private readonly IBlobService _storage;
+        private readonly IDiagnosticContext _diagnosticContext;
+        private readonly IRequestHelper _requestHelper;
+        private readonly BufferedPayloadFactory _bufferedPayloadFactory;
+        private readonly NginxRedirectHelper _nginxRedirectHelper;
+		private readonly Tracer _tracer;
+
+        public BlobsController(IBlobService storage, IDiagnosticContext diagnosticContext, IRequestHelper requestHelper, BufferedPayloadFactory bufferedPayloadFactory, NginxRedirectHelper nginxRedirectHelper, Tracer tracer)
+        {
+            _storage = storage;
+            _diagnosticContext = diagnosticContext;
+            _requestHelper = requestHelper;
+            _bufferedPayloadFactory = bufferedPayloadFactory;
+            _nginxRedirectHelper = nginxRedirectHelper;
+			_tracer = tracer;
+        }
+
+        [HttpGet("{ns}/{id}")]
+        [ProducesDefaultResponseType]
+        public async Task<IActionResult> Get(
+            [Required] NamespaceId ns,
+            [Required] BlobId id,
+            [FromQuery] List<string>? storageLayers = null)
+        {
+            ActionResult? result = await _requestHelper.HasAccessToNamespaceAsync(User, Request, ns, new [] { StorageAclAction.ReadBlobs });
+            if (result != null)
+            {
+                return result;
+            }
+
+            try
+            {
+                BlobContents blobContents = await GetImpl(ns, id, storageLayers, supportsRedirectUri: true);
+
+                if (blobContents.RedirectUri != null)
+                {
+                    return Redirect(blobContents.RedirectUri.ToString());
+                }
+                if (_nginxRedirectHelper.CanRedirect(Request, blobContents))
+                {
+                    return _nginxRedirectHelper.CreateActionResult(blobContents, MediaTypeNames.Application.Octet);
+                }
+                return File(blobContents.Stream, MediaTypeNames.Application.Octet, enableRangeProcessing: true);
+            }
+            catch (BlobNotFoundException e)
+            {
+                return NotFound(new ValidationProblemDetails { Title = $"Blob {e.Blob} not found"});
+            }
+            catch (AuthorizationException e)
+            {
+                return e.Result;
+            }
+        }
+        
+        [HttpHead("{ns}/{id}")]
+        [ProducesDefaultResponseType]
+        public async Task<IActionResult> Head(
+            [Required] NamespaceId ns,
+            [Required] BlobId id,
+            [FromQuery] List<string>? storageLayers = null)
+        {
+            ActionResult? result = await _requestHelper.HasAccessToNamespaceAsync(User, Request, ns, new [] { StorageAclAction.ReadBlobs });
+            if (result != null)
+            {
+                return result;
+            }
+            bool exists = await _storage.ExistsAsync(ns, id, storageLayers);
+
+            if (!exists)
+            {
+                return NotFound(new ValidationProblemDetails { Title = $"Blob {id} not found"});
+            }
+
+            return Ok();
+        }
+
+        [HttpPost("{ns}/exists")]
+        [ProducesDefaultResponseType]
+        public async Task<IActionResult> ExistsMultiple(
+            [Required] NamespaceId ns,
+            [Required] [FromQuery] List<BlobId> id)
+        {
+            ActionResult? result = await _requestHelper.HasAccessToNamespaceAsync(User, Request, ns, new [] { StorageAclAction.ReadBlobs });
+            if (result != null)
+            {
+                return result;
+            }
+
+            ConcurrentBag<BlobId> missingBlobs = new ConcurrentBag<BlobId>();
+
+            IEnumerable<Task> tasks = id.Select(async blob =>
+            {
+                if (!await _storage.ExistsAsync(ns, blob))
+                {
+                    missingBlobs.Add(blob);
+                }
+            });
+            await Task.WhenAll(tasks);
+
+            return Ok(new HeadMultipleResponse { Needs = missingBlobs.ToArray()});
+        }
+
+        [HttpPost("{ns}/exist")]
+        [ProducesDefaultResponseType]
+        public async Task<IActionResult> ExistsBody(
+            [Required] NamespaceId ns,
+            [FromBody] BlobId[] bodyIds)
+        {
+            ActionResult? result = await _requestHelper.HasAccessToNamespaceAsync(User, Request, ns, new [] { StorageAclAction.ReadBlobs });
+            if (result != null)
+            {
+                return result;
+            }
+
+            ConcurrentBag<BlobId> missingBlobs = new ConcurrentBag<BlobId>();
+
+            IEnumerable<Task> tasks = bodyIds.Select(async blob =>
+            {
+                if (!await _storage.ExistsAsync(ns, blob))
+                {
+                    missingBlobs.Add(blob);
+                }
+            });
+            await Task.WhenAll(tasks);
+
+            return Ok(new HeadMultipleResponse { Needs = missingBlobs.ToArray()});
+        }
+
+        private async Task<BlobContents> GetImpl(NamespaceId ns, BlobId blob, List<string>? storageLayers = null, bool supportsRedirectUri = false)
+        {
+            return await _storage.GetObjectAsync(ns, blob, storageLayers, supportsRedirectUri);
+        }
+
+        [HttpPut("{ns}/{id}")]
+        [RequiredContentType(MediaTypeNames.Application.Octet)]
+        [DisableRequestSizeLimit]
+        public async Task<IActionResult> Put(
+            [Required] NamespaceId ns,
+            [Required] BlobId id,
+			CancellationToken cancellationToken)
+        {
+            ActionResult? result = await _requestHelper.HasAccessToNamespaceAsync(User, Request, ns, new [] { StorageAclAction.WriteBlobs });
+            if (result != null)
+            {
+                return result;
+            }
+
+            _diagnosticContext.Set("Content-Length", Request.ContentLength ?? -1);
+
+            try
+            {
+                Uri? uri = await _storage.MaybePutObjectWithRedirectAsync(ns, id, cancellationToken);
+                if (uri != null)
+                {
+                    return Ok(new
+                    {
+                        Identifier = id.ToString(),
+                        RedirectUri = uri,
+                    });
+                }
+                using BufferedPayload payload = await _bufferedPayloadFactory.CreateFromRequest(Request);
+
+                BlobId identifier = await _storage.PutObjectAsync(ns, payload, id, cancellationToken);
+                return Ok(new
+                {
+                    Identifier = identifier.ToString()
+                });
+            }
+            catch (ResourceHasToManyRequestsException)
+            {
+                return StatusCode(StatusCodes.Status429TooManyRequests);
+            }
+            catch (ClientSendSlowException e)
+            {
+                return Problem(e.Message, null, (int)HttpStatusCode.RequestTimeout);
+            }
+        }
+
+        [HttpPost("{ns}")]
+        [RequiredContentType(MediaTypeNames.Application.Octet)]
+        [DisableRequestSizeLimit]
+        public async Task<IActionResult> Post(
+            [Required] NamespaceId ns)
+        {
+            ActionResult? result = await _requestHelper.HasAccessToNamespaceAsync(User, Request, ns, new [] { StorageAclAction.WriteBlobs });
+            if (result != null)
+            {
+                return result;
+            }
+
+            _diagnosticContext.Set("Content-Length", Request.ContentLength ?? -1);
+            try
+            {
+                using BufferedPayload payload = await _bufferedPayloadFactory.CreateFromRequest(Request);
+
+                await using Stream stream = payload.GetStream();
+
+				BlobId id = new BlobId(await IoHash.ComputeAsync(stream));
+                await _storage.PutObjectKnownHashAsync(ns, payload, id);
+                
+                return Ok(new
+                {
+                    Identifier = id.ToString()
+                });
+            }
+            catch (ClientSendSlowException e)
+            {
+                return Problem(e.Message, null, (int)HttpStatusCode.RequestTimeout);
+            }
+        }
+
+        // ReSharper disable UnusedAutoPropertyAccessor.Global
+        // ReSharper disable once ClassNeverInstantiated.Global
+        [System.Diagnostics.CodeAnalysis.SuppressMessage("Design", "CA1034:Nested types should not be visible", Justification = "Used by serialization")]
+        public class BatchOp
+        {
+            // ReSharper disable once InconsistentNaming
+            public enum Operation
+            {
+                INVALID,
+                GET,
+                PUT,
+                HEAD
+            }
+
+            [Required] public NamespaceId? Namespace { get; set; }
+
+            public BlobId? Id { get; set; }
+
+            [Required] public Operation Op { get; set; }
+
+            public byte[]? Content { get; set; }
+        }
+
+        [System.Diagnostics.CodeAnalysis.SuppressMessage("Design", "CA1034:Nested types should not be visible", Justification = "Used by serialization")]
+        public class BatchCall
+        {
+            public BatchOp[]? Operations { get; set; }
+        }
+        // ReSharper restore UnusedAutoPropertyAccessor.Global
+
+        [HttpPost("")]
+        public async Task<IActionResult> Post([FromBody] BatchCall batch, CancellationToken cancellationToken)
+        {
+            static AclAction MapToAclAction(BatchOp.Operation op)
+            {
+                switch (op)
+                {
+                    case BatchOp.Operation.GET:
+                    case BatchOp.Operation.HEAD:
+                        return StorageAclAction.ReadBlobs;
+                    case BatchOp.Operation.PUT:
+						return StorageAclAction.WriteBlobs;
+                    default:
+                        throw new ArgumentOutOfRangeException(nameof(op), op, null);
+                }
+            }
+
+            if (batch?.Operations == null)
+            {
+                throw new ();
+            }
+
+            Task<object?>[] tasks = new Task<object?>[batch.Operations.Length];
+            for (int index = 0; index < batch.Operations.Length; index++)
+            {
+                BatchOp op = batch.Operations[index];
+                if (op.Namespace == null)
+                {
+                    throw new Exception(nameof(op.Namespace));
+                }
+
+                ActionResult? result = await _requestHelper.HasAccessToNamespaceAsync(User, Request, op.Namespace!.Value, new[] { MapToAclAction(op.Op) });
+
+                if (result != null)
+                {
+                    return result;
+                }
+
+                switch (op.Op)
+                {
+                    case BatchOp.Operation.INVALID:
+                        throw new InvalidOperationException($"Op type {BatchOp.Operation.INVALID} is not a valid op type for operation id {op.Id}");
+                    case BatchOp.Operation.GET:
+                        if (op.Id == null)
+                        {
+                            return BadRequest();
+                        }
+
+                        tasks[index] = GetImpl(op.Namespace.Value, op.Id.Value).ContinueWith((t, _) =>
+                        {
+                            // TODO: This is very allocation heavy but given that the end result is a json object we can not really stream this anyway
+                            using BlobContents blobContents = t.Result;
+
+                            using MemoryStream ms = new MemoryStream();
+                            blobContents.Stream.CopyTo(ms);
+                            ms.Seek(0, SeekOrigin.Begin);
+                            string str = Convert.ToBase64String(ms.ToArray());
+                            return (object?) str;
+                        }, null, TaskScheduler.Current);
+                        break;
+                    case BatchOp.Operation.HEAD:
+                        if (op.Id == null)
+                        {
+                            return BadRequest();
+                        }
+
+                        tasks[index] = _storage.ExistsAsync(op.Namespace.Value, op.Id.Value)
+                            .ContinueWith((t,_) => t.Result ? (object?) null : op.Id, null, TaskScheduler.Current);
+                        break;
+                    case BatchOp.Operation.PUT:
+                        {
+                            if (op.Content == null)
+                            {
+                                return BadRequest();
+                            }
+
+                            if (op.Id == null)
+                            {
+                                return BadRequest();
+                            }
+
+                            using MemoryBufferedPayload payload = new MemoryBufferedPayload(op.Content);
+                            tasks[index] = _storage.PutObjectAsync(op.Namespace.Value, payload, op.Id.Value, cancellationToken).ContinueWith((t, _) => (object?) t.Result, null, TaskScheduler.Current);
+                            break;
+                        }
+                    default:
+                        throw new NotImplementedException($"{op.Op} is not a support op type");
+                }
+            }
+
+            await Task.WhenAll(tasks);
+
+            object?[] results = tasks.Select(t => t.Result).ToArray();
+
+            return Ok(results);
+        }
+    }
+
+    public class HeadMultipleResponse
+    {
+        [CbField("needs")]
+        public BlobId[] Needs { get; set; } = null!;
+    }
+}
diff --git a/Engine/Source/Programs/Horde/Horde.Server/Ddc/ContentIdController.cs b/Engine/Source/Programs/Horde/Horde.Server/Ddc/ContentIdController.cs
new file mode 100644
--- /dev/null
+++ b/Engine/Source/Programs/Horde/Horde.Server/Ddc/ContentIdController.cs
@@ -0,0 +1,88 @@
+// Copyright Epic Games, Inc. All Rights Reserved.
+
+using System.Net.Mime;
+using System.Threading.Tasks;
+using EpicGames.AspNet;
+using EpicGames.Horde.Storage;
+using Horde.Server.Storage;
+using Microsoft.AspNetCore.Authorization;
+using Microsoft.AspNetCore.Mvc;
+
+#pragma warning disable CS1591
+
+namespace Horde.Server.Ddc
+{
+	[ApiController]
+    [FormatFilter]
+    [Produces(MediaTypeNames.Application.Json, MediaTypeNames.Application.Octet, CustomMediaTypeNames.UnrealCompactBinary)]
+    [Route("api/v1/content-id")]
+    [Authorize]
+    public class ContentIdController : ControllerBase
+    {
+        private readonly IRequestHelper _requestHelper;
+        private readonly IContentIdService _contentIdStore;
+
+        public ContentIdController(IRequestHelper requestHelper, IContentIdService contentIdStore)
+        {
+            _requestHelper = requestHelper;
+            _contentIdStore = contentIdStore;
+        }
+
+        /// <summary>
+        /// Returns which blobs a content id maps to
+        /// </summary>
+        /// <param name="ns">Namespace. Each namespace is completely separated from each other. Use for different types of data that is never expected to be similar (between two different games for instance). Example: `uc4.ddc`</param>
+        /// <param name="contentId">The content id to resolve </param>
+        [HttpGet("{ns}/{contentId}.{format?}", Order = 500)]
+        [ProducesDefaultResponseType]
+        [ProducesResponseType(type: typeof(ProblemDetails), 400)]
+        public async Task<IActionResult> Resolve(NamespaceId ns, ContentId contentId)
+        {
+            ActionResult? result = await _requestHelper.HasAccessToNamespaceAsync(User, Request, ns, new [] { StorageAclAction.ReadBlobs });
+            if (result != null)
+            {
+                return result;
+            }
+            BlobId[]? blobs = await _contentIdStore.ResolveAsync(ns, contentId, mustBeContentId: true);
+
+            if (blobs == null)
+            {
+                return NotFound(new ProblemDetails { Title = $"Unable to resolve content id {contentId} ({ns})." });
+            }
+
+            return Ok(new ResolvedContentIdResponse(blobs));
+        }
+
+        /// <summary>
+        /// Update a content id mapping
+        /// </summary>
+        /// <param name="ns">Namespace. Each namespace is completely separated from each other. Use for different types of data that is never expected to be similar (between two different games for instance). Example: `uc4.ddc`</param>
+        /// <param name="contentId">The content id to resolve </param>
+        /// <param name="blobIdentifier">The blob identifier to map the content id to</param>
+        /// <param name="contentWeight">The weight of this mapping, higher means more important</param>
+        [HttpPut("{ns}/{contentId}/update/{blobIdentifier}/{contentWeight}", Order = 500)]
+        [ProducesDefaultResponseType]
+        [ProducesResponseType(type: typeof(ProblemDetails), 400)]
+        public async Task<IActionResult> UpdateContentIdMapping(NamespaceId ns, ContentId contentId, BlobId blobIdentifier, int contentWeight)
+        {
+            ActionResult? result = await _requestHelper.HasAccessToNamespaceAsync(User, Request, ns, new [] { StorageAclAction.WriteBlobs });
+            if (result != null)
+            {
+                return result;
+            }
+            await _contentIdStore.PutAsync(ns, contentId, blobIdentifier, contentWeight);
+
+            return Ok();
+        }
+    }
+
+    public class ResolvedContentIdResponse
+    {
+        public BlobId[] Blobs { get; set; }
+
+        public ResolvedContentIdResponse(BlobId[] blobs)
+        {
+            Blobs = blobs;
+        }
+    }
+}
diff --git a/Engine/Source/Programs/Horde/Horde.Server/Ddc/ContentIdService.cs b/Engine/Source/Programs/Horde/Horde.Server/Ddc/ContentIdService.cs
new file mode 100644
--- /dev/null
+++ b/Engine/Source/Programs/Horde/Horde.Server/Ddc/ContentIdService.cs
@@ -0,0 +1,21 @@
+// Copyright Epic Games, Inc. All Rights Reserved.
+
+using EpicGames.Horde.Storage;
+using System;
+using System.Threading.Tasks;
+
+namespace Horde.Server.Ddc
+{
+	class ContentIdService : IContentIdService
+	{
+		public Task<BlobId[]?> ResolveAsync(NamespaceId ns, ContentId contentId, bool mustBeContentId = false)
+		{
+			throw new NotImplementedException();
+		}
+
+		public Task PutAsync(NamespaceId ns, ContentId contentId, BlobId blobIdentifier, int contentWeight)
+		{
+			throw new NotImplementedException();
+		}
+	}
+}
diff --git a/Engine/Source/Programs/Horde/Horde.Server/Ddc/DdcAclAction.cs b/Engine/Source/Programs/Horde/Horde.Server/Ddc/DdcAclAction.cs
deleted file mode 100644
diff --git a/Engine/Source/Programs/Horde/Horde.Server/Ddc/DdcBlobService.cs b/Engine/Source/Programs/Horde/Horde.Server/Ddc/DdcBlobService.cs
deleted file mode 100644
diff --git a/Engine/Source/Programs/Horde/Horde.Server/Ddc/DdcBlobsController.cs b/Engine/Source/Programs/Horde/Horde.Server/Ddc/DdcBlobsController.cs
deleted file mode 100644
diff --git a/Engine/Source/Programs/Horde/Horde.Server/Ddc/DdcContentIdController.cs b/Engine/Source/Programs/Horde/Horde.Server/Ddc/DdcContentIdController.cs
deleted file mode 100644
diff --git a/Engine/Source/Programs/Horde/Horde.Server/Ddc/DdcContentIdService.cs b/Engine/Source/Programs/Horde/Horde.Server/Ddc/DdcContentIdService.cs
deleted file mode 100644
diff --git a/Engine/Source/Programs/Horde/Horde.Server/Ddc/DdcObjectsController.cs b/Engine/Source/Programs/Horde/Horde.Server/Ddc/DdcObjectsController.cs
deleted file mode 100644
diff --git a/Engine/Source/Programs/Horde/Horde.Server/Ddc/DdcRefService.cs b/Engine/Source/Programs/Horde/Horde.Server/Ddc/DdcRefService.cs
deleted file mode 100644
diff --git a/Engine/Source/Programs/Horde/Horde.Server/Ddc/DdcRefsController.cs b/Engine/Source/Programs/Horde/Horde.Server/Ddc/DdcRefsController.cs
deleted file mode 100644
diff --git a/Engine/Source/Programs/Horde/Horde.Server/Ddc/IBlobService.cs b/Engine/Source/Programs/Horde/Horde.Server/Ddc/IBlobService.cs
new file mode 100644
--- /dev/null
+++ b/Engine/Source/Programs/Horde/Horde.Server/Ddc/IBlobService.cs
@@ -0,0 +1,193 @@
+﻿// Copyright Epic Games, Inc. All Rights Reserved.
+
+using System;
+using System.Collections.Generic;
+using System.IO;
+using System.Net.Mime;
+using System.Threading;
+using System.Threading.Tasks;
+using EpicGames.AspNet;
+using EpicGames.Core;
+using EpicGames.Horde.Storage;
+using Microsoft.Extensions.DependencyInjection;
+using OpenTelemetry.Trace;
+using OpenTracing.Util;
+
+#pragma warning disable CS1591
+
+namespace Horde.Server.Ddc
+{
+	/// <summary>
+	/// Interface for the DDC blobs service
+	/// </summary>
+	public interface IBlobService
+	{
+		Task VerifyContentMatchesHashAsync(Stream stream, IoHash expectedHash, CancellationToken cancellationToken = default);
+		Task<BlobId> PutObjectAsync(NamespaceId ns, BufferedPayload payload, BlobId identifier, CancellationToken cancellationToken = default);
+		Task<BlobId> PutObjectAsync(NamespaceId ns, ReadOnlyMemory<byte> payload, BlobId identifier, CancellationToken cancellationToken = default);
+		Task<BlobId> PutObjectKnownHashAsync(NamespaceId ns, BufferedPayload content, BlobId identifier, CancellationToken cancellationToken = default);
+		Task<Uri?> MaybePutObjectWithRedirectAsync(NamespaceId ns, BlobId identifier, CancellationToken cancellationToken = default);
+
+		Task<BlobContents> GetObjectAsync(NamespaceId ns, BlobId blob, List<string>? storageLayers = null, bool supportsRedirectUri = false, CancellationToken cancellationToken = default);
+
+		Task<Uri?> GetObjectWithRedirectAsync(NamespaceId ns, BlobId blobIdentifier, List<string>? storageLayers = null, CancellationToken cancellationToken = default);
+
+		Task<bool> ExistsAsync(NamespaceId ns, BlobId blob, List<string>? storageLayers = null, CancellationToken cancellationToken = default);
+
+		Task<BlobId[]> FilterOutKnownBlobsAsync(NamespaceId ns, IEnumerable<BlobId> blobs, CancellationToken cancellationToken = default);
+		Task<BlobContents> GetObjectsAsync(NamespaceId ns, BlobId[] refRequestBlobReferences, CancellationToken cancellationToken = default);
+	}
+
+	public class BlobNotFoundException : Exception
+	{
+		public NamespaceId Ns { get; }
+		public BlobId Blob { get; }
+
+		public BlobNotFoundException(NamespaceId ns, BlobId blob) : base($"No Blob in Namespace {ns} with id {blob}")
+		{
+			Ns = ns;
+			Blob = blob;
+		}
+
+		public BlobNotFoundException(NamespaceId ns, BlobId blob, string message) : base(message)
+		{
+			Ns = ns;
+			Blob = blob;
+		}
+	}
+
+	public class BlobReplicationException : BlobNotFoundException
+	{
+		public BlobReplicationException(NamespaceId ns, BlobId blob, string message) : base(ns, blob, message)
+		{
+		}
+	}
+
+	public class BlobTooLargeException : Exception
+	{
+		public BlobId Blob { get; }
+
+		public BlobTooLargeException(BlobId blob) : base($"Blob {blob} was to large to cache")
+		{
+			Blob = blob;
+		}
+	}
+
+	public class ResourceHasToManyRequestsException : Exception
+	{
+		public ResourceHasToManyRequestsException(Exception originalException) : base($"To many requests to resource", originalException)
+		{
+		}
+	}
+
+	public class NamespaceNotFoundException : Exception
+	{
+		public NamespaceId Namespace { get; }
+
+		public NamespaceNotFoundException(NamespaceId @namespace) : base($"Could not find namespace {@namespace}")
+		{
+			Namespace = @namespace;
+		}
+
+		public NamespaceNotFoundException(NamespaceId @namespace, string message) : base(message)
+		{
+			Namespace = @namespace;
+		}
+	}
+
+	public static class BlobServiceExtensions
+	{
+		public static async Task<ContentId> PutCompressedObject(this IBlobService blobService, NamespaceId ns, BufferedPayload payload, ContentId? id, IServiceProvider provider)
+		{
+			IContentIdService contentIdStore = provider.GetService<IContentIdService>()!;
+			CompressedBufferUtils compressedBufferUtils = provider.GetService<CompressedBufferUtils>()!;
+			Tracer tracer = provider.GetService<Tracer>()!;
+
+			// decompress the content and generate a identifier from it to verify the identifier we got
+			await using Stream decompressStream = payload.GetStream();
+			// TODO: we should add a overload for decompress content that can work on streams, otherwise we are still limited to 2GB compressed blobs
+			byte[] decompressedContent = compressedBufferUtils.DecompressContent(await decompressStream.ToByteArray());
+
+			IoHash decompressedHash;
+			using (TelemetrySpan _ = tracer.StartActiveSpan("web.hash").SetAttribute("operation.name", "web.hash"))
+			{
+				decompressedHash = IoHash.Compute(decompressedContent);
+			}
+			if (id != null && id.Value.Hash != decompressedHash)
+			{
+				throw new HashMismatchException(decompressedHash, id.Value.Hash);
+			}
+
+			ContentId identifierDecompressedPayload = new ContentId(decompressedHash);
+
+			BlobId identifierCompressedPayload;
+			{
+				using TelemetrySpan _ = tracer.StartActiveSpan("web.hash").SetAttribute("operation.name", "web.hash");
+				await using Stream hashStream = payload.GetStream();
+				identifierCompressedPayload = new BlobId(await IoHash.ComputeAsync(hashStream));
+			}
+
+			// commit the mapping from the decompressed hash to the compressed hash, we run this in parallel with the blob store submit
+			// TODO: let users specify weight of the blob compared to previously submitted content ids
+			int contentIdWeight = (int)payload.Length;
+			Task contentIdStoreTask = contentIdStore.PutAsync(ns, identifierDecompressedPayload, identifierCompressedPayload, contentIdWeight);
+
+			// we still commit the compressed buffer to the object store using the hash of the compressed content
+			{
+				await blobService.PutObjectKnownHashAsync(ns, payload, identifierCompressedPayload);
+			}
+
+			await contentIdStoreTask;
+
+			return identifierDecompressedPayload;
+		}
+
+		public static async Task<(BlobContents, string)> GetCompressedObject(this IBlobService blobService, NamespaceId ns, ContentId contentId, IServiceProvider provider, bool supportsRedirectUri = false)
+		{
+			IContentIdService contentIdStore = provider.GetService<IContentIdService>()!;
+			Tracer tracer = provider.GetService<Tracer>()!;
+
+			BlobId[]? chunks = await contentIdStore.ResolveAsync(ns, contentId, mustBeContentId: false);
+			if (chunks == null || chunks.Length == 0)
+			{
+				throw new ContentIdResolveException(contentId);
+			}
+
+			// single chunk, we just return that chunk
+			if (chunks.Length == 1)
+			{
+				BlobId blobToReturn = chunks[0];
+				string mimeType = CustomMediaTypeNames.UnrealCompressedBuffer;
+				if (contentId.Equals(blobToReturn))
+				{
+					// this was actually the unmapped blob, meaning its not a compressed buffer
+					mimeType = MediaTypeNames.Application.Octet;
+				}
+
+				return (await blobService.GetObjectAsync(ns, blobToReturn, supportsRedirectUri: supportsRedirectUri), mimeType);
+			}
+
+			// chunked content, combine the chunks into a single stream
+			using TelemetrySpan _ = tracer.StartActiveSpan("blob.combine").SetAttribute("operation.name", "blob.combine");
+			Task<BlobContents>[] tasks = new Task<BlobContents>[chunks.Length];
+			for (int i = 0; i < chunks.Length; i++)
+			{
+				// even if it was requested to support redirect, since we need to combine the chunks using redirects is not possible
+				tasks[i] = blobService.GetObjectAsync(ns, chunks[i], supportsRedirectUri: false);
+			}
+
+			MemoryStream ms = new MemoryStream();
+			foreach (Task<BlobContents> task in tasks)
+			{
+				BlobContents blob = await task;
+				await using Stream s = blob.Stream;
+				await s.CopyToAsync(ms);
+			}
+
+			ms.Seek(0, SeekOrigin.Begin);
+
+			// chunking could not have happened for a non compressed buffer so assume it is compressed
+			return (new BlobContents(ms, ms.Length), CustomMediaTypeNames.UnrealCompressedBuffer);
+		}
+	}
+}
diff --git a/Engine/Source/Programs/Horde/Horde.Server/Ddc/IContentIdService.cs b/Engine/Source/Programs/Horde/Horde.Server/Ddc/IContentIdService.cs
new file mode 100644
--- /dev/null
+++ b/Engine/Source/Programs/Horde/Horde.Server/Ddc/IContentIdService.cs
@@ -0,0 +1,50 @@
+// Copyright Epic Games, Inc. All Rights Reserved.
+
+using EpicGames.Horde.Storage;
+using System;
+using System.Threading.Tasks;
+
+#pragma warning disable CS1591
+
+namespace Horde.Server.Ddc
+{
+	public interface IContentIdService
+    {
+        /// <summary>
+        /// Resolve a content id from its hash into the actual blob (that can in turn be chunked into a set of blobs)
+        /// </summary>
+        /// <param name="ns">The namespace to operate in</param>
+        /// <param name="contentId">The identifier for the content id</param>
+        /// <param name="mustBeContentId"></param>
+        /// <returns></returns>
+        Task<BlobId[]?> ResolveAsync(NamespaceId ns, ContentId contentId, bool mustBeContentId = false);
+
+        /// <summary>
+        /// Add a mapping from contentId to blobIdentifier
+        /// </summary>
+        /// <param name="ns">The namespace to operate in</param>
+        /// <param name="contentId">The contentId</param>
+        /// <param name="blobIdentifier">The blob the content id maps to</param>
+        /// <param name="contentWeight">Weight of this identifier compared to previous mappings, used to determine which is more important, lower weight is considered a better fit</param>
+        /// <returns></returns>
+        Task PutAsync(NamespaceId ns, ContentId contentId, BlobId blobIdentifier, int contentWeight);
+    }
+
+    public class InvalidContentIdException : Exception
+    {
+        public InvalidContentIdException(ContentId contentId) : base($"Unknown content id {contentId}")
+        {
+
+        }
+    }
+
+    public class ContentIdResolveException : Exception
+    {
+        public ContentId ContentId { get; }
+
+        public ContentIdResolveException(ContentId contentId) : base($"Unable to find any mapping of contentId {contentId} that has all blobs present")
+        {
+            ContentId = contentId;
+        }
+    }
+}
diff --git a/Engine/Source/Programs/Horde/Horde.Server/Ddc/IDdcBlobService.cs b/Engine/Source/Programs/Horde/Horde.Server/Ddc/IDdcBlobService.cs
deleted file mode 100644
diff --git a/Engine/Source/Programs/Horde/Horde.Server/Ddc/IDdcContentIdService.cs b/Engine/Source/Programs/Horde/Horde.Server/Ddc/IDdcContentIdService.cs
deleted file mode 100644
diff --git a/Engine/Source/Programs/Horde/Horde.Server/Ddc/IDdcRefService.cs b/Engine/Source/Programs/Horde/Horde.Server/Ddc/IDdcRefService.cs
deleted file mode 100644
diff --git a/Engine/Source/Programs/Horde/Horde.Server/Ddc/IRefService.cs b/Engine/Source/Programs/Horde/Horde.Server/Ddc/IRefService.cs
new file mode 100644
--- /dev/null
+++ b/Engine/Source/Programs/Horde/Horde.Server/Ddc/IRefService.cs
@@ -0,0 +1,69 @@
+// Copyright Epic Games, Inc. All Rights Reserved.
+
+using System;
+using System.Collections.Generic;
+using System.Threading;
+using System.Threading.Tasks;
+using EpicGames.Core;
+using EpicGames.Horde.Storage;
+using EpicGames.Serialization;
+
+#pragma warning disable CS1591
+
+namespace Horde.Server.Ddc
+{
+	public interface IRefService
+    {
+        Task<(RefRecord, BlobContents?)> GetAsync(NamespaceId ns, BucketId bucket, RefId key, string[] fields, bool doLastAccessTracking = true, CancellationToken cancellationToken = default);
+        Task<(ContentId[], BlobId[])> PutAsync(NamespaceId ns, BucketId bucket, RefId key, BlobId blobHash, CbObject payload, CancellationToken cancellationToken = default);
+        Task<(ContentId[], BlobId[])> FinalizeAsync(NamespaceId ns, BucketId bucket, RefId key, BlobId blobHash, CancellationToken cancellationToken = default);
+
+        Task<bool> DeleteAsync(NamespaceId ns, BucketId bucket, RefId key, CancellationToken cancellationToken = default);
+
+        Task<bool> ExistsAsync(NamespaceId ns, BucketId bucket, RefId key, CancellationToken cancellationToken = default);
+        Task<List<BlobId>> GetReferencedBlobsAsync(NamespaceId ns, BucketId bucket, RefId key, CancellationToken cancellationToken = default);
+    }
+
+	public class RefRecord
+	{
+		public RefRecord(NamespaceId ns, BucketId bucket, RefId name, DateTime lastAccess, byte[]? inlinePayload, BlobId blobIdentifier, bool isFinalized)
+		{
+			Namespace = ns;
+			Bucket = bucket;
+			Name = name;
+			LastAccess = lastAccess;
+			InlinePayload = inlinePayload;
+			BlobIdentifier = blobIdentifier;
+			IsFinalized = isFinalized;
+		}
+
+		public NamespaceId Namespace { get; }
+		public BucketId Bucket { get; }
+		public RefId Name { get; }
+		public DateTime LastAccess { get; }
+		public byte[]? InlinePayload { get; set; }
+		public BlobId BlobIdentifier { get; set; }
+		public bool IsFinalized { get; }
+	}
+
+	public class RefNotFoundException : Exception
+	{
+		public RefNotFoundException(NamespaceId ns, BucketId bucket, RefId key) : base($"Object not found {key} in bucket {bucket} namespace {ns}")
+		{
+			Namespace = ns;
+			Bucket = bucket;
+			Key = key;
+		}
+
+		public NamespaceId Namespace { get; }
+		public BucketId Bucket { get; }
+		public RefId Key { get; }
+	}
+
+	public class ObjectHashMismatchException : Exception
+    {
+        public ObjectHashMismatchException(NamespaceId ns, BucketId bucket, IoHash name, BlobId suppliedHash, BlobId actualHash) : base($"Object {name} in bucket {bucket} and namespace {ns} did not reference hash {suppliedHash} was referencing {actualHash}")
+        {
+        }
+    }
+}
diff --git a/Engine/Source/Programs/Horde/Horde.Server/Ddc/IReferenceResolver.cs b/Engine/Source/Programs/Horde/Horde.Server/Ddc/IReferenceResolver.cs
--- a/Engine/Source/Programs/Horde/Horde.Server/Ddc/IReferenceResolver.cs
+++ b/Engine/Source/Programs/Horde/Horde.Server/Ddc/IReferenceResolver.cs
@@ -89,11 +89,11 @@ namespace Horde.Server.Ddc
 
     public class ReferenceResolver : IReferenceResolver
     {
-        private readonly IDdcBlobService _blobStore;
-        private readonly IDdcContentIdService _contentIdStore;
+        private readonly IBlobService _blobStore;
+        private readonly IContentIdService _contentIdStore;
         private readonly Tracer _tracer;
 
-        public ReferenceResolver(IDdcBlobService blobStore, IDdcContentIdService contentIdStore, Tracer tracer)
+        public ReferenceResolver(IBlobService blobStore, IContentIdService contentIdStore, Tracer tracer)
         {
             _blobStore = blobStore;
             _contentIdStore = contentIdStore;
diff --git a/Engine/Source/Programs/Horde/Horde.Server/Ddc/ObjectsController.cs b/Engine/Source/Programs/Horde/Horde.Server/Ddc/ObjectsController.cs
new file mode 100644
--- /dev/null
+++ b/Engine/Source/Programs/Horde/Horde.Server/Ddc/ObjectsController.cs
@@ -0,0 +1,265 @@
+﻿// Copyright Epic Games, Inc. All Rights Reserved.
+
+using System;
+using System.Collections.Concurrent;
+using System.Collections.Generic;
+using System.ComponentModel.DataAnnotations;
+using System.Linq;
+using System.Net;
+using System.Net.Mime;
+using System.Threading;
+using System.Threading.Tasks;
+using EpicGames.AspNet;
+using EpicGames.Horde.Storage;
+using EpicGames.Serialization;
+using Horde.Server.Storage;
+using Microsoft.AspNetCore.Authorization;
+using Microsoft.AspNetCore.Mvc;
+using Microsoft.Extensions.Logging;
+using OpenTelemetry.Trace;
+
+#pragma warning disable CS1591
+
+namespace Horde.Server.Ddc
+{
+	using IDiagnosticContext = Serilog.IDiagnosticContext;
+
+    [ApiController]
+    [Route("api/v1/objects", Order = 0)]
+    [Authorize]
+    [Produces(CustomMediaTypeNames.UnrealCompactBinary, MediaTypeNames.Application.Json)]
+    public class ObjectsController : ControllerBase
+    {
+        private readonly IBlobService _storage;
+        private readonly IDiagnosticContext _diagnosticContext;
+        private readonly IRequestHelper _requestHelper;
+        private readonly IReferenceResolver _referenceResolver;
+        private readonly BufferedPayloadFactory _bufferedPayloadFactory;
+		private readonly Tracer _tracer;
+        private readonly ILogger _logger;
+
+        public ObjectsController(IBlobService storage, IDiagnosticContext diagnosticContext, IRequestHelper requestHelper, IReferenceResolver referenceResolver, BufferedPayloadFactory bufferedPayloadFactory, Tracer tracer, ILogger<ObjectsController> logger)
+        {
+            _storage = storage;
+            _diagnosticContext = diagnosticContext;
+            _requestHelper = requestHelper;
+            _referenceResolver = referenceResolver;
+            _bufferedPayloadFactory = bufferedPayloadFactory;
+			_tracer = tracer;
+            _logger = logger;
+        }
+
+        [HttpGet("{ns}/{id}")]
+        [ProducesDefaultResponseType]
+        public async Task<IActionResult> Get(
+            [Required] NamespaceId ns,
+            [Required] BlobId id)
+        {
+            ActionResult? result = await _requestHelper.HasAccessToNamespaceAsync(User, Request, ns, new [] { StorageAclAction.ReadBlobs });
+            if (result != null)
+            {
+                return result;
+            }
+
+            try
+            {
+                BlobContents blobContents = await _storage.GetObjectAsync(ns, id);
+
+                return File(blobContents.Stream, CustomMediaTypeNames.UnrealCompactBinary);
+            }
+            catch (BlobNotFoundException e)
+            {
+                return NotFound(new ValidationProblemDetails {Title = $"Object {e.Blob} not found"});
+            }
+        }
+
+        [HttpHead("{ns}/{id}")]
+        [ProducesDefaultResponseType]
+        public async Task<IActionResult> Head(
+            [Required] NamespaceId ns,
+            [Required] BlobId id)
+        {
+            ActionResult? result = await _requestHelper.HasAccessToNamespaceAsync(User, Request, ns, new [] { StorageAclAction.ReadBlobs });
+            if (result != null)
+            {
+                return result;
+            }
+
+            bool exists = await _storage.ExistsAsync(ns, id);
+
+            if (!exists)
+            {
+                return NotFound(new ValidationProblemDetails {Title = $"Object {id} not found"});
+            }
+
+            return Ok();
+        }
+
+        [HttpPost("{ns}/exists")]
+        [ProducesDefaultResponseType]
+        public async Task<IActionResult> ExistsMultiple(
+            [Required] NamespaceId ns,
+            [Required] [FromQuery] List<BlobId> id)
+        {
+            ActionResult? result = await _requestHelper.HasAccessToNamespaceAsync(User, Request, ns, new [] { StorageAclAction.ReadBlobs });
+            if (result != null)
+            {
+                return result;
+            }
+
+            ConcurrentBag<BlobId> missingBlobs = new ConcurrentBag<BlobId>();
+
+            IEnumerable<Task> tasks = id.Select(async blob =>
+            {
+                if (!await _storage.ExistsAsync(ns, blob))
+                {
+                    missingBlobs.Add(blob);
+                }
+            });
+            await Task.WhenAll(tasks);
+
+            return Ok(new HeadMultipleResponse {Needs = missingBlobs.ToArray()});
+        }
+
+        [HttpPost("{ns}/exist")]
+        [ProducesDefaultResponseType]
+        public async Task<IActionResult> ExistsBody(
+            [Required] NamespaceId ns,
+            [FromBody] BlobId[] bodyIds)
+        {
+            ActionResult? result = await _requestHelper.HasAccessToNamespaceAsync(User, Request, ns, new [] { StorageAclAction.ReadBlobs });
+            if (result != null)
+            {
+                return result;
+            }
+
+            ConcurrentBag<BlobId> missingBlobs = new ConcurrentBag<BlobId>();
+
+            IEnumerable<Task> tasks = bodyIds.Select(async blob =>
+            {
+                if (!await _storage.ExistsAsync(ns, blob))
+                {
+                    missingBlobs.Add(blob);
+                }
+            });
+            await Task.WhenAll(tasks);
+
+            return Ok(new HeadMultipleResponse { Needs = missingBlobs.ToArray() });
+        }
+
+        [HttpPut("{ns}/{id}")]
+        [RequiredContentType(CustomMediaTypeNames.UnrealCompactBinary)]
+        public async Task<IActionResult> Put(
+            [Required] NamespaceId ns,
+            [Required] BlobId id,
+			CancellationToken cancellationToken)
+        {
+            ActionResult? result = await _requestHelper.HasAccessToNamespaceAsync(User, Request, ns, new [] { StorageAclAction.WriteBlobs });
+            if (result != null)
+            {
+                return result;
+            }
+
+            _diagnosticContext.Set("Content-Length", Request.ContentLength ?? -1);
+            try
+            {
+                using BufferedPayload payload = await _bufferedPayloadFactory.CreateFromRequest(Request);
+
+                BlobId identifier = await _storage.PutObjectAsync(ns, payload, id, cancellationToken);
+                return Ok(new PutBlobResponse(identifier));
+            }
+            catch (ClientSendSlowException e)
+            {
+                return Problem(e.Message, null, (int)HttpStatusCode.RequestTimeout);
+            }
+        }
+
+        [HttpGet("{ns}/{id}/references")]
+        public async Task<IActionResult> ResolveReferences(
+            [Required] NamespaceId ns,
+            [Required] BlobId id)
+        {
+            ActionResult? result = await _requestHelper.HasAccessToNamespaceAsync(User, Request, ns, new [] { StorageAclAction.ReadBlobs });
+            if (result != null)
+            {
+                return result;
+            }
+
+            BlobContents blob;
+            try
+            {
+                blob = await _storage.GetObjectAsync(ns, id);
+            }
+            catch (BlobNotFoundException e)
+            {
+                return NotFound(new ValidationProblemDetails {Title = $"Object {e.Blob} not found"});
+            }
+           
+            byte[] blobContents = await blob.Stream.ToByteArray();
+            if (blobContents.Length == 0)
+            {
+                _logger.LogWarning("0 byte object found for {Id} {Namespace}", id, ns);
+            }
+
+            CbObject compactBinaryObject;
+            try
+            {
+                compactBinaryObject = new CbObject(blobContents);
+            }
+            catch (IndexOutOfRangeException)
+            {
+                return Problem(title: $"{id} was not a proper compact binary object.", detail: "Index out of range");
+            }
+
+            try
+            {
+                BlobId[] references = await _referenceResolver.GetReferencedBlobsAsync(ns, compactBinaryObject).ToArrayAsync();
+                return Ok(new ResolvedReferencesResult(references));
+            }
+            catch (PartialReferenceResolveException e)
+            {
+                return BadRequest(new ValidationProblemDetails {Title = $"Object {id} is missing content ids", Detail = $"Following content ids are invalid: {string.Join(",", e.UnresolvedReferences)}"});
+            }
+            catch (ReferenceIsMissingBlobsException e)
+            {
+                return BadRequest(new ValidationProblemDetails {Title = $"Object {id} is missing blobs", Detail = $"Following blobs are missing: {string.Join(",", e.MissingBlobs)}"});
+            }
+        }
+    }
+
+    public class PutBlobResponse
+    {
+        public PutBlobResponse()
+        {
+        }
+
+        public PutBlobResponse(BlobId identifier)
+        {
+            Identifier = identifier;
+        }
+
+        [CbField("identifier")]
+        public BlobId Identifier { get; set; }
+    }
+
+    public class DeletedResponse
+    {
+        public int DeletedCount { get; set; }
+    }
+
+    public class ResolvedReferencesResult
+    {
+        public ResolvedReferencesResult()
+        {
+            References = null!;
+        }
+
+        public ResolvedReferencesResult(BlobId[] references)
+        {
+            References = references;
+        }
+
+        [CbField("references")]
+        public BlobId[] References { get; set; }
+    }
+}
diff --git a/Engine/Source/Programs/Horde/Horde.Server/Ddc/RefNode.cs b/Engine/Source/Programs/Horde/Horde.Server/Ddc/RefNode.cs
new file mode 100644
--- /dev/null
+++ b/Engine/Source/Programs/Horde/Horde.Server/Ddc/RefNode.cs
@@ -0,0 +1,43 @@
+// Copyright Epic Games, Inc. All Rights Reserved.
+
+using System.Collections.Generic;
+using EpicGames.Core;
+using EpicGames.Horde.Storage;
+
+namespace Horde.Server.Ddc
+{
+	/// <summary>
+	/// A node containing ref data
+	/// </summary>
+	[NodeType("{0C7E5F25-454B-4B55-9B4A-F4635106D074}", 1)]
+	class RefNode : Node
+	{
+		public IoHash RootHash { get; }
+		public List<(IoHash Hash, BlobHandle Handle)> References { get; }
+
+		public RefNode(IoHash rootHash)
+		{
+			RootHash = rootHash;
+			References = new List<(IoHash, BlobHandle)>();
+		}
+
+		public RefNode(NodeReader reader)
+		{
+			RootHash = reader.ReadIoHash();
+			References = reader.ReadList(x => (reader.ReadIoHash(), reader.ReadNodeHandle()));
+		}
+
+		/// <inheritdoc/>
+		public override void Serialize(NodeWriter writer)
+		{
+			writer.WriteIoHash(RootHash);
+			writer.WriteList(References, x => WriteReference(writer, x.Hash, x.Handle));
+		}
+
+		static void WriteReference(NodeWriter writer, IoHash hash, BlobHandle handle)
+		{
+			writer.WriteIoHash(hash);
+			writer.WriteNodeHandle(handle);
+		}
+	}
+}
diff --git a/Engine/Source/Programs/Horde/Horde.Server/Ddc/RefService.cs b/Engine/Source/Programs/Horde/Horde.Server/Ddc/RefService.cs
new file mode 100644
--- /dev/null
+++ b/Engine/Source/Programs/Horde/Horde.Server/Ddc/RefService.cs
@@ -0,0 +1,123 @@
+// Copyright Epic Games, Inc. All Rights Reserved.
+
+using System;
+using System.Collections.Generic;
+using System.Linq;
+using System.Threading;
+using System.Threading.Tasks;
+using EpicGames.Core;
+using EpicGames.Horde.Storage;
+using EpicGames.Serialization;
+using Horde.Server.Storage;
+
+namespace Horde.Server.Ddc
+{
+	class RefService : IRefService
+	{
+		readonly StorageService _storageService;
+		readonly BlobService _blobService;
+
+		public RefService(StorageService storageService, BlobService blobService)
+		{
+			_storageService = storageService;
+			_blobService = blobService;
+		}
+
+		static RefName GetRefName(BucketId bucketId, RefId refId) => $"{bucketId}/{refId}";
+
+		public async Task<bool> DeleteAsync(NamespaceId ns, BucketId bucket, RefId key, CancellationToken cancellationToken)
+		{
+			IStorageClient storageClient = await _storageService.GetClientAsync(ns, cancellationToken);
+			await storageClient.DeleteRefAsync(GetRefName(bucket, key), cancellationToken);
+			return true;
+		}
+
+		public async Task<bool> ExistsAsync(NamespaceId ns, BucketId bucket, RefId key, CancellationToken cancellationToken)
+		{
+			IStorageClient storageClient = await _storageService.GetClientAsync(ns, cancellationToken);
+			BlobHandle? handle = await storageClient.TryReadRefTargetAsync(GetRefName(bucket, key), cancellationToken: cancellationToken);
+			return handle != null;
+		}
+
+		public async Task<(ContentId[], BlobId[])> FinalizeAsync(NamespaceId ns, BucketId bucket, RefId key, BlobId blobHash, CancellationToken cancellationToken)
+		{
+			HashSet<BlobId> missingBlobIds = new HashSet<BlobId>();
+
+			// Read all the blobs back in
+			List<BlobId> blobIds = new List<BlobId>();
+			blobIds.Add(blobHash);
+
+			HashSet<BlobId> addedBlobIds = new HashSet<BlobId>();
+			for (int idx = 0; idx < blobIds.Count; idx++)
+			{
+				try
+				{
+					using BlobContents contents = await _blobService.GetObjectAsync(ns, blobIds[idx], storageLayers: null, supportsRedirectUri: false, cancellationToken: cancellationToken);
+					byte[] data = await contents.Stream.ReadAllBytesAsync(cancellationToken);
+
+					CbObject obj = new CbObject(data);
+					obj.IterateAttachments(x =>
+					{
+						BlobId blobId = new BlobId(x.AsAttachment());
+						if (addedBlobIds.Add(blobId))
+						{
+							blobIds.Add(blobId);
+						}
+					});
+				}
+				catch (BlobNotFoundException)
+				{
+					missingBlobIds.Add(blobIds[idx]);
+				}
+			}
+
+			if (missingBlobIds.Count == 0)
+			{
+				IStorageClient storageClient = await _storageService.GetClientAsync(ns, cancellationToken);
+
+				BlobHandle? handle = await storageClient.FindAliasAsync(BlobService.GetAlias(blobHash), cancellationToken).FirstOrDefaultAsync(cancellationToken);
+				if (handle == null)
+				{
+					throw new BlobNotFoundException(ns, blobHash);
+				}
+
+				await storageClient.WriteRefTargetAsync(GetRefName(bucket, key), handle, cancellationToken: cancellationToken);
+			}
+
+			return (Array.Empty<ContentId>(), missingBlobIds.ToArray());
+		}
+
+		public async Task<(RefRecord, BlobContents?)> GetAsync(NamespaceId ns, BucketId bucket, RefId key, string[] fields, bool doLastAccessTracking, CancellationToken cancellationToken)
+		{
+			IStorageClient storageClient = await _storageService.GetClientAsync(ns, cancellationToken);
+			
+			RefNode? node = await storageClient.TryReadNodeAsync<RefNode>(GetRefName(bucket, key), cancellationToken: cancellationToken);
+			if (node == null)
+			{
+				throw new RefNotFoundException(ns, bucket, key);
+			}
+
+			RefRecord record = new RefRecord(ns, bucket, key, DateTime.UtcNow, null, new BlobId(node.RootHash), false);
+			return (record, null);
+		}
+
+		public async Task<List<BlobId>> GetReferencedBlobsAsync(NamespaceId ns, BucketId bucket, RefId key, CancellationToken cancellationToken)
+		{
+			IStorageClient storageClient = await _storageService.GetClientAsync(ns, cancellationToken);
+
+			RefNode? node = await storageClient.TryReadNodeAsync<RefNode>(GetRefName(bucket, key), cancellationToken: cancellationToken);
+			if (node == null)
+			{
+				throw new RefNotFoundException(ns, bucket, key);
+			}
+
+			return node.References.Select(x => new BlobId(x.Hash)).ToList();
+		}
+
+		public async Task<(ContentId[], BlobId[])> PutAsync(NamespaceId ns, BucketId bucket, RefId key, BlobId blobHash, CbObject payload, CancellationToken cancellationToken)
+		{
+			await _blobService.PutObjectAsync(ns, payload.GetView(), blobHash, cancellationToken);
+			return await FinalizeAsync(ns, bucket, key, blobHash, cancellationToken);
+		}
+	}
+}
diff --git a/Engine/Source/Programs/Horde/Horde.Server/Ddc/RefsController.cs b/Engine/Source/Programs/Horde/Horde.Server/Ddc/RefsController.cs
new file mode 100644
--- /dev/null
+++ b/Engine/Source/Programs/Horde/Horde.Server/Ddc/RefsController.cs
@@ -0,0 +1,1289 @@
+// Copyright Epic Games, Inc. All Rights Reserved.
+
+using System;
+using System.Collections.Concurrent;
+using System.Collections.Generic;
+using System.ComponentModel.DataAnnotations;
+using System.Globalization;
+using System.IO;
+using System.Linq;
+using System.Net;
+using System.Net.Mime;
+using System.Text;
+using System.Text.Json.Serialization;
+using System.Threading;
+using System.Threading.Tasks;
+using EpicGames.AspNet;
+using EpicGames.Core;
+using EpicGames.Horde.Storage;
+using EpicGames.Serialization;
+using Horde.Server.Acls;
+using Horde.Server.Storage;
+using JetBrains.Annotations;
+using Microsoft.AspNetCore.Authorization;
+using Microsoft.AspNetCore.Http;
+using Microsoft.AspNetCore.Http.Extensions;
+using Microsoft.AspNetCore.Mvc;
+using Microsoft.Extensions.DependencyInjection;
+using Microsoft.Extensions.Logging;
+using OpenTelemetry.Trace;
+
+#pragma warning disable CS1591
+
+namespace Horde.Server.Ddc
+{
+	using IDiagnosticContext = Serilog.IDiagnosticContext;
+
+    [ApiController]
+    [FormatFilter]
+    [Produces(MediaTypeNames.Application.Json, MediaTypeNames.Application.Octet, CustomMediaTypeNames.UnrealCompactBinary)]
+    [Route("api/v1/refs")]
+    [Authorize]
+    public class RefsController : ControllerBase
+    {
+        private readonly IDiagnosticContext _diagnosticContext;
+        private readonly FormatResolver _formatResolver;
+        private readonly BufferedPayloadFactory _bufferedPayloadFactory;
+        private readonly IReferenceResolver _referenceResolver;
+        private readonly NginxRedirectHelper _nginxRedirectHelper;
+        private readonly IRequestHelper _requestHelper;
+        private readonly Tracer _tracer;
+
+        private readonly ILogger _logger;
+        private readonly IRefService _refService;
+        private readonly IBlobService _blobStore;
+
+        public RefsController(IRefService objectService, IBlobService blobStore, IDiagnosticContext diagnosticContext, FormatResolver formatResolver, BufferedPayloadFactory bufferedPayloadFactory, IReferenceResolver referenceResolver, NginxRedirectHelper nginxRedirectHelper, IRequestHelper requestHelper, Tracer tracer, ILogger<RefsController> logger)
+        {
+            _refService = objectService;
+            _blobStore = blobStore;
+            _diagnosticContext = diagnosticContext;
+            _formatResolver = formatResolver;
+            _bufferedPayloadFactory = bufferedPayloadFactory;
+            _referenceResolver = referenceResolver;
+            _nginxRedirectHelper = nginxRedirectHelper;
+            _requestHelper = requestHelper;
+            _tracer = tracer;
+            _logger = logger;
+        }
+
+        /// <summary>
+        /// Returns a refs key
+        /// </summary>
+        /// <param name="ns">Namespace. Each namespace is completely separated from each other. Use for different types of data that is never expected to be similar (between two different games for instance). Example: `uc4.ddc`</param>
+        /// <param name="bucket">The category/type of record you are caching. Is a clustered key together with the actual key, but all records in the same bucket can be dropped easily. Example: `terrainTexture` </param>
+        /// <param name="key">The unique name of this particular key. `iAmAVeryValidKey`</param>
+        /// <param name="format">Optional specifier to set which output format is used json/raw/cb</param>
+        [HttpGet("{ns}/{bucket}/{key}.{format?}", Order = 500)]
+        [Produces(MediaTypeNames.Application.Json, MediaTypeNames.Application.Octet, CustomMediaTypeNames.UnrealCompactBinary, CustomMediaTypeNames.JupiterInlinedPayload, CustomMediaTypeNames.UnrealCompactBinaryPackage)]
+        public async Task<IActionResult> Get(
+            [FromRoute] [Required] NamespaceId ns,
+            [FromRoute] [Required] BucketId bucket,
+            [FromRoute] [Required] RefId key,
+            [FromRoute] string? format = null)
+        {
+            ActionResult? accessResult = await _requestHelper.HasAccessToNamespaceAsync(User, Request, ns, new [] { StorageAclAction.ReadRefs });
+            if (accessResult != null)
+            {
+                return accessResult;
+            }
+
+            try
+            {
+                (RefRecord objectRecord, BlobContents? blob) = await _refService.GetAsync(ns, bucket, key, Array.Empty<string>());
+
+                if (blob == null)
+                {
+                    throw new InvalidOperationException($"Blob was null when attempting to fetch {ns} {bucket} {key}");
+                }
+
+                if (!objectRecord.IsFinalized)
+                {
+                    // we do not consider un-finalized objects as valid
+                    return BadRequest(new ProblemDetails { Title = $"Object {objectRecord.Bucket} {objectRecord.Name} is not finalized." });
+                }
+
+                Response.Headers[CommonHeaders.HashHeaderName] = objectRecord.BlobIdentifier.ToString();
+                Response.Headers[CommonHeaders.LastAccessHeaderName] = objectRecord.LastAccess.ToString(CultureInfo.InvariantCulture);
+
+                async Task WriteBody(BlobContents blobContents, string contentType)
+                {
+                    IServerTiming? serverTiming = Request.HttpContext.RequestServices.GetService<IServerTiming>();
+                    using ServerTimingMetricScoped? serverTimingScope = serverTiming?.CreateServerTimingMetricScope("body.write", "Time spent writing body");
+
+                    long contentLength = blobContents.Length;
+                    using TelemetrySpan scope = _tracer.StartActiveSpan("body.write").SetAttribute("operation.name", "body.write");
+                    scope.SetAttribute("content-length", contentLength);
+                    const int BufferSize = 64 * 1024;
+                    Stream outputStream = Response.Body;
+                    Response.ContentLength = contentLength;
+                    Response.ContentType = contentType;
+                    Response.StatusCode = StatusCodes.Status200OK;
+                    try
+                    {
+                        await StreamCopyOperation.CopyToAsync(blobContents.Stream, outputStream, count: null, bufferSize: BufferSize, cancel: Response.HttpContext.RequestAborted);
+                    }
+                    catch (OperationCanceledException)
+                    {
+                        // do not raise exceptions for cancelled writes
+                        // as we have already started writing a response we can not change the status code
+                        // so we just drop a warning and proceed
+                        _logger.LogWarning("The operation was canceled while writing the body");
+                    }
+                }
+
+                string responseType = _formatResolver.GetResponseType(Request, format, CustomMediaTypeNames.UnrealCompactBinary);
+                Tracer.CurrentSpan.SetAttribute("response-type", responseType);
+                switch (responseType)
+                {
+                    case CustomMediaTypeNames.UnrealCompactBinary:
+                    {
+                        // for compact binary we can just serialize our internal object
+                        await WriteBody(blob, CustomMediaTypeNames.UnrealCompactBinary);
+
+                        break;
+                    }
+                    case MediaTypeNames.Application.Octet:
+                    {
+                        byte[] blobMemory = await blob.Stream.ToByteArray();
+                        CbObject cb = new CbObject(blobMemory);
+
+                        (int,CbField?) CountFields(CbObject o)
+                        {
+                            int count = 0;
+                            CbField? foundField = null;
+                            cb.IterateAttachments(field =>
+                            {
+                                ++count;
+                                if (field.IsBinaryAttachment())
+                                {
+                                    foundField = field;
+                                }
+                            });
+
+                            return (count, foundField);
+                        }
+
+                        // breaking lambda call into private method to workaround incorrect triggering of CA1508 - [REDACTED_PROFILE_URL]/roslyn-analyzers/issues/5254
+                        (int countOfAttachmentFields,CbField? binaryAttachmentField) = CountFields(cb);
+
+                        if (countOfAttachmentFields == 1 && binaryAttachmentField != null)
+                        {
+                            // there is a single attachment field and that is of the binary attachment type, fetch that attachment and return it instead of the compact binary
+                            // this is so that we match the uploaded that happened as a octet-stream which generates a small cb object with a single attachment
+
+                            IoHash hash = binaryAttachmentField.AsBinaryAttachment();
+
+                            BlobContents referencedBlobContents = await _blobStore.GetObjectAsync(ns, new BlobId(hash));
+
+                            if (_nginxRedirectHelper.CanRedirect(Request, referencedBlobContents))
+                            {
+                                return _nginxRedirectHelper.CreateActionResult(referencedBlobContents, MediaTypeNames.Application.Octet);
+                            }
+                            await WriteBody(referencedBlobContents, MediaTypeNames.Application.Octet);
+                            break;
+                        }
+
+                        // this doesn't look like the generated compact binary so we just return the payload
+                        await using BlobContents contents = new(blobMemory);
+                        await WriteBody(contents, MediaTypeNames.Application.Octet);
+                        break;
+                    }
+                    case MediaTypeNames.Application.Json:
+                    {
+                        byte[] blobMemory;
+                        {
+                            using TelemetrySpan scope = _tracer.StartActiveSpan("json.readblob").SetAttribute("operation.name", "json.readblob");
+                            blobMemory = await blob.Stream.ToByteArray();
+                        }
+                        CbObject cb = new CbObject(blobMemory);
+                        string s = cb.ToJson();
+                        await using BlobContents contents = new BlobContents(Encoding.UTF8.GetBytes(s));
+                        await WriteBody(contents, MediaTypeNames.Application.Json);
+                        break;
+
+                    }
+                    case CustomMediaTypeNames.UnrealCompactBinaryPackage:
+                    {
+                        using TelemetrySpan packageScope = _tracer.StartActiveSpan("cbpackage.fetch").SetAttribute("operation.name", "cbpackage.fetch");
+                        byte[] blobMemory = await blob.Stream.ToByteArray();
+                        CbObject cb = new CbObject(blobMemory);
+
+                        IAsyncEnumerable<Attachment> attachments = _referenceResolver.GetAttachmentsAsync(ns, cb);
+
+                        using CbPackageBuilder writer = new CbPackageBuilder();
+                        writer.AddAttachment(objectRecord.BlobIdentifier.Hash, CbPackageAttachmentFlags.IsObject, blobMemory);
+
+                        await Parallel.ForEachAsync(attachments, async (attachment, token) =>
+                        {
+                            IoHash attachmentHash = attachment.AsIoHash();
+                            CbPackageAttachmentFlags flags = 0;
+
+                            try
+                            {
+                                BlobContents attachmentContents;
+                                if (attachment is BlobAttachment blobAttachment)
+                                {
+                                    BlobId referencedBlob = blobAttachment.Identifier;
+                                    attachmentContents = await _blobStore.GetObjectAsync(ns, referencedBlob);
+                                }
+                                else if (attachment is ObjectAttachment objectAttachment)
+                                {
+                                    flags |= CbPackageAttachmentFlags.IsObject;
+                                    BlobId referencedBlob = objectAttachment.Identifier;
+                                    attachmentContents = await _blobStore.GetObjectAsync(ns, referencedBlob);
+                                }
+                                else if (attachment is ContentIdAttachment contentIdAttachment)
+                                {
+
+                                    ContentId contentId = contentIdAttachment.Identifier;
+                                    (attachmentContents, string mime) = await _blobStore.GetCompressedObject(ns, contentId, HttpContext.RequestServices);
+                                    if (mime == CustomMediaTypeNames.UnrealCompressedBuffer)
+                                    {
+                                        flags |= CbPackageAttachmentFlags.IsCompressed;
+                                    }
+                                    else
+                                    {
+                                        // this resolved to a uncompressed blob, the content id existed the the compressed blob didn't
+                                        // so resetting flags to indicate this.
+                                        flags = 0;
+                                    }
+                                }
+                                else
+                                {
+                                    throw new NotSupportedException($"Unknown attachment type {attachment.GetType()}");
+                                }
+
+                                writer.AddAttachment(attachmentHash, flags, attachmentContents.Stream, (ulong)attachmentContents.Length);
+                            }
+                            catch (Exception e)
+                            {
+                                (CbObject errorObject, HttpStatusCode _) = ToErrorResult(e);
+                                writer.AddAttachment(attachmentHash,
+                                    CbPackageAttachmentFlags.IsError | CbPackageAttachmentFlags.IsObject,
+                                    errorObject.GetView().ToArray());
+                            }
+                        });
+
+                        byte[] packageBytes;
+                        {
+                            using TelemetrySpan _ = _tracer.StartActiveSpan("cbpackage.buffer").SetAttribute("operation.name", "cbpackage.buffer");
+                            packageBytes = await writer.ToByteArray();
+                        }
+                        await using BlobContents contents = new BlobContents(packageBytes);
+                        await WriteBody(contents, CustomMediaTypeNames.UnrealCompactBinaryPackage);
+                        break;
+                    }
+                    case CustomMediaTypeNames.JupiterInlinedPayload:
+                    {
+                        byte[] blobMemory = await blob.Stream.ToByteArray();
+                        CbObject cb = new CbObject(blobMemory);
+
+                        static (int, int) CountFields(CbObject o)
+                        {
+                            int countOfBinaryAttachmentFields = 0;
+                            int countOfAttachmentFields = 0;
+
+                            o.IterateAttachments(field =>
+                            {
+                                if (field.IsBinaryAttachment())
+                                {
+                                    ++countOfBinaryAttachmentFields;
+                                }
+
+                                if (field.IsAttachment())
+                                {
+                                    ++countOfAttachmentFields;
+                                }
+                            });
+
+                            return (countOfBinaryAttachmentFields, countOfAttachmentFields);
+                        }
+                        // breaking lambda call into private method to workaround incorrect triggering of CA1508 - [REDACTED_PROFILE_URL]/roslyn-analyzers/issues/5254
+                        (int countOfAttachmentFields, int countOfBinaryAttachmentFields) = CountFields(cb);
+
+                        // if the object consists of a single attachment field we return this attachment field instead
+                        if (countOfBinaryAttachmentFields == 1 && countOfAttachmentFields == 1)
+                        {
+                            // fetch the blob so we can resolve any content ids in it
+                            List<BlobId> referencedBlobs;
+                            try
+                            {
+                                IAsyncEnumerable<BlobId> referencedBlobsEnumerable = _referenceResolver.GetReferencedBlobsAsync(ns, cb);
+                                referencedBlobs = await referencedBlobsEnumerable.ToListAsync();
+                            }
+                            catch (PartialReferenceResolveException)
+                            {
+                                return NotFound(new ProblemDetails {Title = $"Object {bucket} {key} in namespace {ns} was missing some content ids"});
+                            }
+                            catch (ReferenceIsMissingBlobsException)
+                            {
+                                return NotFound(new ProblemDetails {Title = $"Object {bucket} {key} in namespace {ns} was missing some blobs"});
+                            }
+
+                            if (referencedBlobs.Count == 1)
+                            {
+                                BlobId attachmentToSend = referencedBlobs.First();
+                                try
+                                {
+                                    BlobContents referencedBlobContents = await _blobStore.GetObjectAsync(ns, attachmentToSend);
+                                    Response.Headers[CommonHeaders.InlinePayloadHash] = attachmentToSend.ToString();
+
+                                    if (_nginxRedirectHelper.CanRedirect(Request, referencedBlobContents))
+                                    {
+                                        return _nginxRedirectHelper.CreateActionResult(referencedBlobContents, CustomMediaTypeNames.JupiterInlinedPayload);
+                                    }
+                                    await WriteBody(referencedBlobContents, CustomMediaTypeNames.JupiterInlinedPayload);
+                                }
+                                catch (BlobNotFoundException)
+                                {
+                                    return NotFound(new ProblemDetails {Title = $"Object {bucket} {key} in namespace {ns} was missing blob {attachmentToSend}"});
+                                }
+                                catch (Exception ex)
+                                {
+                                    _logger.LogError(ex, "Unknown exception encountered while writing body for jupiter inlined payload.");
+                                    throw;
+                                }
+                                return new EmptyResult();
+                            }
+                            else if (referencedBlobs.Count == 0)
+                            {
+                                return NotFound(new ProblemDetails
+                                {
+                                    Title =
+                                        $"Object {objectRecord.Bucket} {objectRecord.Name} did not resolve into any objects that we could find."
+                                });
+                            }
+
+                            return BadRequest(new ProblemDetails
+                            {
+                                Title =
+                                    $"Object {objectRecord.Bucket} {objectRecord.Name} contained a content id which resolved to more then 1 blob, unable to inline this object. Use compact object response instead."
+                            });
+                        }
+                        else if (countOfBinaryAttachmentFields == 0 && countOfAttachmentFields == 0)
+                        {
+                            // no attachments so we just return the compact object instead
+                            await using BlobContents contents = new BlobContents(blobMemory);
+                            await WriteBody(contents, CustomMediaTypeNames.JupiterInlinedPayload);
+                            return new EmptyResult();
+                        }
+
+                        return BadRequest(new ProblemDetails
+                        {
+                            Title =
+                                $"Object {objectRecord.Bucket} {objectRecord.Name} had more then 1 binary attachment field, unable to inline this object. Use compact object response instead."
+                        });
+                    }
+                    default:
+                        throw new NotImplementedException($"Unknown expected response type {responseType}");
+                }
+                
+                // this result is ignored as we write to the body explicitly
+                return new EmptyResult();
+
+            }
+            catch (NamespaceNotFoundException e)
+            {
+                return NotFound(new ProblemDetails {Title = $"Namespace {e.Namespace} did not exist"});
+            }
+            catch (RefNotFoundException e)
+            {
+                return NotFound(new ProblemDetails { Title = $"Object {e.Bucket} {e.Key} did not exist" });
+            }
+            catch (BlobNotFoundException e)
+            {
+                return NotFound(new ProblemDetails { Title = $"Object {e.Blob} in {e.Ns} not found" });
+            }
+        }
+
+       
+
+    /// <summary>
+    /// Returns the metadata about a ref key
+    /// </summary>
+    /// <param name="ns">Namespace. Each namespace is completely separated from each other. Use for different types of data that is never expected to be similar (between two different games for instance). Example: `uc4.ddc`</param>
+    /// <param name="bucket">The category/type of record you are caching. Is a clustered key together with the actual key, but all records in the same bucket can be dropped easily. Example: `terrainTexture` </param>
+    /// <param name="key">The unique name of this particular key. `iAmAVeryValidKey`</param>
+    /// <param name="fields">The fields to include in the response, omit this to include everything.</param>
+    [HttpGet("{ns}/{bucket}/{key}/metadata", Order = 500)]
+    public async Task<IActionResult> GetMetadata(
+        [FromRoute] [Required] NamespaceId ns,
+        [FromRoute] [Required] BucketId bucket,
+        [FromRoute] [Required] RefId key,
+        [FromQuery] string[] fields)
+    {
+        ActionResult? accessResult = await _requestHelper.HasAccessToNamespaceAsync(User, Request, ns, new [] { StorageAclAction.ReadRefs });
+        if (accessResult != null)
+        {
+            return accessResult;
+        }
+
+        try
+        {
+            (RefRecord objectRecord, BlobContents? _) = await _refService.GetAsync(ns, bucket, key, fields);
+
+            return Ok(new RefMetadataResponse(objectRecord));
+        }
+        catch (NamespaceNotFoundException e)
+        {
+            return NotFound(new ProblemDetails {Title = $"Namespace {e.Namespace} did not exist"});
+        }
+        catch (RefNotFoundException e)
+        {
+            return NotFound(new ProblemDetails { Title = $"Object {e.Bucket} {e.Key} did not exist" });
+        }
+        catch (BlobNotFoundException e)
+        {
+            return NotFound(new ProblemDetails { Title = $"Object {e.Blob} in {e.Ns} not found" });
+        }
+    }
+
+        /// <summary>
+        /// Checks if a object exists
+        /// </summary>
+        /// <param name="ns">Namespace. Each namespace is completely separated from each other. Use for different types of data that is never expected to be similar (between two different games for instance). Example: `uc4.ddc`</param>
+        /// <param name="bucket">The category/type of record you are caching. Is a clustered key together with the actual key, but all records in the same bucket can be dropped easily. Example: `terrainTexture` </param>
+        /// <param name="key">The unique name of this particular key. `iAmAVeryValidKey`</param>
+        /// <returns>200 if it existed, 400 otherwise</returns>
+        [HttpHead("{ns}/{bucket}/{key}", Order = 500)]
+        [ProducesResponseType(type: typeof(OkResult), 200)]
+        [ProducesResponseType(type: typeof(ValidationProblemDetails), 400)]
+        public async Task<IActionResult> Head(
+            [FromRoute] [Required] NamespaceId ns,
+            [FromRoute] [Required] BucketId bucket,
+            [FromRoute] [Required] RefId key)
+        {
+            ActionResult? accessResult = await _requestHelper.HasAccessToNamespaceAsync(User, Request, ns, new [] { StorageAclAction.ReadRefs });
+            if (accessResult != null)
+            {
+                return accessResult;
+            }
+
+            try
+            {
+                (RefRecord record, BlobContents? blob) = await _refService.GetAsync(ns, bucket, key, new string[] {"blobIdentifier", "IsFinalized"});
+                Response.Headers[CommonHeaders.HashHeaderName] = record.BlobIdentifier.ToString();
+
+                if (!record.IsFinalized)
+                {
+                    return NotFound(new ProblemDetails {Title = $"Object {bucket} {key} in namespace {ns} is not finalized."});
+                }
+
+                blob ??= await _blobStore.GetObjectAsync(ns, record.BlobIdentifier);
+
+                // we have to verify the blobs are available locally, as the record of the key is replicated a head of the content
+                // TODO: Once we support inline replication this step is not needed as at least one region as this blob, just maybe not this current one
+                byte[] blobContents = await blob.Stream.ToByteArray();
+                CbObject compactBinaryObject = new CbObject(blobContents);
+                // the reference resolver will throw if any blob is missing, so no need to do anything other then process each reference
+                IAsyncEnumerable<BlobId> references = _referenceResolver.GetReferencedBlobsAsync(ns, compactBinaryObject);
+                List<BlobId>? _ = await references.ToListAsync();
+
+                // we have to verify the blobs are available locally, as the record of the key is replicated a head of the content
+                // TODO: Once we support inline replication this step is not needed as at least one region as this blob, just maybe not this current one
+                BlobId[] unknownBlobs = await _blobStore.FilterOutKnownBlobsAsync(ns, new BlobId[] { record.BlobIdentifier });
+                if (unknownBlobs.Length != 0)
+                {
+                    return NotFound(new ProblemDetails {Title = $"Object {bucket} {key} in namespace {ns} had at least one missing blob."});
+                }
+            }
+            catch (NamespaceNotFoundException e)
+            {
+                return NotFound(new ProblemDetails {Title = $"Namespace {e.Namespace} did not exist"});
+            }
+            catch (BlobNotFoundException e)
+            {
+                return NotFound(new ProblemDetails { Title = $"Blob {e.Blob} in namespace {ns} did not exist" });
+            }
+            catch (RefNotFoundException e)
+            {
+                return NotFound(new ProblemDetails {Title = $"Object {e.Bucket} {e.Key} in namespace {e.Namespace} did not exist"});
+            }
+            catch (PartialReferenceResolveException)
+            {
+                return NotFound(new ProblemDetails {Title = $"Object {bucket} {key} in namespace {ns} was missing some content ids"});
+            }
+            catch (ReferenceIsMissingBlobsException)
+            {
+                return NotFound(new ProblemDetails {Title = $"Object {bucket} {key} in namespace {ns} was missing some blobs"});
+            }
+
+            return Ok();
+        }
+
+        [HttpGet("{ns}/exists")]
+        [ProducesDefaultResponseType]
+        public async Task<IActionResult> ExistsMultiple(
+            [FromRoute] [Required] NamespaceId ns,
+            [FromQuery] [Required] List<string> names)
+        {
+            ActionResult? accessResult = await _requestHelper.HasAccessToNamespaceAsync(User, Request, ns, new [] { StorageAclAction.ReadRefs });
+            if (accessResult != null)
+            {
+                return accessResult;
+            }
+
+            ConcurrentBag<(BucketId, RefId)> missingObject = new ();
+
+            List<(BucketId, RefId)> requestedNames = new List<(BucketId, RefId)>();
+            foreach (string name in names)
+            {
+                int separatorIndex = name.IndexOf(".", StringComparison.Ordinal);
+                if (separatorIndex == -1)
+                {
+                    return BadRequest(new ProblemDetails() { Title = $"Key {name} did not contain a '.' separator" });
+                }
+
+                BucketId bucket = new BucketId(name.Substring(0, separatorIndex));
+				RefId key = RefId.Parse(name.Substring(separatorIndex + 1));
+                requestedNames.Add((bucket, key));
+            }
+
+            IEnumerable<Task> tasks = requestedNames.Select(async pair =>
+            {
+                (BucketId bucket, RefId key) = pair;
+                try
+                {
+                    (RefRecord record, BlobContents? blob) =
+                        await _refService.GetAsync(ns, bucket, key, new string[] { "blobIdentifier" });
+
+                    blob ??= await _blobStore.GetObjectAsync(ns, record.BlobIdentifier);
+
+                    // we have to verify the blobs are available locally, as the record of the key is replicated a head of the content
+                    // TODO: Once we support inline replication this step is not needed as at least one region as this blob, just maybe not this current one
+                    byte[] blobContents = await blob.Stream.ToByteArray();
+                    CbObject cb = new CbObject(blobContents);
+                    // the reference resolver will throw if any blob is missing, so no need to do anything other then process each reference
+                    IAsyncEnumerable<BlobId> references = _referenceResolver.GetReferencedBlobsAsync(ns, cb);
+                    List<BlobId>? _ = await references.ToListAsync();
+                }
+                catch (RefNotFoundException)
+                {
+                    missingObject.Add((bucket, key));
+                }
+                catch (PartialReferenceResolveException)
+                {
+                    missingObject.Add((bucket, key));
+                }
+                catch (ReferenceIsMissingBlobsException)
+                {
+                    missingObject.Add((bucket, key));
+                }
+                catch (BlobNotFoundException)
+                {
+                    missingObject.Add((bucket, key));
+                }
+            });
+            await Task.WhenAll(tasks);
+
+            return Ok(new ExistCheckMultipleRefsResponse(missingObject.ToList()));
+        }
+
+        [HttpPut("{ns}/{bucket}/{key}.{format?}", Order = 500)]
+        [DisableRequestSizeLimit]
+        public async Task<IActionResult> PutObject(
+            [FromRoute] [Required] NamespaceId ns,
+            [FromRoute] [Required] BucketId bucket,
+            [FromRoute] [Required] RefId key,
+			CancellationToken cancellationToken)
+        {
+            ActionResult? accessResult = await _requestHelper.HasAccessToNamespaceAsync(User, Request, ns, new [] { StorageAclAction.WriteRefs });
+            if (accessResult != null)
+            {
+                return accessResult;
+            }
+
+            _diagnosticContext.Set("Content-Length", Request.ContentLength ?? -1);
+
+            CbObject payloadObject;
+            BlobId blobHeader;
+
+            try
+            {
+                using BufferedPayload payload = await _bufferedPayloadFactory.CreateFromRequest(Request);
+
+                BlobId headerHash;
+                if (Request.Headers.ContainsKey(CommonHeaders.HashHeaderName))
+                {
+                    headerHash = BlobId.Parse(Request.Headers[CommonHeaders.HashHeaderName]);
+                }
+                else
+                {
+                    return BadRequest(new ProblemDetails
+                    {
+                        Title = $"Missing expected header {CommonHeaders.HashHeaderName}"
+                    });
+                }
+
+                blobHeader = headerHash;
+
+                switch (Request.ContentType)
+                {
+                    case MediaTypeNames.Application.Json:
+                    {
+                        // TODO: define a scheme for how a json object specifies references
+
+                        blobHeader = await _blobStore.PutObjectAsync(ns, payload, headerHash, cancellationToken);
+
+                        // TODO: convert the json object into a compact binary instead
+                        CbWriter writer = new CbWriter();
+                        writer.BeginObject();
+						writer.WriteBinaryAttachmentValue(blobHeader.Hash);
+                        writer.EndObject();
+
+                        byte[] blob = writer.ToByteArray();
+                        payloadObject = new CbObject(blob);
+						blobHeader = new BlobId(IoHash.Compute(blob));
+                        break;
+                    }
+                    case CustomMediaTypeNames.UnrealCompactBinary:
+                    {
+                        await using MemoryStream ms = new MemoryStream();
+                        await using Stream payloadStream = payload.GetStream();
+                        await payloadStream.CopyToAsync(ms, cancellationToken);
+                        payloadObject = new CbObject(ms.ToArray());
+                        break;
+                    }
+                    case MediaTypeNames.Application.Octet:
+                    {
+                        blobHeader = await _blobStore.PutObjectAsync(ns, payload, headerHash, cancellationToken);
+
+                        CbWriter writer = new CbWriter();
+                        writer.BeginObject();
+						writer.WriteBinaryAttachment("RawHash", blobHeader.Hash);
+                        writer.WriteInteger("RawSize", payload.Length);
+                        writer.EndObject();
+
+                        byte[] blob = writer.ToByteArray();
+                        payloadObject = new CbObject(blob);
+						blobHeader = new BlobId(IoHash.Compute(blob));
+                        break;
+                    }
+                    default:
+                        throw new Exception($"Unknown request type {Request.ContentType}, if submitting a blob please use {MediaTypeNames.Application.Octet}");
+                }
+            }
+            catch (HashMismatchException e)
+            {
+                return BadRequest(new ProblemDetails
+                {
+                    Title = $"Incorrect hash, got hash \"{e.SuppliedHash}\" but hash of content was determined to be \"{e.ContentHash}\""
+                });
+            }
+            catch (ClientSendSlowException e)
+            {
+                return Problem(e.Message, null, (int)HttpStatusCode.RequestTimeout);
+            }
+
+            (ContentId[] missingReferences, BlobId[] missingBlobs) = await _refService.PutAsync(ns, bucket, key, blobHeader, payloadObject);
+
+			List<IoHash> missingHashes = new List<IoHash>();
+			missingHashes.AddRange(missingReferences.Select(x => x.Hash));
+            missingHashes.AddRange(missingBlobs.Select(x => x.Hash));
+            return Ok(new PutObjectResponse(missingHashes.ToArray()));
+        }
+
+        [HttpPut("{ns}/{bucket}/{key}", Order = 300)]
+        [DisableRequestSizeLimit]
+        [RequiredContentType(CustomMediaTypeNames.UnrealCompactBinaryPackage)]
+        public async Task<IActionResult> PutPackage(
+            [FromRoute][Required] NamespaceId ns,
+            [FromRoute][Required] BucketId bucket,
+            [FromRoute][Required] RefId key,
+			CancellationToken cancellationToken)
+        {
+            ActionResult? accessResult = await _requestHelper.HasAccessToNamespaceAsync(User, Request, ns, new [] { StorageAclAction.WriteRefs });
+            if (accessResult != null)
+            {
+                return accessResult;
+            }
+
+            _diagnosticContext.Set("Content-Length", Request.ContentLength ?? -1);
+
+            byte[] b = await ReadRawBodyAsync(Request);
+            CbPackageReader packageReader = await CbPackageReader.Create(new MemoryStream(b));
+
+            try
+            {
+                await foreach ((CbPackageAttachmentEntry entry, byte[] blob) in packageReader.IterateAttachments())
+                {
+                    if (entry.Flags.HasFlag(CbPackageAttachmentFlags.IsError))
+                    {
+                        return BadRequest(new ProblemDetails
+                        {
+                            Title = $"Package contained attachment with error {entry.AttachmentHash}\""
+                        });
+                    }
+                    if (entry.Flags.HasFlag(CbPackageAttachmentFlags.IsCompressed))
+                    {
+                        using MemoryBufferedPayload payload = new MemoryBufferedPayload(blob);
+                        await _blobStore.PutCompressedObject(ns, payload, new ContentId(entry.AttachmentHash), HttpContext.RequestServices);
+                    }
+                    else
+                    {
+                        await _blobStore.PutObjectAsync(ns, blob, new BlobId(entry.AttachmentHash), cancellationToken);
+                    }
+                }
+            }
+            catch (HashMismatchException e)
+            {
+                return BadRequest(new ProblemDetails
+                {
+                    Title = $"Incorrect hash, got hash \"{e.SuppliedHash}\" but hash of content was determined to be \"{e.ContentHash}\""
+                });
+            }
+            
+            CbObject rootObject = packageReader.RootObject;
+            BlobId rootObjectHash = new BlobId(packageReader.RootHash);
+
+            (ContentId[] missingReferences, BlobId[] missingBlobs) = await _refService.PutAsync(ns, bucket, key, rootObjectHash, rootObject);
+			return Ok(new PutObjectResponse(missingReferences, missingBlobs));
+        }
+
+		static async Task<byte[]> ReadRawBodyAsync(HttpRequest request)
+		{
+			long? contentLength = request.GetTypedHeaders().ContentLength;
+
+			if (contentLength == null)
+			{
+				throw new Exception("Expected content-length on all raw body requests");
+			}
+
+			Tracer? tracer = request.HttpContext.RequestServices.GetService<Tracer>();
+			using TelemetrySpan? scope = tracer?.StartActiveSpan("readbody")
+				.SetAttribute("operation.name", "readbody")
+				.SetAttribute("Content-Length", contentLength.Value);
+
+			await using MemoryStream ms = new MemoryStream((int)contentLength);
+			DateTime readStart = DateTime.Now;
+			await request.Body.CopyToAsync(ms);
+			TimeSpan duration = DateTime.Now - readStart;
+
+#if WITH_DOGSTATSD
+			double rate = contentLength.Value / duration.TotalSeconds;
+            StatsdClient.DogStatsd.Histogram("jupiter.stream_throughput", rate, tags: new string[] {"sourceIdentifier:" + "readbody"});
+#endif
+			return ms.GetBuffer();
+		}
+
+		[HttpPost("{ns}/{bucket}/{key}/finalize/{hash}.{format?}")]
+        public async Task<IActionResult> FinalizeObject(
+            [FromRoute] [Required] NamespaceId ns,
+            [FromRoute] [Required] BucketId bucket,
+            [FromRoute] [Required] RefId key,
+            [FromRoute] [Required] BlobId hash)
+        {
+            ActionResult? accessResult = await _requestHelper.HasAccessToNamespaceAsync(User, Request, ns, new [] { StorageAclAction.WriteRefs });
+            if (accessResult != null)
+            {
+                return accessResult;
+            }
+
+            try
+            {
+                (ContentId[] missingReferences, BlobId[] missingBlobs) = await _refService.FinalizeAsync(ns, bucket, key, hash);
+				return Ok(new PutObjectResponse(missingReferences, missingBlobs));
+            }
+            catch (ObjectHashMismatchException e)
+            {
+                return BadRequest(e.Message);
+            }
+            catch (RefNotFoundException e)
+            {
+                return NotFound(e.Message);
+            }
+        }
+
+        [HttpPost("{ns}")]
+        [Consumes(CustomMediaTypeNames.UnrealCompactBinary)]
+        [Produces(CustomMediaTypeNames.UnrealCompactBinary)]
+        public async Task<IActionResult> Batch(
+            [FromRoute] [Required] NamespaceId ns,
+            [FromBody] [Required] RefBatchOps ops)
+        {
+            AclAction ActionForOp(RefBatchOps.RefBatchOp.RefOperation op)
+            {
+                switch (op)
+                {
+                    case RefBatchOps.RefBatchOp.RefOperation.GET:
+                        return StorageAclAction.ReadRefs;
+                    case RefBatchOps.RefBatchOp.RefOperation.PUT:
+                        return StorageAclAction.WriteRefs;
+                    case RefBatchOps.RefBatchOp.RefOperation.HEAD:
+						return StorageAclAction.ReadRefs;
+                    default:
+                        throw new ArgumentOutOfRangeException(nameof(op), op, null);
+                }
+            }
+
+            AclAction[] requiredActions = ops.Ops.Select(op => ActionForOp(op.Op)).ToArray();
+
+            ActionResult? accessResult = await _requestHelper.HasAccessToNamespaceAsync(User, Request, ns, requiredActions);
+            if (accessResult != null)
+            {
+                return accessResult;
+            }
+
+            HashSet<uint> usedOpIds = new HashSet<uint>();
+            foreach (RefBatchOps.RefBatchOp batchOp in ops.Ops)
+            {
+                bool added = usedOpIds.Add(batchOp.OpId);
+                if (!added)
+                {
+                    return BadRequest(new ProblemDetails { Title = $"Duplicate op ids used for id: {batchOp.OpId}" });
+                }
+            }
+            ConcurrentDictionary<uint, (CbObject, HttpStatusCode)> results = new();
+
+            async Task<(CbObject, HttpStatusCode)> BatchGetOp(RefBatchOps.RefBatchOp op)
+            {
+                try
+                {
+                    (RefRecord objectRecord, BlobContents? blob) = await _refService.GetAsync(ns, op.Bucket, op.Key, Array.Empty<string>());
+
+                    if (!objectRecord.IsFinalized)
+                    {
+                        throw new Exception("Object is not finalized");
+                    }
+
+                    if (blob == null)
+                    {
+                        throw new Exception();
+                    }
+
+                    CbObject cb = new CbObject(await blob.Stream.ToByteArray());
+
+                    if (op.ResolveAttachments ?? false)
+                    {
+                        IAsyncEnumerable<BlobId> references = _referenceResolver.GetReferencedBlobsAsync(ns, cb);
+                        List<BlobId>? _ = await references.ToListAsync();
+                    }
+
+                    return (cb, HttpStatusCode.OK);
+                }
+                catch (Exception ex) when( ex is RefNotFoundException or PartialReferenceResolveException or ReferenceIsMissingBlobsException)
+                {
+                    return ToErrorResult(ex, HttpStatusCode.NotFound);
+                }
+                catch (Exception e)
+                {
+                    return ToErrorResult(e);
+                }
+            }
+
+            async Task<(CbObject, HttpStatusCode)> BatchHeadOp(RefBatchOps.RefBatchOp op)
+            {
+                try
+                {
+                    (RefRecord record, BlobContents? blob) = await _refService.GetAsync(ns, op.Bucket, op.Key, new string[] { "blobIdentifier" });
+
+                    if (!record.IsFinalized)
+                    {
+                        return (CbObject.Build(writer => writer.WriteBool("exists", false)), HttpStatusCode.NotFound);
+                    }
+
+                    blob ??= await _blobStore.GetObjectAsync(ns, record.BlobIdentifier);
+
+                    if (op.ResolveAttachments ?? false)
+                    {
+                        byte[] blobContents = await blob.Stream.ToByteArray();
+                        CbObject cb = new CbObject(blobContents);
+                        // the reference resolver will throw if any blob is missing, so no need to do anything other then process each reference
+                        IAsyncEnumerable<BlobId> references = _referenceResolver.GetReferencedBlobsAsync(ns, cb);
+                        List<BlobId>? _ = await references.ToListAsync();
+                    }
+
+                    if (blob == null)
+                    {
+                        throw new Exception();
+                    }
+
+                    return (CbObject.Build(writer => writer.WriteBool("exists", true)), HttpStatusCode.OK);
+                }
+                catch (Exception ex) when( ex is RefNotFoundException or PartialReferenceResolveException or ReferenceIsMissingBlobsException)
+                {
+                    return (CbObject.Build(writer => writer.WriteBool("exists", false)), HttpStatusCode.NotFound);
+                }
+                catch (Exception e)
+                {
+                    return ToErrorResult(e);
+                }
+            }
+
+            async Task<(CbObject, HttpStatusCode)> BatchPutOp(RefBatchOps.RefBatchOp op)
+            {
+                try
+                {
+                    if (op.Payload == null || op.Payload.Equals(CbObject.Empty))
+                    {
+                        throw new Exception($"Missing payload for operation: {op.OpId}");
+                    }
+
+                    if (op.PayloadHash == null)
+                    {
+                        throw new Exception($"Missing payload hash for operation: {op.OpId}");
+                    }
+
+                    IoHash headerHash = op.PayloadHash.Value;
+					IoHash objectHash = IoHash.Compute(op.Payload.GetView().Span);
+				
+                    if (!headerHash.Equals(objectHash))
+                    {
+                        throw new HashMismatchException(op.PayloadHash.Value, objectHash);
+                    }
+
+                    (ContentId[] missingReferences, BlobId[] missingBlobs) = await _refService.PutAsync(ns, op.Bucket, op.Key, new BlobId(objectHash), op.Payload);
+                    return (CbSerializer.Serialize(new PutObjectResponse(missingReferences, missingBlobs)), HttpStatusCode.OK);
+                }
+                catch (Exception e)
+                {
+                    return ToErrorResult(e);
+                }
+            }
+
+            await Parallel.ForEachAsync(ops.Ops, CancellationToken.None, async (op, token) =>
+            {
+                switch (op.Op)
+                {
+                    case RefBatchOps.RefBatchOp.RefOperation.GET:
+                        results.TryAdd(op.OpId, await BatchGetOp(op));
+                        break;
+                    case RefBatchOps.RefBatchOp.RefOperation.PUT:
+                        results.TryAdd(op.OpId, await BatchPutOp(op));
+                        break;
+                    case RefBatchOps.RefBatchOp.RefOperation.HEAD:
+                        results.TryAdd(op.OpId, await BatchHeadOp(op));
+                        break;
+                    case RefBatchOps.RefBatchOp.RefOperation.INVALID:
+                    default:
+                        throw new NotImplementedException($"Unknown op type {op.Op}");
+                }
+
+                await Task.CompletedTask;
+            });
+
+            return Ok(new BatchOpsResponse()
+            {
+                Results = results.Select(result =>
+                {
+                    return new BatchOpsResponse.OpResponses()
+                    {
+                        OpId = result.Key,
+                        Response = result.Value.Item1,
+                        StatusCode = (int)result.Value.Item2
+                    };
+                }).ToList()
+            });
+        }
+
+        private static (CbObject, HttpStatusCode) ToErrorResult(Exception exception, HttpStatusCode statusCode = HttpStatusCode.InternalServerError)
+        {
+            Exception e = exception;
+            CbWriter writer = new CbWriter();
+            writer.BeginObject();
+            writer.WriteString("title", e.Message);
+            writer.WriteInteger("status", (int)statusCode);
+            if (e.StackTrace != null)
+            {
+                writer.WriteString("stackTrace", e.StackTrace);
+            }
+            writer.EndObject();
+            return (writer.ToObject(), statusCode);
+        }
+
+        /// <summary>
+        /// Delete a individual refs key
+        /// </summary>
+        /// <param name="ns">Namespace. Each namespace is completely separated from each other. Use for different types of data that is never expected to be similar (between two different games for instance)</param>
+        /// <param name="bucket">The category/type of record you are caching. Is a clustered key together with the actual key, but all records in the same bucket can be dropped easily.</param>
+        /// <param name="key">The unique name of this particular key</param>
+        [HttpDelete("{ns}/{bucket}/{key}", Order = 500)]
+        [ProducesResponseType(200)]
+        [ProducesResponseType(404)]
+        public async Task<IActionResult> Delete(
+            [FromRoute] [Required] NamespaceId ns,
+            [FromRoute] [Required] BucketId bucket,
+            [FromRoute] [Required] RefId key)
+        {
+            ActionResult? accessResult = await _requestHelper.HasAccessToNamespaceAsync(User, Request, ns, new [] { StorageAclAction.DeleteRefs });
+            if (accessResult != null)
+            {
+                return accessResult;
+            }
+
+            try
+            {
+                bool deleted = await _refService.DeleteAsync(ns, bucket, key);
+                if (!deleted)
+                {
+                    return NotFound(new ProblemDetails { Title = $"Object {key} in bucket {bucket} and namespace {ns} did not exist" });
+                }
+                return Ok(new RefDeletedResponse(deleted ? 1 : 0));
+            }
+            catch (NamespaceNotFoundException e)
+            {
+                return NotFound(new ProblemDetails { Title = $"Namespace {e.Namespace} did not exist" });
+            }
+            catch (RefNotFoundException)
+            {
+                return NotFound(new ProblemDetails { Title = $"Object {key} in bucket {bucket} and namespace {ns} did not exist" });
+            }
+        }
+    }
+
+    public class RefDeletedResponse
+    {
+        public RefDeletedResponse()
+        {
+
+        }
+
+        public RefDeletedResponse(int deletedCount)
+        {
+            DeletedCount = deletedCount;
+        }
+
+        [CbField("deletedCount")]
+        public int DeletedCount { get; set; }
+    }
+
+    public class BucketDeletedResponse
+    {
+        public BucketDeletedResponse()
+        {
+
+        }
+
+        public BucketDeletedResponse(long countOfDeletedRecords)
+        {
+            CountOfDeletedRecords = countOfDeletedRecords;
+        }
+
+        [CbField("countOfDeletedRecords")]
+        public long CountOfDeletedRecords { get; set; }
+    }
+
+    public class RefBatchOps
+    {
+        public RefBatchOps()
+        {
+            Ops = Array.Empty<RefBatchOp>();
+        }
+
+        [System.Diagnostics.CodeAnalysis.SuppressMessage("Design", "CA1034:Nested types should not be visible", Justification = "For serialization only")]
+        public class RefBatchOp
+        {
+            public RefBatchOp()
+            {
+                Payload = null;
+                PayloadHash = null;
+            }
+
+            public enum RefOperation
+            {
+                INVALID,
+                GET,
+                PUT,
+                HEAD,
+            }
+
+            [Required]
+            [CbField("opId")]
+            public uint OpId { get; set; }
+
+            [CbField("op")]
+            [JsonIgnore]
+            [UsedImplicitly]
+            public string OpString
+            {
+                get => Op.ToString();
+                set => Op = Enum.Parse<RefOperation>(value);
+            }
+
+            [Required]
+            public RefOperation Op { get; set; } = RefOperation.INVALID;
+
+            [Required]
+            [CbField("bucket")]
+            public BucketId Bucket { get; set; }
+
+            [Required]
+            [CbField("key")]
+            public RefId Key { get; set; }
+
+            [CbField("resolveAttachments")]
+            public bool? ResolveAttachments { get; set; } = null;
+
+            [CbField("payload")]
+            public CbObject? Payload { get; set; } = null;
+
+            [CbField("payloadHash")] 
+            public IoHash? PayloadHash { get; set; } = null;
+        }
+
+        [CbField("ops")]
+        public RefBatchOp[] Ops { get; set; }
+    }
+
+    public class BatchOpsResponse
+    {
+        public BatchOpsResponse()
+        {
+
+        }
+
+        [System.Diagnostics.CodeAnalysis.SuppressMessage("Design", "CA1034:Nested types should not be visible", Justification = "For serialization only")]
+        public class OpResponses
+        {
+            public OpResponses()
+            {
+
+            }
+            [CbField("opId")]
+            public uint OpId { get; set; }
+
+            [CbField("response")]
+            public CbObject Response { get; set; } = null!;
+
+            [CbField("statusCode")]
+            public int StatusCode { get; set; }
+        }
+
+        [CbField("results")]
+        [System.Diagnostics.CodeAnalysis.SuppressMessage("Usage", "CA2227:Collection properties should be read only", Justification = "Used by serialization")]
+        public List<OpResponses> Results { get; set; } = new List<OpResponses>();
+    }
+
+    public class RefMetadataResponse
+    {
+        public RefMetadataResponse()
+        {
+            InlinePayload = null!;
+        }
+
+        [JsonConstructor]
+        public RefMetadataResponse(NamespaceId ns, BucketId bucket, RefId name, BlobId payloadIdentifier, DateTime lastAccess, bool isFinalized, byte[]? inlinePayload)
+        {
+            Ns = ns;
+            Bucket = bucket;
+            Name = name;
+            PayloadIdentifier = payloadIdentifier;
+            LastAccess = lastAccess;
+            IsFinalized = isFinalized;
+            InlinePayload = inlinePayload;
+        }
+
+        public RefMetadataResponse(RefRecord objectRecord)
+        {
+            Ns = objectRecord.Namespace;
+            Bucket = objectRecord.Bucket;
+            Name = objectRecord.Name;
+            PayloadIdentifier = objectRecord.BlobIdentifier;
+            LastAccess = objectRecord.LastAccess;
+            IsFinalized = objectRecord.IsFinalized;
+            InlinePayload = objectRecord.InlinePayload;
+        }
+
+        [CbField("ns")]
+        public NamespaceId Ns { get; set; }
+        
+        [CbField("bucket")]
+        public BucketId Bucket { get; set; }
+
+        [CbField("name")]
+        public RefId Name { get; set; }
+
+        [CbField("payloadIdentifier")]
+        public BlobId PayloadIdentifier { get; set; }
+
+        [CbField("lastAccess")]
+        public DateTime LastAccess { get; set; }
+
+        [CbField("isFinalized")]
+        public bool IsFinalized { get; set; }
+
+        [CbField("inlinePayload")]
+        public byte[]? InlinePayload { get; set; }
+    }
+
+    public class PutObjectResponse
+    {
+        public PutObjectResponse()
+        {
+            Needs = null!;
+        }
+
+        public PutObjectResponse(IoHash[] missingReferences)
+        {
+            Needs = missingReferences;
+        }
+
+		public PutObjectResponse(IEnumerable<ContentId> missingContentIds, IEnumerable<BlobId> missingBlobIds)
+			: this(Enumerable.Concat(missingContentIds.Select(x => x.Hash), missingBlobIds.Select(x => x.Hash)).ToArray())
+		{
+		}
+
+        [CbField("needs")]
+        public IoHash[] Needs { get; set; }
+    }
+
+    public class ExistCheckMultipleRefsResponse
+    {
+        public ExistCheckMultipleRefsResponse(List<(BucketId, RefId)> missing)
+        {
+            Missing = missing.Select(pair =>
+            {
+                (BucketId bucketId, RefId ioHashKey) = pair;
+                return new MissingReference()
+                {
+                    Bucket = bucketId,
+                    Key = ioHashKey,
+                };
+            }).ToList();
+        }
+
+        [JsonConstructor]
+        public ExistCheckMultipleRefsResponse(List<MissingReference> missing)
+        {
+            Missing = missing;
+        }
+
+        [CbField("missing")]
+        [System.Diagnostics.CodeAnalysis.SuppressMessage("Usage", "CA2227:Collection properties should be read only", Justification = "Used by serialization")]
+        public List<MissingReference> Missing { get; set; }
+
+        [System.Diagnostics.CodeAnalysis.SuppressMessage("Design", "CA1034:Nested types should not be visible", Justification = "For serialization only")]
+        public class MissingReference
+        {
+            [CbField("bucket")]
+            public BucketId Bucket { get; set; }
+
+            [CbField("key")]
+            public RefId Key { get; set; }
+        }
+    }
+
+    public class GetNamespacesResponse
+    {
+        [JsonConstructor]
+        public GetNamespacesResponse(NamespaceId[] namespaces)
+        {
+            Namespaces = namespaces;
+        }
+
+        public NamespaceId[] Namespaces { get; set; }
+    }
+
+    public class HashMismatchException : Exception
+    {
+        public IoHash SuppliedHash { get; }
+        public IoHash ContentHash { get; }
+
+        public HashMismatchException(IoHash suppliedHash, IoHash contentHash) : base($"ID was not a hash of the content uploaded. Supplied hash was: {suppliedHash} but hash of content was {contentHash}")
+        {
+            SuppliedHash = suppliedHash;
+            ContentHash = contentHash;
+        }
+    }
+}
diff --git a/Engine/Source/Programs/Horde/Horde.Server/Startup.cs b/Engine/Source/Programs/Horde/Horde.Server/Startup.cs
--- a/Engine/Source/Programs/Horde/Horde.Server/Startup.cs
+++ b/Engine/Source/Programs/Horde/Horde.Server/Startup.cs
@@ -521,10 +521,10 @@ namespace Horde.Server
 			}
 
 			// DDC
-			services.AddSingleton<IDdcBlobService, DdcBlobService>();
-			services.AddSingleton<IDdcRefService, DdcRefService>();
+			services.AddSingleton<IBlobService, BlobService>();
+			services.AddSingleton<IRefService, RefService>();
 			services.AddSingleton<IReferenceResolver, ReferenceResolver>();
-			services.AddSingleton<IDdcContentIdService, ContentIdStore>();
+			services.AddSingleton<IContentIdService, ContentIdService>();
 
 			// Secret providers
 			services.AddSingleton<SecretService>();
