Horde: Fix DI failure during tests.

diff --git a/Engine/Source/Programs/Horde/Horde.Server/Ddc/RefService.cs b/Engine/Source/Programs/Horde/Horde.Server/Ddc/RefService.cs
--- a/Engine/Source/Programs/Horde/Horde.Server/Ddc/RefService.cs
+++ b/Engine/Source/Programs/Horde/Horde.Server/Ddc/RefService.cs
@@ -15,9 +15,9 @@ namespace Horde.Server.Ddc
 	class RefService : IRefService
 	{
 		readonly StorageService _storageService;
-		readonly BlobService _blobService;
+		readonly IBlobService _blobService;
 
-		public RefService(StorageService storageService, BlobService blobService)
+		public RefService(StorageService storageService, IBlobService blobService)
 		{
 			_storageService = storageService;
 			_blobService = blobService;
