Horde: Add parameters to MongoDB traces.

diff --git a/Engine/Source/Programs/Horde/Horde.Server/Utilities/MongoExtensions.cs b/Engine/Source/Programs/Horde/Horde.Server/Utilities/MongoExtensions.cs
--- a/Engine/Source/Programs/Horde/Horde.Server/Utilities/MongoExtensions.cs
+++ b/Engine/Source/Programs/Horde/Horde.Server/Utilities/MongoExtensions.cs
@@ -9,6 +9,7 @@ using System.Text.Json;
 using System.Threading;
 using System.Threading.Tasks;
 using MongoDB.Bson;
+using MongoDB.Bson.Serialization;
 using MongoDB.Driver;
 
 namespace Horde.Server.Utilities
@@ -28,6 +29,26 @@ namespace Horde.Server.Utilities
 			return BsonUtils.ToDateTimeFromMillisecondsSinceEpoch(BsonUtils.ToMillisecondsSinceEpoch(time));
 		}
 
+		/// <summary>
+		/// Renders a filter definition to a document using the default serializer registry
+		/// </summary>
+		public static BsonDocument Render<T>(this FilterDefinition<T> filter)
+		{
+			IBsonSerializerRegistry serializerRegistry = BsonSerializer.SerializerRegistry;
+			IBsonSerializer<T> documentSerializer = serializerRegistry.GetSerializer<T>();
+			return filter.Render(documentSerializer, serializerRegistry);
+		}
+
+		/// <summary>
+		/// Renders a filter definition to a document using the default serializer registry
+		/// </summary>
+		public static BsonValue Render<T>(this UpdateDefinition<T> update)
+		{
+			IBsonSerializerRegistry serializerRegistry = BsonSerializer.SerializerRegistry;
+			IBsonSerializer<T> documentSerializer = serializerRegistry.GetSerializer<T>();
+			return update.Render(documentSerializer, serializerRegistry);
+		}
+
 		/// <summary>
 		/// Executes a query with a particular index hint
 		/// </summary>
diff --git a/Engine/Source/Programs/Horde/Horde.Server/Utilities/MongoTracingCollection.cs b/Engine/Source/Programs/Horde/Horde.Server/Utilities/MongoTracingCollection.cs
--- a/Engine/Source/Programs/Horde/Horde.Server/Utilities/MongoTracingCollection.cs
+++ b/Engine/Source/Programs/Horde/Horde.Server/Utilities/MongoTracingCollection.cs
@@ -42,14 +42,14 @@ namespace Horde.Server.Utilities
 		public IAsyncCursor<TResult> Aggregate<TResult>(PipelineDefinition<T, TResult> pipeline, AggregateOptions options = null!, CancellationToken cancellationToken = default)
 		{
 			_upgradeTask.Wait(cancellationToken);
-			return _collection.Aggregate(pipeline, options,  cancellationToken);
+			return _collection.Aggregate(pipeline, options, cancellationToken);
 		}
 
 		/// <inheritdoc />
 		public IAsyncCursor<TResult> Aggregate<TResult>(IClientSessionHandle session, PipelineDefinition<T, TResult> pipeline, AggregateOptions options = null!, CancellationToken cancellationToken = default)
 		{
 			_upgradeTask.Wait(cancellationToken);
-			return _collection.Aggregate(session, pipeline, options,  cancellationToken);
+			return _collection.Aggregate(session, pipeline, options, cancellationToken);
 		}
 
 		/// <inheritdoc />
@@ -57,7 +57,7 @@ namespace Horde.Server.Utilities
 		{
 			await _upgradeTask.WaitAsync(cancellationToken);
 			using TelemetrySpan span = _tracer.StartMongoDbSpan(nameof(AggregateAsync), _collection);
-			return await _collection.AggregateAsync(pipeline, options,  cancellationToken);
+			return await _collection.AggregateAsync(pipeline, options, cancellationToken);
 		}
 
 		/// <inheritdoc />
@@ -103,14 +103,14 @@ namespace Horde.Server.Utilities
 		public BulkWriteResult<T> BulkWrite(IEnumerable<WriteModel<T>> requests, BulkWriteOptions options = null!, CancellationToken cancellationToken = default)
 		{
 			_upgradeTask.Wait(cancellationToken);
-			return _collection.BulkWrite(requests, options,  cancellationToken);
+			return _collection.BulkWrite(requests, options, cancellationToken);
 		}
 
 		/// <inheritdoc />
 		public BulkWriteResult<T> BulkWrite(IClientSessionHandle session, IEnumerable<WriteModel<T>> requests, BulkWriteOptions options = null!, CancellationToken cancellationToken = default)
 		{
 			_upgradeTask.Wait(cancellationToken);
-			return _collection.BulkWrite(session, requests, options,  cancellationToken);
+			return _collection.BulkWrite(session, requests, options, cancellationToken);
 		}
 
 		/// <inheritdoc />
@@ -118,7 +118,7 @@ namespace Horde.Server.Utilities
 		{
 			await _upgradeTask.WaitAsync(cancellationToken);
 			using TelemetrySpan span = _tracer.StartMongoDbSpan(nameof(BulkWriteAsync), _collection);
-			return await _collection.BulkWriteAsync(requests, options,  cancellationToken);
+			return await _collection.BulkWriteAsync(requests, options, cancellationToken);
 		}
 
 		/// <inheritdoc />
@@ -126,21 +126,21 @@ namespace Horde.Server.Utilities
 		{
 			await _upgradeTask.WaitAsync(cancellationToken);
 			using TelemetrySpan span = _tracer.StartMongoDbSpan(nameof(BulkWriteAsync), _collection);
-			return await _collection.BulkWriteAsync(session, requests, options,  cancellationToken);
+			return await _collection.BulkWriteAsync(session, requests, options, cancellationToken);
 		}
 
 		/// <inheritdoc />
 		public long Count(FilterDefinition<T> filter, CountOptions options = null!, CancellationToken cancellationToken = default)
 		{
 			_upgradeTask.Wait(cancellationToken);
-			return _collection.Count(filter, options,  cancellationToken);
+			return _collection.Count(filter, options, cancellationToken);
 		}
 
 		/// <inheritdoc />
 		public long Count(IClientSessionHandle session, FilterDefinition<T> filter, CountOptions options = null!, CancellationToken cancellationToken = default)
 		{
 			_upgradeTask.Wait(cancellationToken);
-			return _collection.Count(session, filter, options,  cancellationToken);
+			return _collection.Count(session, filter, options, cancellationToken);
 		}
 
 		/// <inheritdoc />
@@ -148,7 +148,7 @@ namespace Horde.Server.Utilities
 		{
 			await _upgradeTask.WaitAsync(cancellationToken);
 			using TelemetrySpan span = _tracer.StartMongoDbSpan(nameof(CountAsync), _collection);
-			return await _collection.CountAsync(filter, options,  cancellationToken);
+			return await _collection.CountAsync(filter, options, cancellationToken);
 		}
 
 		/// <inheritdoc />
@@ -156,164 +156,164 @@ namespace Horde.Server.Utilities
 		{
 			await _upgradeTask.WaitAsync(cancellationToken);
 			using TelemetrySpan span = _tracer.StartMongoDbSpan(nameof(CountAsync), _collection);
-			return await _collection.CountAsync(session, filter, options,  cancellationToken);
+			return await _collection.CountAsync(session, filter, options, cancellationToken);
 		}
 
 		/// <inheritdoc />
 		public long CountDocuments(FilterDefinition<T> filter, CountOptions options = null!, CancellationToken cancellationToken = default)
 		{
 			_upgradeTask.Wait(cancellationToken);
-			return _collection.CountDocuments(filter, options,  cancellationToken);
+			return _collection.CountDocuments(filter, options, cancellationToken);
 		}
 
 		/// <inheritdoc />
 		public long CountDocuments(IClientSessionHandle session, FilterDefinition<T> filter, CountOptions options = null!, CancellationToken cancellationToken = default)
 		{
 			_upgradeTask.Wait(cancellationToken);
-			return _collection.CountDocuments(session, filter, options,  cancellationToken);
+			return _collection.CountDocuments(session, filter, options, cancellationToken);
 		}
 
 		/// <inheritdoc />
 		public async Task<long> CountDocumentsAsync(FilterDefinition<T> filter, CountOptions options = null!, CancellationToken cancellationToken = default)
 		{
 			await _upgradeTask;
-			using TelemetrySpan span = _tracer.StartMongoDbSpan(nameof(CountDocumentsAsync), _collection);
-			return await _collection.CountDocumentsAsync(filter, options,  cancellationToken);
+			using TelemetrySpan span = _tracer.StartMongoDbSpan(nameof(CountDocumentsAsync), _collection, filter);
+			return await _collection.CountDocumentsAsync(filter, options, cancellationToken);
 		}
 
 		/// <inheritdoc />
 		public async Task<long> CountDocumentsAsync(IClientSessionHandle session, FilterDefinition<T> filter, CountOptions options = null!, CancellationToken cancellationToken = default)
 		{
 			await _upgradeTask.WaitAsync(cancellationToken);
-			using TelemetrySpan span = _tracer.StartMongoDbSpan(nameof(CountDocumentsAsync), _collection);
-			return await _collection.CountDocumentsAsync(session, filter, options,  cancellationToken);
+			using TelemetrySpan span = _tracer.StartMongoDbSpan(nameof(CountDocumentsAsync), _collection, filter);
+			return await _collection.CountDocumentsAsync(session, filter, options, cancellationToken);
 		}
 
 		/// <inheritdoc />
 		public DeleteResult DeleteMany(FilterDefinition<T> filter, CancellationToken cancellationToken = default)
 		{
 			_upgradeTask.Wait(cancellationToken);
-			return _collection.DeleteMany(filter,  cancellationToken);
+			return _collection.DeleteMany(filter, cancellationToken);
 		}
 
 		/// <inheritdoc />
 		public DeleteResult DeleteMany(FilterDefinition<T> filter, DeleteOptions options, CancellationToken cancellationToken = default)
 		{
 			_upgradeTask.Wait(cancellationToken);
-			return _collection.DeleteMany(filter, options,  cancellationToken);
+			return _collection.DeleteMany(filter, options, cancellationToken);
 		}
 
 		/// <inheritdoc />
 		public DeleteResult DeleteMany(IClientSessionHandle session, FilterDefinition<T> filter, DeleteOptions options = null!, CancellationToken cancellationToken = default)
 		{
 			_upgradeTask.Wait(cancellationToken);
-			return _collection.DeleteMany(session, filter, options,  cancellationToken);
+			return _collection.DeleteMany(session, filter, options, cancellationToken);
 		}
 
 		/// <inheritdoc />
 		public async Task<DeleteResult> DeleteManyAsync(FilterDefinition<T> filter, CancellationToken cancellationToken = default)
 		{
 			await _upgradeTask.WaitAsync(cancellationToken);
-			using TelemetrySpan span = _tracer.StartMongoDbSpan(nameof(DeleteManyAsync), _collection);
-			return await _collection.DeleteManyAsync(filter,  cancellationToken);
+			using TelemetrySpan span = _tracer.StartMongoDbSpan(nameof(DeleteManyAsync), _collection, filter);
+			return await _collection.DeleteManyAsync(filter, cancellationToken);
 		}
 
 		/// <inheritdoc />
 		public async Task<DeleteResult> DeleteManyAsync(FilterDefinition<T> filter, DeleteOptions options, CancellationToken cancellationToken = default)
 		{
 			await _upgradeTask.WaitAsync(cancellationToken);
-			using TelemetrySpan span = _tracer.StartMongoDbSpan(nameof(DeleteManyAsync), _collection);
-			return await _collection.DeleteManyAsync(filter, options,  cancellationToken);
+			using TelemetrySpan span = _tracer.StartMongoDbSpan(nameof(DeleteManyAsync), _collection, filter);
+			return await _collection.DeleteManyAsync(filter, options, cancellationToken);
 		}
 
 		/// <inheritdoc />
 		public async Task<DeleteResult> DeleteManyAsync(IClientSessionHandle session, FilterDefinition<T> filter, DeleteOptions options = null!, CancellationToken cancellationToken = default)
 		{
 			await _upgradeTask;
-			using TelemetrySpan span = _tracer.StartMongoDbSpan(nameof(DeleteManyAsync), _collection);
-			return await _collection.DeleteManyAsync(session, filter, options,  cancellationToken);
+			using TelemetrySpan span = _tracer.StartMongoDbSpan(nameof(DeleteManyAsync), _collection, filter);
+			return await _collection.DeleteManyAsync(session, filter, options, cancellationToken);
 		}
 
 		/// <inheritdoc />
 		public DeleteResult DeleteOne(FilterDefinition<T> filter, CancellationToken cancellationToken = default)
 		{
 			_upgradeTask.Wait(cancellationToken);
-			return _collection.DeleteOne(filter,  cancellationToken);
+			return _collection.DeleteOne(filter, cancellationToken);
 		}
 
 		/// <inheritdoc />
 		public DeleteResult DeleteOne(FilterDefinition<T> filter, DeleteOptions options, CancellationToken cancellationToken = default)
 		{
 			_upgradeTask.Wait(cancellationToken);
-			return _collection.DeleteOne(filter, options,  cancellationToken);
+			return _collection.DeleteOne(filter, options, cancellationToken);
 		}
 
 		/// <inheritdoc />
 		public DeleteResult DeleteOne(IClientSessionHandle session, FilterDefinition<T> filter, DeleteOptions options = null!, CancellationToken cancellationToken = default)
 		{
 			_upgradeTask.Wait(cancellationToken);
-			return _collection.DeleteOne(session, filter, options,  cancellationToken);
+			return _collection.DeleteOne(session, filter, options, cancellationToken);
 		}
 
 		/// <inheritdoc />
 		public async Task<DeleteResult> DeleteOneAsync(FilterDefinition<T> filter, CancellationToken cancellationToken = default)
 		{
 			await _upgradeTask.WaitAsync(cancellationToken);
-			using TelemetrySpan span = _tracer.StartMongoDbSpan(nameof(DeleteOneAsync), _collection);
-			return await _collection.DeleteOneAsync(filter,  cancellationToken);
+			using TelemetrySpan span = _tracer.StartMongoDbSpan(nameof(DeleteOneAsync), _collection, filter);
+			return await _collection.DeleteOneAsync(filter, cancellationToken);
 		}
 
 		/// <inheritdoc />
 		public async Task<DeleteResult> DeleteOneAsync(FilterDefinition<T> filter, DeleteOptions options, CancellationToken cancellationToken = default)
 		{
 			await _upgradeTask.WaitAsync(cancellationToken);
-			using TelemetrySpan span = _tracer.StartMongoDbSpan(nameof(DeleteOneAsync), _collection);
-			return await _collection.DeleteOneAsync(filter, options,  cancellationToken);
+			using TelemetrySpan span = _tracer.StartMongoDbSpan(nameof(DeleteOneAsync), _collection, filter);
+			return await _collection.DeleteOneAsync(filter, options, cancellationToken);
 		}
 
 		/// <inheritdoc />
 		public async Task<DeleteResult> DeleteOneAsync(IClientSessionHandle session, FilterDefinition<T> filter, DeleteOptions options = null!, CancellationToken cancellationToken = default)
 		{
 			await _upgradeTask.WaitAsync(cancellationToken);
-			using TelemetrySpan span = _tracer.StartMongoDbSpan(nameof(DeleteOneAsync), _collection);
-			return await _collection.DeleteOneAsync(session, filter, options,  cancellationToken);
+			using TelemetrySpan span = _tracer.StartMongoDbSpan(nameof(DeleteOneAsync), _collection, filter);
+			return await _collection.DeleteOneAsync(session, filter, options, cancellationToken);
 		}
 
 		/// <inheritdoc />
 		public IAsyncCursor<TField> Distinct<TField>(FieldDefinition<T, TField> field, FilterDefinition<T> filter, DistinctOptions options = null!, CancellationToken cancellationToken = default)
 		{
 			_upgradeTask.Wait(cancellationToken);
-			return _collection.Distinct(field, filter, options,  cancellationToken);
+			return _collection.Distinct(field, filter, options, cancellationToken);
 		}
 
 		/// <inheritdoc />
 		public IAsyncCursor<TField> Distinct<TField>(IClientSessionHandle session, FieldDefinition<T, TField> field, FilterDefinition<T> filter, DistinctOptions options = null!, CancellationToken cancellationToken = default)
 		{
 			_upgradeTask.Wait(cancellationToken);
-			return _collection.Distinct(session, field, filter, options,  cancellationToken);
+			return _collection.Distinct(session, field, filter, options, cancellationToken);
 		}
 
 		/// <inheritdoc />
 		public async Task<IAsyncCursor<TField>> DistinctAsync<TField>(FieldDefinition<T, TField> field, FilterDefinition<T> filter, DistinctOptions options = null!, CancellationToken cancellationToken = default)
 		{
 			await _upgradeTask.WaitAsync(cancellationToken);
-			using TelemetrySpan span = _tracer.StartMongoDbSpan(nameof(DistinctAsync), _collection);
-			return await _collection.DistinctAsync(field, filter, options,  cancellationToken);
+			using TelemetrySpan span = _tracer.StartMongoDbSpan(nameof(DistinctAsync), _collection, filter);
+			return await _collection.DistinctAsync(field, filter, options, cancellationToken);
 		}
 
 		/// <inheritdoc />
 		public async Task<IAsyncCursor<TField>> DistinctAsync<TField>(IClientSessionHandle session, FieldDefinition<T, TField> field, FilterDefinition<T> filter, DistinctOptions options = null!, CancellationToken cancellationToken = default)
 		{
 			await _upgradeTask.WaitAsync(cancellationToken);
-			using TelemetrySpan span = _tracer.StartMongoDbSpan(nameof(DistinctAsync), _collection);
-			return await _collection.DistinctAsync(session, field, filter, options,  cancellationToken);
+			using TelemetrySpan span = _tracer.StartMongoDbSpan(nameof(DistinctAsync), _collection, filter);
+			return await _collection.DistinctAsync(session, field, filter, options, cancellationToken);
 		}
 
 		/// <inheritdoc />
 		public long EstimatedDocumentCount(EstimatedDocumentCountOptions options = null!, CancellationToken cancellationToken = default)
 		{
 			_upgradeTask.Wait(cancellationToken);
-			return _collection.EstimatedDocumentCount(options,  cancellationToken);
+			return _collection.EstimatedDocumentCount(options, cancellationToken);
 		}
 
 		/// <inheritdoc />
@@ -321,97 +321,97 @@ namespace Horde.Server.Utilities
 		{
 			await _upgradeTask.WaitAsync(cancellationToken);
 			using TelemetrySpan span = _tracer.StartMongoDbSpan(nameof(EstimatedDocumentCountAsync), _collection);
-			return await _collection.EstimatedDocumentCountAsync(options,  cancellationToken);
+			return await _collection.EstimatedDocumentCountAsync(options, cancellationToken);
 		}
 
 		/// <inheritdoc />
 		public IAsyncCursor<TProjection> FindSync<TProjection>(FilterDefinition<T> filter, FindOptions<T, TProjection> options = null!, CancellationToken cancellationToken = default)
 		{
 			_upgradeTask.Wait(cancellationToken);
-			return _collection.FindSync(filter, options,  cancellationToken);
+			return _collection.FindSync(filter, options, cancellationToken);
 		}
 
 		/// <inheritdoc />
 		public IAsyncCursor<TProjection> FindSync<TProjection>(IClientSessionHandle session, FilterDefinition<T> filter, FindOptions<T, TProjection> options = null!, CancellationToken cancellationToken = default)
 		{
 			_upgradeTask.Wait(cancellationToken);
-			return _collection.FindSync(session, filter, options,  cancellationToken);
+			return _collection.FindSync(session, filter, options, cancellationToken);
 		}
 
 		/// <inheritdoc />
 		public async Task<IAsyncCursor<TProjection>> FindAsync<TProjection>(FilterDefinition<T> filter, FindOptions<T, TProjection> options = null!, CancellationToken cancellationToken = default)
 		{
 			await _upgradeTask.WaitAsync(cancellationToken);
-			using TelemetrySpan span = _tracer.StartMongoDbSpan(nameof(FindAsync), _collection);
-			return await _collection.FindAsync(filter, options,  cancellationToken);
+			using TelemetrySpan span = _tracer.StartMongoDbSpan(nameof(FindAsync), _collection, filter);
+			return await _collection.FindAsync(filter, options, cancellationToken);
 		}
 
 		/// <inheritdoc />
 		public async Task<IAsyncCursor<TProjection>> FindAsync<TProjection>(IClientSessionHandle session, FilterDefinition<T> filter, FindOptions<T, TProjection> options = null!, CancellationToken cancellationToken = default)
 		{
 			await _upgradeTask.WaitAsync(cancellationToken);
-			using TelemetrySpan span = _tracer.StartMongoDbSpan(nameof(FindAsync), _collection);
-			return await _collection.FindAsync(session, filter, options,  cancellationToken);
+			using TelemetrySpan span = _tracer.StartMongoDbSpan(nameof(FindAsync), _collection, filter);
+			return await _collection.FindAsync(session, filter, options, cancellationToken);
 		}
 
 		/// <inheritdoc />
 		public TProjection FindOneAndDelete<TProjection>(FilterDefinition<T> filter, FindOneAndDeleteOptions<T, TProjection> options = null!, CancellationToken cancellationToken = default)
 		{
 			_upgradeTask.Wait(cancellationToken);
-			return _collection.FindOneAndDelete(filter, options,  cancellationToken);
+			return _collection.FindOneAndDelete(filter, options, cancellationToken);
 		}
 
 		/// <inheritdoc />
 		public TProjection FindOneAndDelete<TProjection>(IClientSessionHandle session, FilterDefinition<T> filter, FindOneAndDeleteOptions<T, TProjection> options = null!, CancellationToken cancellationToken = default)
 		{
 			_upgradeTask.Wait(cancellationToken);
-			return _collection.FindOneAndDelete(session, filter, options,  cancellationToken);
+			return _collection.FindOneAndDelete(session, filter, options, cancellationToken);
 		}
 
 		/// <inheritdoc />
 		public async Task<TProjection> FindOneAndDeleteAsync<TProjection>(FilterDefinition<T> filter, FindOneAndDeleteOptions<T, TProjection> options = null!, CancellationToken cancellationToken = default)
 		{
 			await _upgradeTask.WaitAsync(cancellationToken);
-			using TelemetrySpan span = _tracer.StartMongoDbSpan(nameof(FindOneAndDeleteAsync), _collection);
-			return await _collection.FindOneAndDeleteAsync(filter, options,  cancellationToken);
+			using TelemetrySpan span = _tracer.StartMongoDbSpan(nameof(FindOneAndDeleteAsync), _collection, filter);
+			return await _collection.FindOneAndDeleteAsync(filter, options, cancellationToken);
 		}
 
 		/// <inheritdoc />
 		public async Task<TProjection> FindOneAndDeleteAsync<TProjection>(IClientSessionHandle session, FilterDefinition<T> filter, FindOneAndDeleteOptions<T, TProjection> options = null!, CancellationToken cancellationToken = default)
 		{
 			await _upgradeTask.WaitAsync(cancellationToken);
-			using TelemetrySpan span = _tracer.StartMongoDbSpan(nameof(FindOneAndDeleteAsync), _collection);
-			return await _collection.FindOneAndDeleteAsync(session, filter, options,  cancellationToken);
+			using TelemetrySpan span = _tracer.StartMongoDbSpan(nameof(FindOneAndDeleteAsync), _collection, filter);
+			return await _collection.FindOneAndDeleteAsync(session, filter, options, cancellationToken);
 		}
 
 		/// <inheritdoc />
 		public TProjection FindOneAndReplace<TProjection>(FilterDefinition<T> filter, T replacement, FindOneAndReplaceOptions<T, TProjection> options = null!, CancellationToken cancellationToken = default)
 		{
 			_upgradeTask.Wait(cancellationToken);
-			return _collection.FindOneAndReplace(filter, replacement, options,  cancellationToken);
+			return _collection.FindOneAndReplace(filter, replacement, options, cancellationToken);
 		}
 
 		/// <inheritdoc />
 		public TProjection FindOneAndReplace<TProjection>(IClientSessionHandle session, FilterDefinition<T> filter, T replacement, FindOneAndReplaceOptions<T, TProjection> options = null!, CancellationToken cancellationToken = default)
 		{
 			_upgradeTask.Wait(cancellationToken);
-			return _collection.FindOneAndReplace(session, filter, replacement, options,  cancellationToken);
+			return _collection.FindOneAndReplace(session, filter, replacement, options, cancellationToken);
 		}
 
 		/// <inheritdoc />
 		public async Task<TProjection> FindOneAndReplaceAsync<TProjection>(FilterDefinition<T> filter, T replacement, FindOneAndReplaceOptions<T, TProjection> options = null!, CancellationToken cancellationToken = default)
 		{
 			await _upgradeTask.WaitAsync(cancellationToken);
-			using TelemetrySpan span = _tracer.StartMongoDbSpan(nameof(FindOneAndReplaceAsync), _collection);
-			return await _collection.FindOneAndReplaceAsync(filter, replacement, options,  cancellationToken);
+			using TelemetrySpan span = _tracer.StartMongoDbSpan(nameof(FindOneAndReplaceAsync), _collection, filter, document: replacement);
+			return await _collection.FindOneAndReplaceAsync(filter, replacement, options, cancellationToken);
 		}
 
 		/// <inheritdoc />
 		public async Task<TProjection> FindOneAndReplaceAsync<TProjection>(IClientSessionHandle session, FilterDefinition<T> filter, T replacement, FindOneAndReplaceOptions<T, TProjection> options = null!, CancellationToken cancellationToken = default)
 		{
 			await _upgradeTask.WaitAsync(cancellationToken);
-			using TelemetrySpan span = _tracer.StartMongoDbSpan(nameof(FindOneAndReplaceAsync), _collection);
-			return await _collection.FindOneAndReplaceAsync(session, filter, replacement, options,  cancellationToken);
+			using TelemetrySpan span = _tracer.StartMongoDbSpan(nameof(FindOneAndReplaceAsync), _collection, filter, document: replacement);
+			return await _collection.FindOneAndReplaceAsync(session, filter, replacement, options, cancellationToken);
 		}
 
 		/// <inheritdoc />
@@ -419,51 +419,51 @@ namespace Horde.Server.Utilities
 			FindOneAndUpdateOptions<T, TProjection> options = null!, CancellationToken cancellationToken = default)
 		{
 			_upgradeTask.Wait(cancellationToken);
-			return _collection.FindOneAndUpdate(filter, update, options,  cancellationToken);
+			return _collection.FindOneAndUpdate(filter, update, options, cancellationToken);
 		}
 
 		/// <inheritdoc />
 		public TProjection FindOneAndUpdate<TProjection>(IClientSessionHandle session, FilterDefinition<T> filter, UpdateDefinition<T> update, FindOneAndUpdateOptions<T, TProjection> options = null!, CancellationToken cancellationToken = default)
 		{
 			_upgradeTask.Wait(cancellationToken);
-			return _collection.FindOneAndUpdate(session, filter, update, options,  cancellationToken);
+			return _collection.FindOneAndUpdate(session, filter, update, options, cancellationToken);
 		}
 
 		/// <inheritdoc />
 		public async Task<TProjection> FindOneAndUpdateAsync<TProjection>(FilterDefinition<T> filter, UpdateDefinition<T> update, FindOneAndUpdateOptions<T, TProjection> options = null!, CancellationToken cancellationToken = default)
 		{
 			await _upgradeTask.WaitAsync(cancellationToken);
-			using TelemetrySpan span = _tracer.StartMongoDbSpan(nameof(FindOneAndUpdateAsync), _collection);
-			return await _collection.FindOneAndUpdateAsync(filter, update, options,  cancellationToken);
+			using TelemetrySpan span = _tracer.StartMongoDbSpan(nameof(FindOneAndUpdateAsync), _collection, filter, update);
+			return await _collection.FindOneAndUpdateAsync(filter, update, options, cancellationToken);
 		}
 
 		/// <inheritdoc />
 		public async Task<TProjection> FindOneAndUpdateAsync<TProjection>(IClientSessionHandle session, FilterDefinition<T> filter, UpdateDefinition<T> update, FindOneAndUpdateOptions<T, TProjection> options = null!, CancellationToken cancellationToken = default)
 		{
 			await _upgradeTask.WaitAsync(cancellationToken);
-			using TelemetrySpan span = _tracer.StartMongoDbSpan(nameof(FindOneAndUpdateAsync), _collection);
-			return await _collection.FindOneAndUpdateAsync(session, filter, update, options,  cancellationToken);
+			using TelemetrySpan span = _tracer.StartMongoDbSpan(nameof(FindOneAndUpdateAsync), _collection, filter, update);
+			return await _collection.FindOneAndUpdateAsync(session, filter, update, options, cancellationToken);
 		}
 
 		/// <inheritdoc />
 		public void InsertOne(T document, InsertOneOptions options = null!, CancellationToken cancellationToken = default)
 		{
 			_upgradeTask.Wait(cancellationToken);
-			_collection.InsertOne(document, options,  cancellationToken);
+			_collection.InsertOne(document, options, cancellationToken);
 		}
 
 		/// <inheritdoc />
 		public void InsertOne(IClientSessionHandle session, T document, InsertOneOptions options = null!, CancellationToken cancellationToken = default)
 		{
 			_upgradeTask.Wait(cancellationToken);
-			_collection.InsertOne(session, document, options,  cancellationToken);
+			_collection.InsertOne(session, document, options, cancellationToken);
 		}
 
 		/// <inheritdoc />
 		public async Task InsertOneAsync(T document, CancellationToken cancellationToken)
 		{
 			await _upgradeTask.WaitAsync(cancellationToken);
-			using TelemetrySpan span = _tracer.StartMongoDbSpan(nameof(InsertOneAsync), _collection);
+			using TelemetrySpan span = _tracer.StartMongoDbSpan(nameof(InsertOneAsync), _collection, document: document);
 			await _collection.InsertOneAsync(document, cancellationToken);
 		}
 
@@ -471,30 +471,30 @@ namespace Horde.Server.Utilities
 		public async Task InsertOneAsync(T document, InsertOneOptions options = null!, CancellationToken cancellationToken = default)
 		{
 			await _upgradeTask.WaitAsync(cancellationToken);
-			using TelemetrySpan span = _tracer.StartMongoDbSpan(nameof(InsertOneAsync), _collection);
-			await _collection.InsertOneAsync(document, options,  cancellationToken);
+			using TelemetrySpan span = _tracer.StartMongoDbSpan(nameof(InsertOneAsync), _collection, document: document);
+			await _collection.InsertOneAsync(document, options, cancellationToken);
 		}
 
 		/// <inheritdoc />
 		public async Task InsertOneAsync(IClientSessionHandle session, T document, InsertOneOptions options = null!, CancellationToken cancellationToken = default)
 		{
 			await _upgradeTask.WaitAsync(cancellationToken);
-			using TelemetrySpan span = _tracer.StartMongoDbSpan(nameof(InsertOneAsync), _collection);
-			await _collection.InsertOneAsync(session, document, options,  cancellationToken);
+			using TelemetrySpan span = _tracer.StartMongoDbSpan(nameof(InsertOneAsync), _collection, document: document);
+			await _collection.InsertOneAsync(session, document, options, cancellationToken);
 		}
 
 		/// <inheritdoc />
 		public void InsertMany(IEnumerable<T> documents, InsertManyOptions options = null!, CancellationToken cancellationToken = default)
 		{
 			_upgradeTask.Wait(cancellationToken);
-			_collection.InsertMany(documents, options,  cancellationToken);
+			_collection.InsertMany(documents, options, cancellationToken);
 		}
 
 		/// <inheritdoc />
 		public void InsertMany(IClientSessionHandle session, IEnumerable<T> documents, InsertManyOptions options = null!, CancellationToken cancellationToken = default)
 		{
 			_upgradeTask.Wait(cancellationToken);
-			_collection.InsertMany(session, documents, options,  cancellationToken);
+			_collection.InsertMany(session, documents, options, cancellationToken);
 		}
 
 		/// <inheritdoc />
@@ -502,7 +502,7 @@ namespace Horde.Server.Utilities
 		{
 			await _upgradeTask.WaitAsync(cancellationToken);
 			using TelemetrySpan span = _tracer.StartMongoDbSpan(nameof(InsertManyAsync), _collection);
-			await _collection.InsertManyAsync(documents, options,  cancellationToken);
+			await _collection.InsertManyAsync(documents, options, cancellationToken);
 		}
 
 		/// <inheritdoc />
@@ -510,21 +510,21 @@ namespace Horde.Server.Utilities
 		{
 			await _upgradeTask.WaitAsync(cancellationToken);
 			using TelemetrySpan span = _tracer.StartMongoDbSpan(nameof(InsertManyAsync), _collection);
-			await _collection.InsertManyAsync(session, documents, options,  cancellationToken);
+			await _collection.InsertManyAsync(session, documents, options, cancellationToken);
 		}
 
 		/// <inheritdoc />
 		public IAsyncCursor<TResult> MapReduce<TResult>(BsonJavaScript map, BsonJavaScript reduce, MapReduceOptions<T, TResult> options = null!, CancellationToken cancellationToken = default)
 		{
 			_upgradeTask.Wait(cancellationToken);
-			return _collection.MapReduce(map, reduce, options,  cancellationToken);
+			return _collection.MapReduce(map, reduce, options, cancellationToken);
 		}
 
 		/// <inheritdoc />
 		public IAsyncCursor<TResult> MapReduce<TResult>(IClientSessionHandle session, BsonJavaScript map, BsonJavaScript reduce, MapReduceOptions<T, TResult> options = null!, CancellationToken cancellationToken = default)
 		{
 			_upgradeTask.Wait(cancellationToken);
-			return _collection.MapReduce(session, map, reduce, options,  cancellationToken);
+			return _collection.MapReduce(session, map, reduce, options, cancellationToken);
 		}
 
 		/// <inheritdoc />
@@ -532,7 +532,7 @@ namespace Horde.Server.Utilities
 		{
 			await _upgradeTask.WaitAsync(cancellationToken);
 			using TelemetrySpan span = _tracer.StartMongoDbSpan(nameof(MapReduceAsync), _collection);
-			return await _collection.MapReduceAsync(map, reduce, options,  cancellationToken);
+			return await _collection.MapReduceAsync(map, reduce, options, cancellationToken);
 		}
 
 		/// <inheritdoc />
@@ -540,7 +540,7 @@ namespace Horde.Server.Utilities
 		{
 			await _upgradeTask.WaitAsync(cancellationToken);
 			using TelemetrySpan span = _tracer.StartMongoDbSpan(nameof(MapReduceAsync), _collection);
-			return await _collection.MapReduceAsync(session, map, reduce, options,  cancellationToken);
+			return await _collection.MapReduceAsync(session, map, reduce, options, cancellationToken);
 		}
 
 		/// <inheritdoc />
@@ -553,82 +553,82 @@ namespace Horde.Server.Utilities
 		public ReplaceOneResult ReplaceOne(FilterDefinition<T> filter, T replacement, ReplaceOptions options = null!, CancellationToken cancellationToken = default)
 		{
 			_upgradeTask.Wait(cancellationToken);
-			return _collection.ReplaceOne(filter, replacement, options,  cancellationToken);
+			return _collection.ReplaceOne(filter, replacement, options, cancellationToken);
 		}
 
 		/// <inheritdoc />
 		public ReplaceOneResult ReplaceOne(FilterDefinition<T> filter, T replacement, UpdateOptions options, CancellationToken cancellationToken = default)
 		{
 			_upgradeTask.Wait(cancellationToken);
-			return _collection.ReplaceOne(filter, replacement, options,  cancellationToken);
+			return _collection.ReplaceOne(filter, replacement, options, cancellationToken);
 		}
 
 		/// <inheritdoc />
 		public ReplaceOneResult ReplaceOne(IClientSessionHandle session, FilterDefinition<T> filter, T replacement, ReplaceOptions options = null!, CancellationToken cancellationToken = default)
 		{
 			_upgradeTask.Wait(cancellationToken);
-			return _collection.ReplaceOne(session, filter, replacement, options,  cancellationToken);
+			return _collection.ReplaceOne(session, filter, replacement, options, cancellationToken);
 		}
 
 		/// <inheritdoc />
 		public ReplaceOneResult ReplaceOne(IClientSessionHandle session, FilterDefinition<T> filter, T replacement, UpdateOptions options, CancellationToken cancellationToken = default)
 		{
 			_upgradeTask.Wait(cancellationToken);
-			return _collection.ReplaceOne(session, filter, replacement, options,  cancellationToken);
+			return _collection.ReplaceOne(session, filter, replacement, options, cancellationToken);
 		}
 
 		/// <inheritdoc />
 		public async Task<ReplaceOneResult> ReplaceOneAsync(FilterDefinition<T> filter, T replacement, ReplaceOptions options = null!, CancellationToken cancellationToken = default)
 		{
 			await _upgradeTask.WaitAsync(cancellationToken);
-			using TelemetrySpan span = _tracer.StartMongoDbSpan(nameof(ReplaceOneAsync), _collection);
-			return await _collection.ReplaceOneAsync(filter, replacement, options,  cancellationToken);
+			using TelemetrySpan span = _tracer.StartMongoDbSpan(nameof(ReplaceOneAsync), _collection, filter, document: replacement);
+			return await _collection.ReplaceOneAsync(filter, replacement, options, cancellationToken);
 		}
 
 		/// <inheritdoc />
 		public async Task<ReplaceOneResult> ReplaceOneAsync(FilterDefinition<T> filter, T replacement, UpdateOptions options, CancellationToken cancellationToken = default)
 		{
 			await _upgradeTask.WaitAsync(cancellationToken);
-			using TelemetrySpan span = _tracer.StartMongoDbSpan(nameof(ReplaceOneAsync), _collection);
-			return await _collection.ReplaceOneAsync(filter, replacement, options,  cancellationToken);
+			using TelemetrySpan span = _tracer.StartMongoDbSpan(nameof(ReplaceOneAsync), _collection, filter, document: replacement);
+			return await _collection.ReplaceOneAsync(filter, replacement, options, cancellationToken);
 		}
 
 		/// <inheritdoc />
 		public async Task<ReplaceOneResult> ReplaceOneAsync(IClientSessionHandle session, FilterDefinition<T> filter, T replacement, ReplaceOptions options = null!, CancellationToken cancellationToken = default)
 		{
 			await _upgradeTask.WaitAsync(cancellationToken);
-			using TelemetrySpan span = _tracer.StartMongoDbSpan(nameof(ReplaceOneAsync), _collection);
-			return await _collection.ReplaceOneAsync(session, filter, replacement, options,  cancellationToken);
+			using TelemetrySpan span = _tracer.StartMongoDbSpan(nameof(ReplaceOneAsync), _collection, filter, document: replacement);
+			return await _collection.ReplaceOneAsync(session, filter, replacement, options, cancellationToken);
 		}
 
 		/// <inheritdoc />
 		public async Task<ReplaceOneResult> ReplaceOneAsync(IClientSessionHandle session, FilterDefinition<T> filter, T replacement, UpdateOptions options, CancellationToken cancellationToken = default)
 		{
 			await _upgradeTask.WaitAsync(cancellationToken);
-			using TelemetrySpan span = _tracer.StartMongoDbSpan(nameof(ReplaceOneAsync), _collection);
-			return await _collection.ReplaceOneAsync(session, filter, replacement, options,  cancellationToken);
+			using TelemetrySpan span = _tracer.StartMongoDbSpan(nameof(ReplaceOneAsync), _collection, filter, document: replacement);
+			return await _collection.ReplaceOneAsync(session, filter, replacement, options, cancellationToken);
 		}
 
 		/// <inheritdoc />
 		public UpdateResult UpdateMany(FilterDefinition<T> filter, UpdateDefinition<T> update, UpdateOptions options = null!, CancellationToken cancellationToken = default)
 		{
 			_upgradeTask.Wait(cancellationToken);
-			return _collection.UpdateMany(filter, update, options,  cancellationToken);
+			return _collection.UpdateMany(filter, update, options, cancellationToken);
 		}
 
 		/// <inheritdoc />
 		public UpdateResult UpdateMany(IClientSessionHandle session, FilterDefinition<T> filter, UpdateDefinition<T> update, UpdateOptions options = null!, CancellationToken cancellationToken = default)
 		{
 			_upgradeTask.Wait(cancellationToken);
-			return _collection.UpdateMany(session, filter, update, options,  cancellationToken);
+			return _collection.UpdateMany(session, filter, update, options, cancellationToken);
 		}
 
 		/// <inheritdoc />
 		public async Task<UpdateResult> UpdateManyAsync(FilterDefinition<T> filter, UpdateDefinition<T> update, UpdateOptions options = null!, CancellationToken cancellationToken = default)
 		{
 			await _upgradeTask.WaitAsync(cancellationToken);
-			using TelemetrySpan span = _tracer.StartMongoDbSpan(nameof(UpdateManyAsync), _collection);
-			return await _collection.UpdateManyAsync(filter, update, options,  cancellationToken);
+			using TelemetrySpan span = _tracer.StartMongoDbSpan(nameof(UpdateManyAsync), _collection, filter, update);
+			return await _collection.UpdateManyAsync(filter, update, options, cancellationToken);
 		}
 
 		/// <inheritdoc />
@@ -636,51 +636,51 @@ namespace Horde.Server.Utilities
 		{
 			await _upgradeTask.WaitAsync(cancellationToken);
 			using TelemetrySpan span = _tracer.StartMongoDbSpan(nameof(UpdateManyAsync), _collection);
-			return await _collection.UpdateManyAsync(session, filter, update, options,  cancellationToken);
+			return await _collection.UpdateManyAsync(session, filter, update, options, cancellationToken);
 		}
 
 		/// <inheritdoc />
 		public UpdateResult UpdateOne(FilterDefinition<T> filter, UpdateDefinition<T> update, UpdateOptions options = null!, CancellationToken cancellationToken = default)
 		{
 			_upgradeTask.Wait(cancellationToken);
-			return _collection.UpdateOne(filter, update, options,  cancellationToken);
+			return _collection.UpdateOne(filter, update, options, cancellationToken);
 		}
 
 		/// <inheritdoc />
 		public UpdateResult UpdateOne(IClientSessionHandle session, FilterDefinition<T> filter, UpdateDefinition<T> update, UpdateOptions options = null!, CancellationToken cancellationToken = default)
 		{
 			_upgradeTask.Wait(cancellationToken);
-			return _collection.UpdateOne(session, filter, update, options,  cancellationToken);
+			return _collection.UpdateOne(session, filter, update, options, cancellationToken);
 		}
 
 		/// <inheritdoc />
 		public async Task<UpdateResult> UpdateOneAsync(FilterDefinition<T> filter, UpdateDefinition<T> update, UpdateOptions options = null!, CancellationToken cancellationToken = default)
 		{
 			await _upgradeTask.WaitAsync(cancellationToken);
-			using TelemetrySpan span = _tracer.StartMongoDbSpan(nameof(UpdateOneAsync), _collection);
-			return await _collection.UpdateOneAsync(filter, update, options,  cancellationToken);
+			using TelemetrySpan span = _tracer.StartMongoDbSpan(nameof(UpdateOneAsync), _collection, filter, update);
+			return await _collection.UpdateOneAsync(filter, update, options, cancellationToken);
 		}
 
 		/// <inheritdoc />
 		public async Task<UpdateResult> UpdateOneAsync(IClientSessionHandle session, FilterDefinition<T> filter, UpdateDefinition<T> update, UpdateOptions options = null!, CancellationToken cancellationToken = default)
 		{
 			await _upgradeTask.WaitAsync(cancellationToken);
-			using TelemetrySpan span = _tracer.StartMongoDbSpan(nameof(UpdateOneAsync), _collection);
-			return await _collection.UpdateOneAsync(session, filter, update, options,  cancellationToken);
+			using TelemetrySpan span = _tracer.StartMongoDbSpan(nameof(UpdateOneAsync), _collection, filter, update);
+			return await _collection.UpdateOneAsync(session, filter, update, options, cancellationToken);
 		}
 
 		/// <inheritdoc />
 		public IChangeStreamCursor<TResult> Watch<TResult>(PipelineDefinition<ChangeStreamDocument<T>, TResult> pipeline, ChangeStreamOptions options = null!, CancellationToken cancellationToken = default)
 		{
 			_upgradeTask.Wait(cancellationToken);
-			return _collection.Watch(pipeline, options,  cancellationToken);
+			return _collection.Watch(pipeline, options, cancellationToken);
 		}
 
 		/// <inheritdoc />
 		public IChangeStreamCursor<TResult> Watch<TResult>(IClientSessionHandle session, PipelineDefinition<ChangeStreamDocument<T>, TResult> pipeline, ChangeStreamOptions options = null!, CancellationToken cancellationToken = default)
 		{
 			_upgradeTask.Wait(cancellationToken);
-			return _collection.Watch(session, pipeline, options,  cancellationToken);
+			return _collection.Watch(session, pipeline, options, cancellationToken);
 		}
 
 		/// <inheritdoc />
@@ -688,7 +688,7 @@ namespace Horde.Server.Utilities
 		{
 			await _upgradeTask.WaitAsync(cancellationToken);
 			using TelemetrySpan span = _tracer.StartMongoDbSpan(nameof(WatchAsync), _collection);
-			return await _collection.WatchAsync(pipeline, options,  cancellationToken);
+			return await _collection.WatchAsync(pipeline, options, cancellationToken);
 		}
 
 		/// <inheritdoc />
@@ -696,7 +696,7 @@ namespace Horde.Server.Utilities
 		{
 			await _upgradeTask.WaitAsync(cancellationToken);
 			using TelemetrySpan span = _tracer.StartMongoDbSpan(nameof(WatchAsync), _collection);
-			return await _collection.WatchAsync(session, pipeline, options,  cancellationToken);
+			return await _collection.WatchAsync(session, pipeline, options, cancellationToken);
 		}
 
 		/// <inheritdoc />
diff --git a/Engine/Source/Programs/Horde/Horde.Server/Utilities/OpenTelemetry.cs b/Engine/Source/Programs/Horde/Horde.Server/Utilities/OpenTelemetry.cs
--- a/Engine/Source/Programs/Horde/Horde.Server/Utilities/OpenTelemetry.cs
+++ b/Engine/Source/Programs/Horde/Horde.Server/Utilities/OpenTelemetry.cs
@@ -13,6 +13,8 @@ using Horde.Server.Streams;
 using Microsoft.AspNetCore.Http;
 using Microsoft.Extensions.DependencyInjection;
 using Microsoft.Extensions.Logging;
+using MongoDB.Bson;
+using MongoDB.Bson.Serialization;
 using MongoDB.Driver;
 using OpenTelemetry.Exporter;
 using OpenTelemetry.Logs;
@@ -308,9 +310,12 @@ public static class OpenTelemetrySpanExtensions
 	/// <param name="tracer">Current tracer being extended</param>
 	/// <param name="spanName">Name of the span</param>
 	/// <param name="collection">An optional MongoDB collection, the name will be used as an attribute</param>
+	/// <param name="filter">Optional filter to attach to the trace</param>
+	/// <param name="update">Optional update to attach to the trace</param>
+	/// <param name="document">Document in the parmaeter</param>
 	/// <returns>A new telemetry span</returns>
 	[MethodImpl(MethodImplOptions.AggressiveInlining)]
-	public static TelemetrySpan StartMongoDbSpan<T>(this Tracer tracer, string spanName, IMongoCollection<T>? collection = null)
+	public static TelemetrySpan StartMongoDbSpan<T>(this Tracer tracer, string spanName, IMongoCollection<T>? collection = null, FilterDefinition<T>? filter = null, UpdateDefinition<T>? update = null, T? document = default)
 	{
 		string name = "mongodb." + spanName;
 		TelemetrySpan span = OpenTelemetryTracers.MongoDb
@@ -323,7 +328,19 @@ public static class OpenTelemetrySpanExtensions
 		{
 			span.SetAttribute("collection", collection.CollectionNamespace.CollectionName);
 		}
-		
+		if (filter != null)
+		{
+			span.SetAttribute("filter", filter.Render().ToJson());
+		}
+		if (update != null)
+		{
+			span.SetAttribute("update", update.Render().ToJson());
+		}
+		if (document != null)
+		{
+			span.SetAttribute("document", document.ToJson());
+		}
+
 		return span;
 	}
 	
