Horde: Add a scope for all shared tickers.

diff --git a/Engine/Source/Programs/Horde/Horde.Server/Utilities/Clock.cs b/Engine/Source/Programs/Horde/Horde.Server/Utilities/Clock.cs
--- a/Engine/Source/Programs/Horde/Horde.Server/Utilities/Clock.cs
+++ b/Engine/Source/Programs/Horde/Horde.Server/Utilities/Clock.cs
@@ -10,6 +10,8 @@ using Horde.Server;
 using Horde.Server.Server;
 using Microsoft.Extensions.Logging;
 using Microsoft.Extensions.Options;
+using OpenTelemetry.Trace;
+using OpenTracing;
 using StackExchange.Redis;
 using TimeZoneConverter;
 
@@ -231,6 +233,7 @@ namespace HordeCommon
 		}
 
 		readonly RedisService _redis;
+		readonly Tracer _tracer;
 		readonly TimeZoneInfo _timeZone;
 
 		/// <inheritdoc/>
@@ -242,9 +245,10 @@ namespace HordeCommon
 		/// <summary>
 		/// Constructor
 		/// </summary>
-		public Clock(RedisService redis, IOptions<ServerSettings> settings)
+		public Clock(RedisService redis, Tracer tracer, IOptions<ServerSettings> settings)
 		{
 			_redis = redis;
+			_tracer = tracer;
 
 			string? timeZoneName = settings.Value.ScheduleTimeZone;
 			_timeZone = (timeZoneName == null) ? TimeZoneInfo.Local : TZConvert.GetTimeZoneInfo(timeZoneName);
@@ -260,15 +264,18 @@ namespace HordeCommon
 		public ITicker AddSharedTicker(string name, TimeSpan delay, Func<CancellationToken, ValueTask> tickAsync, ILogger logger)
 		{
 			RedisKey key = new RedisKey($"tick/{name}");
-			return ClockExtensions.AddTicker(this, name, delay / 4, token => TriggerSharedAsync(key, delay, tickAsync, token), logger);
+			return ClockExtensions.AddTicker(this, name, delay / 4, token => TriggerSharedAsync(name, key, delay, tickAsync, token), logger);
 		}
 
-		async ValueTask TriggerSharedAsync(RedisKey key, TimeSpan interval, Func<CancellationToken, ValueTask> tickAsync, CancellationToken cancellationToken)
+		async ValueTask TriggerSharedAsync(string name, RedisKey key, TimeSpan interval, Func<CancellationToken, ValueTask> tickAsync, CancellationToken cancellationToken)
 		{
 			using (RedisLock sharedLock = new (_redis.GetDatabase(), key))
 			{
 				if (await sharedLock.AcquireAsync(interval, false))
 				{
+					using TelemetrySpan span = _tracer.StartActiveSpan(name);
+					span.SetAttribute("interval", interval.TotalSeconds);
+
 					await tickAsync(cancellationToken);
 				}
 			}
