Horde: Periodically expire invalid agent sessions.

diff --git a/Engine/Source/Programs/Horde/Horde.Server/Agents/Sessions/ISessionCollection.cs b/Engine/Source/Programs/Horde/Horde.Server/Agents/Sessions/ISessionCollection.cs
--- a/Engine/Source/Programs/Horde/Horde.Server/Agents/Sessions/ISessionCollection.cs
+++ b/Engine/Source/Programs/Horde/Horde.Server/Agents/Sessions/ISessionCollection.cs
@@ -57,7 +57,7 @@ namespace Horde.Server.Agents.Sessions
 		/// <param name="properties">The agent properties</param>
 		/// <param name="resources">Resources which the agent has</param>
 		/// <returns>Async task</returns>
-		Task UpdateAsync(SessionId sessionId, DateTime finishTime, IReadOnlyList<string> properties, IReadOnlyDictionary<string, int> resources);
+		Task UpdateAsync(SessionId sessionId, DateTime? finishTime, IReadOnlyList<string>? properties, IReadOnlyDictionary<string, int>? resources);
 
 		/// <summary>
 		/// Delete a session from the collection
diff --git a/Engine/Source/Programs/Horde/Horde.Server/Agents/Sessions/SessionCollection.cs b/Engine/Source/Programs/Horde/Horde.Server/Agents/Sessions/SessionCollection.cs
--- a/Engine/Source/Programs/Horde/Horde.Server/Agents/Sessions/SessionCollection.cs
+++ b/Engine/Source/Programs/Horde/Horde.Server/Agents/Sessions/SessionCollection.cs
@@ -117,14 +117,22 @@ namespace Horde.Server.Agents.Sessions
 		}
 
 		/// <inheritdoc/>
-		public Task UpdateAsync(SessionId sessionId, DateTime finishTime, IReadOnlyList<string> properties, IReadOnlyDictionary<string, int> resources)
+		public Task UpdateAsync(SessionId sessionId, DateTime? finishTime, IReadOnlyList<string>? properties, IReadOnlyDictionary<string, int>? resources)
 		{
-			UpdateDefinition<SessionDocument> update = Builders<SessionDocument>.Update
-				.Set(x => x.FinishTime, finishTime)
-				.Set(x => x.Properties, new List<string>(properties))
-				.Set(x => x.Resources, new Dictionary<string, int>(resources));
-
-			return _sessions.FindOneAndUpdateAsync(x => x.Id == sessionId, update);
+			List<UpdateDefinition<SessionDocument>> updates = new List<UpdateDefinition<SessionDocument>>();
+			if (finishTime != null)
+			{
+				updates.Add(Builders<SessionDocument>.Update.Set(x => x.FinishTime, finishTime));
+			}
+			if (properties != null)
+			{
+				updates.Add(Builders<SessionDocument>.Update.Set(x => x.Properties, new List<string>(properties)));
+			}
+			if (resources != null)
+			{
+				updates.Add(Builders<SessionDocument>.Update.Set(x => x.Resources, new Dictionary<string, int>(resources)));
+			}
+			return _sessions.FindOneAndUpdateAsync(x => x.Id == sessionId, Builders<SessionDocument>.Update.Combine(updates));
 		}
 
 		/// <inheritdoc/>
diff --git a/Engine/Source/Programs/Horde/Horde.Server/Server/ConsistencyService.cs b/Engine/Source/Programs/Horde/Horde.Server/Server/ConsistencyService.cs
--- a/Engine/Source/Programs/Horde/Horde.Server/Server/ConsistencyService.cs
+++ b/Engine/Source/Programs/Horde/Horde.Server/Server/ConsistencyService.cs
@@ -11,6 +11,7 @@ using Horde.Server.Utilities;
 using HordeCommon;
 using Microsoft.Extensions.Hosting;
 using Microsoft.Extensions.Logging;
+using Horde.Server.Agents;
 
 namespace Horde.Server.Server
 {
@@ -19,18 +20,22 @@ namespace Horde.Server.Server
 	/// </summary>
 	class ConsistencyService : IHostedService, IDisposable
 	{
+		readonly IAgentCollection _agentCollection;
 		readonly ISessionCollection _sessionCollection;
 		readonly ILeaseCollection _leaseCollection;
+		readonly IClock _clock;
 		readonly ITicker _ticker;
 		readonly ILogger<ConsistencyService> _logger;
 
 		/// <summary>
 		/// Constructor
 		/// </summary>
-		public ConsistencyService(ISessionCollection sessionCollection, ILeaseCollection leaseCollection, IClock clock, ILogger<ConsistencyService> logger)
+		public ConsistencyService(IAgentCollection agentCollection, ISessionCollection sessionCollection, ILeaseCollection leaseCollection, IClock clock, ILogger<ConsistencyService> logger)
 		{
+			_agentCollection = agentCollection;
 			_sessionCollection = sessionCollection;
 			_leaseCollection = leaseCollection;
+			_clock = clock;
 			_ticker = clock.AddSharedTicker<ConsistencyService>(TimeSpan.FromMinutes(20.0), TickLeaderAsync, logger);
 			_logger = logger;
 		}
@@ -51,9 +56,32 @@ namespace Horde.Server.Server
 		/// <returns>Async task</returns>
 		async ValueTask TickLeaderAsync(CancellationToken stoppingToken)
 		{
+			_logger.LogInformation("Ticking consistency service...");
+
+			// Find all the active sessions
 			List<ISession> sessions = await _sessionCollection.FindActiveSessionsAsync();
 			Dictionary<SessionId, ISession> sessionIdToInstance = sessions.ToDictionary(x => x.Id, x => x);
 
+			// Find all the active agents
+			List<IAgent> agents = await _agentCollection.FindAsync(status: AgentStatus.Ok);
+			Dictionary<AgentId, IAgent> agentIdToInstance = agents.ToDictionary(x => x.Id, x => x);
+
+			// Find any sessions that do not have a finish time despite their agents running something else
+			DateTime utcNow = _clock.UtcNow;
+			foreach(ISession session in sessions)
+			{
+				if (!agentIdToInstance.TryGetValue(session.AgentId, out IAgent? agent) || agent.SessionId != session.Id)
+				{
+					agent = await _agentCollection.GetAsync(session.AgentId);
+					if (agent == null || agent.SessionId != session.Id)
+					{
+						_logger.LogWarning("Forcing agent {AgentId} session {SessionId} to complete.", session.AgentId, session.Id);
+						await _sessionCollection.UpdateAsync(session.Id, utcNow, null, null);
+						sessionIdToInstance.Remove(session.Id);
+					}
+				}
+			}
+
 			// Find any leases that are still running when their session has terminated
 			List<ILease> leases = await _leaseCollection.FindActiveLeasesAsync();
 			foreach (ILease lease in leases)
