Horde: Suppress warnings for issues with out of sync resolve state, if the issue was closed due to being stale.

diff --git a/Engine/Source/Programs/Horde/Horde.Server/Issues/IssueCollection.cs b/Engine/Source/Programs/Horde/Horde.Server/Issues/IssueCollection.cs
--- a/Engine/Source/Programs/Horde/Horde.Server/Issues/IssueCollection.cs
+++ b/Engine/Source/Programs/Horde/Horde.Server/Issues/IssueCollection.cs
@@ -831,7 +831,7 @@ namespace Horde.Server.Issues
 						for (int idx = results.Count - 1; idx >= 0; idx--)
 						{
 							Issue issue = results[idx];
-							if ((issue.ResolvedAt != null) != resolved.Value)
+							if ((issue.ResolvedAt != null) != resolved.Value && issue.ResolvedById != IIssue.ResolvedByTimeoutId)
 							{
 								_logger.LogWarning("Issue {IssueId} has resolved state out of sync with spans", issue.Id);
 								results.RemoveAt(idx);
