Horde: Add better indexes over the sessions collection.

- Compound index should now be able to satisfy regular search queries using IXSCAN.
- Query for active sessions is still satisfied through IXSCAN over FinishTime.

diff --git a/Engine/Source/Programs/Horde/Horde.Server/Agents/Sessions/SessionCollection.cs b/Engine/Source/Programs/Horde/Horde.Server/Agents/Sessions/SessionCollection.cs
--- a/Engine/Source/Programs/Horde/Horde.Server/Agents/Sessions/SessionCollection.cs
+++ b/Engine/Source/Programs/Horde/Horde.Server/Agents/Sessions/SessionCollection.cs
@@ -69,8 +69,7 @@ namespace Horde.Server.Agents.Sessions
 		public SessionCollection(MongoService mongoService)
 		{
 			List<MongoIndex<SessionDocument>> indexes = new List<MongoIndex<SessionDocument>>();
-			indexes.Add(keys => keys.Ascending(x => x.AgentId));
-			indexes.Add(keys => keys.Ascending(x => x.StartTime));
+			indexes.Add(keys => keys.Ascending(x => x.AgentId).Ascending(x => x.StartTime).Ascending(x => x.FinishTime));
 			indexes.Add(keys => keys.Ascending(x => x.FinishTime));
 
 			_sessions = mongoService.GetCollection<SessionDocument>("Sessions", indexes);
