Horde: Add bisection to dashboard job detail step
#rnx
#skipci

diff --git a/Engine/Source/Programs/Horde/HordeDashboard/src/backend/Api.ts b/Engine/Source/Programs/Horde/HordeDashboard/src/backend/Api.ts
--- a/Engine/Source/Programs/Horde/HordeDashboard/src/backend/Api.ts
+++ b/Engine/Source/Programs/Horde/HordeDashboard/src/backend/Api.ts
@@ -1075,6 +1075,9 @@ export type GetJobResponse = {
 	/** The user that started this job */
 	abortedByUserInfo?: GetThinUserInfoResponse;
 
+	/** Whether job was created by a bisection task */
+	startedByBisectTaskId?: string;
+
 	/** The roles to impersonate when executing this job */
 	roles: string[];
 
@@ -4630,4 +4633,122 @@ export type GetToolSummaryResponse = {
 
 	/** Version of tool */
 	version?: string;
-}
\ No newline at end of file
+}
+
+/** Job Bisection */
+
+/// State of a bisect task
+export enum BisectTaskState {
+	/// Currently running	
+	Running = "Running",
+
+	/// Cancelled by a user	
+	Cancelled = "Cancelled",
+
+	/// Finished running. The first job/change identifies the first failure.	
+	Succeeded = "Succeeded",
+
+	/// Task failed due to not having a job before the first failure.	
+	MissingHistory = "MissingHistory",
+
+	/// Task failed due to the stream no longer existing.	
+	MissingStream = "MissingStream",
+
+	/// Task failed due to the first job no longer existing.	
+	MissingJob = "MissingJob",
+
+	/// Task failed due to template no longer existing.	
+	MissingTemplate = "MissingTemplate"
+}
+
+
+export type CommitTag = {
+	text: string;
+}
+
+/// Request to create a new bisect task
+export type CreateBisectTaskRequest = {
+
+	/// Job containing the node to check	
+	jobId: string;
+
+	/// Name of the node to query	
+	nodeName: string;
+
+	/// Commit tag to filter possible changes against	
+	commitTags?: CommitTag[];
+
+	/// Set of changes to ignore. Can be modified later through UpdateBisectTaskRequest
+	ignoreChanges?: number[];
+
+	/// Set of job ids to ignore. Can be modified later through UpdateBisectTaskRequest"
+	ignoreJobs?: string[];
+}
+
+/// Response from creating a bisect task
+export type CreateBisectTaskResponse = {
+
+	/// Identifier for the new bisect task	
+	bisectTaskId: string;
+}
+
+/// Information about a bisect task
+export type GetBisectTaskResponse = {
+
+	/// Identifier for this task
+	id: string;
+
+	/// Current task state
+	state: BisectTaskState;
+
+	/// User that initiated the search
+	owner: GetThinUserInfoResponse;
+
+	/// Stream being searched
+	streamId: string;
+
+	/// Template within the stream to execute
+	templateId: string;
+
+	/// Name of the step to search for
+	nodeName: string;
+
+	/// Outcome to search for
+	outcome: JobStepOutcome;
+
+	/// Starting job id for the bisection
+	initialJobId: string;
+
+	/// Starting change for the bisection
+	initialChange: number;
+
+	/// First known job id that is broken
+	currentJobId: string;
+
+	/// Changelist number of the first broken job id
+	currentChange: number;
+
+	/// The steps that have been run on bisection
+	steps?: GetJobStepRefResponse[];
+
+}
+
+/// Updates the state of a bisect task
+export type UpdateBisectTaskRequest = {
+	/// Cancels the current task	
+	cancel?: boolean;
+
+	/// List of change numbers to include in the search. 	
+	includeChanges?: number[];
+
+	/// List of change numbers to exclude from the search.	
+	encludeChanges?: number[];
+
+	/// List of jobs to include in the search.	
+	includeJobs?: string[]; 
+
+	/// List of jobs to exclude from the search.	
+	excludeJobs?: string[];
+}
+
+
diff --git a/Engine/Source/Programs/Horde/HordeDashboard/src/backend/Backend.ts b/Engine/Source/Programs/Horde/HordeDashboard/src/backend/Backend.ts
--- a/Engine/Source/Programs/Horde/HordeDashboard/src/backend/Backend.ts
+++ b/Engine/Source/Programs/Horde/HordeDashboard/src/backend/Backend.ts
@@ -1,7 +1,7 @@
 // Copyright Epic Games, Inc. All Rights Reserved.
 
 import templateCache from '../backend/TemplateCache';
-import { AgentData, AgentQuery, ArtifactData, AuditLogEntry, AuditLogQuery, BatchUpdatePoolRequest, ChangeSummaryData, CreateDeviceRequest, CreateDeviceResponse, CreateExternalIssueRequest, CreateExternalIssueResponse, CreateJobRequest, CreateJobResponse, CreateNoticeRequest, CreatePoolRequest, CreateSoftwareResponse, CreateSubscriptionRequest, CreateSubscriptionResponse, CreateZipRequest, DashboardPreference, DevicePoolTelemetryQuery, DeviceTelemetryQuery, EventData, FindArtifactsResponse, FindIssueResponse, FindJobTimingsResponse, GetAgentSoftwareChannelResponse, GetArtifactDirectoryResponse, GetArtifactZipRequest, GetDashboardConfigResponse, GetDevicePlatformResponse, GetDevicePoolResponse, GetDevicePoolTelemetryResponse, GetDeviceReservationResponse, GetDeviceResponse, GetDeviceTelemetryResponse, GetExternalIssueProjectResponse, GetExternalIssueResponse, GetGraphResponse, GetIssueStreamResponse, GetJobsTabResponse, GetJobStepRefResponse, GetJobStepTraceResponse, GetJobTimingResponse, GetLogEventResponse, GetNoticeResponse, GetNotificationResponse, GetPerforceServerStatusResponse, GetPoolResponse, GetServerInfoResponse, GetServerSettingsResponse, GetSoftwareResponse, GetSubscriptionResponse, GetTestDataDetailsResponse, GetTestDataRefResponse, GetTestMetaResponse, GetTestResponse, GetTestsRequest, GetTestStreamResponse, GetToolSummaryResponse, GetUserResponse, GetUtilizationTelemetryResponse, GlobalConfig, IssueData, IssueQuery, IssueQueryV2, JobData, JobQuery, JobsTabColumnType, JobStepOutcome, JobStreamQuery, JobTimingsQuery, LeaseData, LogData, LogLineData, PoolData, ProjectData, ScheduleData, ScheduleQuery, SearchLogFileResponse, ServerUpdateResponse, SessionData, StreamData, TabType, TestData, UpdateAgentRequest, UpdateDeviceRequest, UpdateGlobalConfigRequest, UpdateIssueRequest, UpdateJobRequest, UpdateLeaseRequest, UpdateNoticeRequest, UpdateNotificationsRequest, UpdatePoolRequest, UpdateServerSettingsRequest, UpdateStepRequest, UpdateStepResponse, UpdateTemplateRefRequest, UpdateUserRequest, UsersQuery } from './Api';
+import { AgentData, AgentQuery, ArtifactData, AuditLogEntry, AuditLogQuery, BatchUpdatePoolRequest, ChangeSummaryData, CreateBisectTaskRequest, CreateBisectTaskResponse, CreateDeviceRequest, CreateDeviceResponse, CreateExternalIssueRequest, CreateExternalIssueResponse, CreateJobRequest, CreateJobResponse, CreateNoticeRequest, CreatePoolRequest, CreateSoftwareResponse, CreateSubscriptionRequest, CreateSubscriptionResponse, CreateZipRequest, DashboardPreference, DevicePoolTelemetryQuery, DeviceTelemetryQuery, EventData, FindArtifactsResponse, FindIssueResponse, FindJobTimingsResponse, GetAgentSoftwareChannelResponse, GetArtifactDirectoryResponse, GetArtifactZipRequest, GetBisectTaskResponse, GetDashboardConfigResponse, GetDevicePlatformResponse, GetDevicePoolResponse, GetDevicePoolTelemetryResponse, GetDeviceReservationResponse, GetDeviceResponse, GetDeviceTelemetryResponse, GetExternalIssueProjectResponse, GetExternalIssueResponse, GetGraphResponse, GetIssueStreamResponse, GetJobsTabResponse, GetJobStepRefResponse, GetJobStepTraceResponse, GetJobTimingResponse, GetLogEventResponse, GetNoticeResponse, GetNotificationResponse, GetPerforceServerStatusResponse, GetPoolResponse, GetServerInfoResponse, GetServerSettingsResponse, GetSoftwareResponse, GetSubscriptionResponse, GetTestDataDetailsResponse, GetTestDataRefResponse, GetTestMetaResponse, GetTestResponse, GetTestsRequest, GetTestStreamResponse, GetToolSummaryResponse, GetUserResponse, GetUtilizationTelemetryResponse, GlobalConfig, IssueData, IssueQuery, IssueQueryV2, JobData, JobQuery, JobsTabColumnType, JobStepOutcome, JobStreamQuery, JobTimingsQuery, LeaseData, LogData, LogLineData, PoolData, ProjectData, ScheduleData, ScheduleQuery, SearchLogFileResponse, ServerUpdateResponse, SessionData, StreamData, TabType, TestData, UpdateAgentRequest, UpdateBisectTaskRequest, UpdateDeviceRequest, UpdateGlobalConfigRequest, UpdateIssueRequest, UpdateJobRequest, UpdateLeaseRequest, UpdateNoticeRequest, UpdateNotificationsRequest, UpdatePoolRequest, UpdateServerSettingsRequest, UpdateStepRequest, UpdateStepResponse, UpdateTemplateRefRequest, UpdateUserRequest, UsersQuery } from './Api';
 import dashboard from './Dashboard';
 import { ChallengeStatus, Fetch } from './Fetch';
 import graphCache, { GraphQuery } from './GraphCache';
@@ -266,7 +266,7 @@ export class Backend {
                 }
 
                 if (!includeGraph) {
-                    resolve(response); 
+                    resolve(response);
                     return;
                 }
 
@@ -1655,6 +1655,50 @@ export class Backend {
         });
     }
 
+    // create job bisection
+    createBisectTask(create: CreateBisectTaskRequest): Promise<CreateBisectTaskResponse> {
+        return new Promise<CreateBisectTaskResponse>((resolve, reject) => {
+            this.backend.post(`/api/v1/bisect`, create).then((response) => {
+                resolve(response.data);
+            }).catch(reason => {
+                reject(reason);
+            });
+        });
+    }
+
+    // get a bisection task
+    getBisectTask(id: string): Promise<GetBisectTaskResponse> {
+        return new Promise<GetBisectTaskResponse>((resolve, reject) => {
+            this.backend.get(`/api/v1/bisect/${id}`).then((value) => {
+                resolve(value.data as GetBisectTaskResponse);
+            }).catch(reason => {
+                reject(reason);
+            });
+        });
+    }
+
+    // update a visection task
+    updateBisectTask(id: string, request: UpdateBisectTaskRequest): Promise<void> {
+        return new Promise<void>((resolve, reject) => {
+            this.backend.patch(`/api/v1/bisect/${id}`, request).then(() => {
+                resolve();
+            }).catch(reason => {
+                reject(reason);
+            });
+        });
+    }
+
+    // get a job's bisection tasks
+    getJobBisectTasks(jobId: string): Promise<GetBisectTaskResponse[]> {
+        return new Promise<GetBisectTaskResponse[]>((resolve, reject) => {
+            this.backend.get(`/api/v1/bisect/job/${jobId}`).then((value) => {
+                resolve(value.data as GetBisectTaskResponse[]);
+            }).catch(reason => {
+                reject(reason);
+            });
+        });
+    }
+
     private async update() {
 
         if (this.updateID === "updating") {
diff --git a/Engine/Source/Programs/Horde/HordeDashboard/src/backend/Fetch.ts b/Engine/Source/Programs/Horde/HordeDashboard/src/backend/Fetch.ts
--- a/Engine/Source/Programs/Horde/HordeDashboard/src/backend/Fetch.ts
+++ b/Engine/Source/Programs/Horde/HordeDashboard/src/backend/Fetch.ts
@@ -178,6 +178,30 @@ export class Fetch {
 
     }
 
+    patch(url: string, data?: any, config?: FetchRequestConfig) {
+
+        url = this.buildUrl(url, config);
+
+        return new Promise<FetchResponse>((resolve, reject) => {
+
+            fetch(url, this.buildRequest("PATCH", data)).then(response => {
+
+                this.handleResponse(response, url, "PATCH", resolve, reject, config);
+
+            }).catch(reason => {
+
+                handleError({
+                    reason: reason,
+                    mode: "PATCH",
+                    url: url
+                }, true);
+
+                reject(reason);
+            });
+        });
+
+    }
+
 
     delete(url: string, config?: FetchRequestConfig) {
 
@@ -327,13 +351,13 @@ export class Fetch {
 
     }
 
-    private buildRequest(method: "GET" | "POST" | "PUT" | "DELETE", data?: any, config?: FetchRequestConfig): RequestInit {
+    private buildRequest(method: "GET" | "POST" | "PUT" | "DELETE" | "PATCH", data?: any, config?: FetchRequestConfig): RequestInit {
 
         const headers = this.buildHeaders();
 
         let body: any = undefined;
 
-        if (method === "POST" || method === "PUT") {
+        if (method === "POST" || method === "PUT" || method === "PATCH") {
             if (data) {
                 if (config?.formData) {
                     body = data;
diff --git a/Engine/Source/Programs/Horde/HordeDashboard/src/components/bisection/CreateModal.tsx b/Engine/Source/Programs/Horde/HordeDashboard/src/components/bisection/CreateModal.tsx
new file mode 100644
--- /dev/null
+++ b/Engine/Source/Programs/Horde/HordeDashboard/src/components/bisection/CreateModal.tsx
@@ -0,0 +1,108 @@
+import { Modal, Stack, Text, IconButton, TextField, PrimaryButton, DefaultButton, MessageBar, MessageBarType } from "@fluentui/react";
+import { hordeClasses } from "../../styles/Styles";
+import { useState } from "react";
+import backend from "../../backend";
+import { CommitTag, CreateBisectTaskResponse } from "../../backend/Api";
+
+export const BisectionCreateModal: React.FC<{ jobId: string, nodeName: string, onClose: (response?: CreateBisectTaskResponse) => void }> = ({ jobId, nodeName, onClose }) => {
+
+   const [state, setState] = useState<{ submitting?: boolean, errorMsg?: string, commitTags?: string, ignoreChanges?: string, ignoreJobs?: string }>({});
+
+   const onBisect = async () => {
+
+      setState({ ...state, submitting: true });
+
+      try {
+
+         let commitTags: CommitTag[] | undefined;
+         let ignoreChanges: number[] | undefined;
+         let ignoreJobs: string[] | undefined;
+
+         if (state.commitTags) {
+            commitTags = state.commitTags.split(',').map(s => { return { text: s.trim() } }).filter(s => !!s.text)
+         }
+
+         if (state.ignoreChanges) {
+            ignoreChanges = state.ignoreChanges.split(',').map(s => { return parseInt(s.trim()) }).filter(s => !isNaN(s));
+         }
+
+         if (state.ignoreJobs) {
+            ignoreJobs = state.ignoreJobs.split(',').map(s => { return s.trim() }).filter(s => !!s.length);
+         }
+
+         const request = {
+            jobId: jobId,
+            nodeName: nodeName,
+            commitTags: commitTags,
+            ignoreChanges: ignoreChanges,
+            ignoreJobs: ignoreJobs
+         };
+
+         const response = await backend.createBisectTask(request);
+         console.log(response);
+
+         setState({ ...state, submitting: false, errorMsg: "" });
+         onClose(response);
+
+      } catch (reason) {
+
+         setState({ ...state, submitting: false, errorMsg: reason as string });
+      }
+
+   };
+
+
+   return <Modal isOpen={true} topOffsetFixed={true} styles={{ main: { padding: 8, width: 640, height: '570px', backgroundColor: '#FFFFFF', hasBeenOpened: false, top: "80px", position: "absolute" } }} className={hordeClasses.modal} onDismiss={() => { onClose() }}>
+      <Stack styles={{ root: { paddingTop: 8, paddingLeft: 24, paddingRight: 12, paddingBottom: 8 } }}>
+         <Stack tokens={{ childrenGap: 12 }}>
+            <Stack horizontal styles={{ root: { padding: 0 } }}>
+               <Stack style={{ paddingLeft: 0, paddingTop: 4 }} grow>
+                  <Text variant="mediumPlus" styles={{ root: { fontFamily: "Horde Open Sans SemiBold" } }}>Step Bisection</Text>
+               </Stack>
+               <Stack grow horizontalAlign="end">
+                  <IconButton
+                     iconProps={{ iconName: 'Cancel' }}
+                     onClick={() => { onClose(); }}
+                  />
+               </Stack>
+            </Stack>
+            <Stack style={{ paddingRight: 12 }} tokens={{ childrenGap: 12 }}>
+               <Stack>
+                  {!!state.errorMsg && <MessageBar
+                     messageBarType={MessageBarType.error}
+                     isMultiline={false}> {state.errorMsg}
+                  </MessageBar>}
+               </Stack>
+               <Stack>
+                  <TextField label="JobId" value={jobId} disabled={true} />
+               </Stack>
+               <Stack>
+                  <TextField label="Node Name" value={nodeName} disabled={true} />
+               </Stack>
+               <Stack>
+                  <TextField label="Commit Tags" multiline spellCheck={false} placeholder="Comma delimited list of commit tags" onChange={(event, newValue) => setState({ ...state, commitTags: newValue })} />
+               </Stack>
+               <Stack>
+                  <TextField label="Ignore Changes" multiline spellCheck={false} placeholder="Comma delimited list of changes to ignore" onChange={(event, newValue) => setState({ ...state, ignoreChanges: newValue })} />
+               </Stack>
+               <Stack>
+                  <TextField label="Ignore Jobs" multiline spellCheck={false} placeholder="Comma delimited list of job ids to ignore" onChange={(event, newValue) => setState({ ...state, ignoreJobs: newValue })} />
+               </Stack>
+            </Stack>
+            <Stack>
+               <Stack horizontal>
+                  <Stack grow />
+                  <Stack horizontal style={{ paddingTop: 24 }}>
+                     <Stack style={{ paddingTop: 8, paddingRight: 24 }}>
+                        <PrimaryButton disabled={state.submitting} text="Bisect" onClick={() => { onBisect() }} />
+                     </Stack>
+                     <Stack style={{ paddingTop: 8, paddingRight: 24 }}>
+                        <DefaultButton disabled={state.submitting} text="Cancel" onClick={() => { onClose() }} />
+                     </Stack>
+                  </Stack>
+               </Stack>
+            </Stack>
+         </Stack>
+      </Stack>
+   </Modal>
+}
\ No newline at end of file
diff --git a/Engine/Source/Programs/Horde/HordeDashboard/src/components/jobDetailsV2/JobDetailBisection.tsx b/Engine/Source/Programs/Horde/HordeDashboard/src/components/jobDetailsV2/JobDetailBisection.tsx
new file mode 100644
--- /dev/null
+++ b/Engine/Source/Programs/Horde/HordeDashboard/src/components/jobDetailsV2/JobDetailBisection.tsx
@@ -0,0 +1,378 @@
+// Copyright Epic Games, Inc. All Rights Reserved.
+
+import { DefaultButton, DetailsList, DetailsListLayoutMode, IColumn, IGroup, IGroupHeaderProps, IconButton, MessageBar, MessageBarType, Modal, PrimaryButton, SelectionMode, Stack, Text } from '@fluentui/react';
+import { mergeStyleSets } from '@fluentui/react/lib/Styling';
+import { observer } from 'mobx-react-lite';
+import moment from 'moment';
+import React, { useEffect, useState } from 'react';
+import { Link } from 'react-router-dom';
+import backend from '../../backend';
+import { BisectTaskState, GetBisectTaskResponse, GetJobStepRefResponse } from '../../backend/Api';
+import dashboard, { StatusColor } from '../../backend/Dashboard';
+import { Markdown } from '../../base/components/Markdown';
+import { ISideRailLink } from '../../base/components/SideRail';
+import { displayTimeZone, getElapsedString } from '../../base/utilities/timeUtils';
+import { hordeClasses } from '../../styles/Styles';
+import { ChangeButton } from '../ChangeButton';
+import { StepRefStatusIcon } from '../StatusIcon';
+import { JobDataView, JobDetailsV2 } from './JobDetailsViewCommon';
+
+
+const sideRail: ISideRailLink = { text: "Bisection", url: "rail_detail_bisection" };
+
+
+class StepBisectionView extends JobDataView {
+
+   filterUpdated() {
+      // this.updateReady();
+   }
+
+   async set(stepId?: string) {
+
+      const details = this.details;
+      if (!details) {
+         return;
+      }
+
+      if ((stepId && this.stepId === stepId) || !details.jobId || !details.jobData?.id) {
+         return;
+      }
+
+      this.stepId = stepId;
+
+      if (!this.bisections) {
+
+         try {
+
+            let jobId = details.jobData?.startedByBisectTaskId ? details.jobData?.startedByBisectTaskId : details.jobId;
+
+            if (jobId !== details.jobId) {
+               const task = await backend.getBisectTask(details.jobData.startedByBisectTaskId!);
+               jobId = task.initialJobId;
+
+            }
+
+            this.bisections = await backend.getJobBisectTasks(jobId);
+            if (stepId) {
+               this.bisections = this.bisections.filter(r => r.nodeName === details.getStepName(stepId));
+            }
+            
+
+         } catch (error) {
+            console.error(error);
+         }
+         finally {
+            this.initialize(this.bisections?.length ? [sideRail] : undefined);
+         }
+      }
+   }
+
+   clear() {
+      this.stepId = undefined;
+      this.bisections = undefined;
+      super.clear();
+   }
+
+   detailsUpdated() {
+
+      if (!this.details?.jobData) {
+         return;
+      }
+
+      this.updateReady();
+
+   }
+
+   stepId?: string;
+
+   bisections?: GetBisectTaskResponse[];
+
+   order = 9;
+
+}
+
+JobDetailsV2.registerDataView("StepBisectionView", (details: JobDetailsV2) => new StepBisectionView(details));
+
+const customStyles = mergeStyleSets({
+   details: {
+      selectors: {
+         '.ms-GroupHeader,.ms-GroupHeader:hover': {
+            background: "#DFDEDD",
+         },
+         '.ms-GroupHeader-title': {
+            cursor: "default"
+         },
+         '.ms-GroupHeader-expand,.ms-GroupHeader-expand:hover': {
+            background: "#DFDEDD"
+         },
+         '.ms-DetailsRow': {
+            animation: "none",
+            background: "unset"
+         },
+         '.ms-DetailsRow:hover': {
+            background: "#F3F2F1"
+         }
+      },
+   }
+
+});
+
+
+export const CancelBisectionModal: React.FC<{ jobDetails: JobDetailsV2, dataView: StepBisectionView, bisectTaskId: string, onClose: (canceled:boolean) => void }> = ({ jobDetails, dataView, bisectTaskId, onClose }) => {
+
+   const [state, setState] = useState<{ submitting?: boolean, error?: string }>({});
+
+   const onCancelTask = async () => {
+
+      setState({ ...state, submitting: true });
+
+      try {
+
+         await backend.updateBisectTask(bisectTaskId, {
+            cancel: true
+         });
+         
+         setState({ ...state, submitting: false, error: "" });
+         onClose(true);
+         
+      } catch (reason) {
+
+         setState({ ...state, submitting: false, error: reason as string });
+      }
+
+   };
+
+   const height = state.error ? 200 : 180;
+
+   return <Modal className={hordeClasses.modal} isOpen={true} topOffsetFixed={true} styles={{ main: { padding: 8, width: 540, height: height, minHeight: height, hasBeenOpened: false, top: "24px", position: "absolute" } }} onDismiss={() => { onClose(false) }}>
+      <Stack horizontal styles={{ root: { padding: 8 } }}>
+         <Stack.Item grow={2}>
+            <Text variant="mediumPlus">Cancel Bisection</Text>
+         </Stack.Item>
+         <Stack.Item grow={0}>
+            <IconButton
+               iconProps={{ iconName: 'Cancel' }}
+               ariaLabel="Close popup modal"
+               onClick={() => { onClose(false); }}
+            />
+         </Stack.Item>
+      </Stack>
+
+      <Stack styles={{ root: { paddingLeft: 8, width: 540, paddingTop: 12 } }}>
+         <Stack tokens={{ childrenGap: 12 }}>
+            {!!state.error && <MessageBar
+               messageBarType={MessageBarType.error}
+               isMultiline={false}> {state.error} </MessageBar>}
+
+         </Stack>
+      </Stack>
+
+      <Stack tokens={{ childrenGap: 16 }} styles={{ root: { paddingTop: 32, paddingLeft: 8, paddingBottom: 8 } }}>
+         <Stack horizontal>
+            <Stack>
+            </Stack>
+            <Stack grow />
+            <Stack>
+               <Stack horizontal tokens={{ childrenGap: 12 }} style={{ paddingRight: 24 }}>
+                  <PrimaryButton text="Cancel Task" disabled={state.submitting ?? false} onClick={() => { onCancelTask() }} />
+                  <DefaultButton text="Go Back" disabled={state.submitting} onClick={() => { onClose(false); }} />
+               </Stack>
+            </Stack>
+         </Stack>
+      </Stack>
+   </Modal>;
+
+};
+
+
+const StepBisectionList: React.FC<{ jobDetails: JobDetailsV2, stepId?: string, dataView: StepBisectionView }> = ({ jobDetails, stepId, dataView }) => {
+
+   const [showCancelModal, setShowCancelModal] = useState("");
+
+   type BisectionItem = {
+      bisection: GetBisectTaskResponse;
+      step?: GetJobStepRefResponse;
+   }
+
+   const bisections = dataView.bisections;
+   if (!bisections?.length) {
+      return null;
+   }
+
+   const items: BisectionItem[] = [];
+
+   const groups: IGroup[] = [];
+
+   bisections.forEach((b, cindex) => {
+
+      const key = `group_${b.id}`;
+      groups.push({
+         key: key,
+         name: `Bisection ${b.id} - ${b.state}`,
+         data: b,
+         startIndex: items.length,
+         count: (b.steps?.length ?? 0),
+         level: 0,
+         isCollapsed: false //bisections.length > 0 && cindex !== 0
+      });
+
+      b.steps?.sort((a, b) => a.change - b.change).forEach(s => {
+         items.push({
+            bisection: b,
+            step: s
+         })
+      })
+   })
+
+   const columns = [
+      { key: 'column1', name: 'Name', minWidth: 620, maxWidth: 620, isResizable: false },
+      { key: 'column2', name: 'Change', minWidth: 80, maxWidth: 80, isResizable: false },
+      { key: 'column3', name: 'Started', minWidth: 180, maxWidth: 180, isResizable: false },
+      { key: 'column4', name: 'Duration', minWidth: 110, maxWidth: 110, isResizable: false },
+   ];
+
+   const renderItem = (item: BisectionItem, index?: number, column?: IColumn) => {
+
+      const step = item?.step;
+
+      if (!column || !step) {
+         return null;
+      }
+
+      if (column.name === "Name") {
+         return <Link to={`/job/${step.jobId}?step=${step.stepId}`}><Stack horizontal>{<StepRefStatusIcon stepRef={step} />}<Text>{item.bisection.nodeName}</Text></Stack></Link>;
+      }
+
+      if (column.name === "Change") {
+
+
+         return <ChangeButton job={jobDetails.jobData!} stepRef={step} hideAborted={true} />;
+
+      }
+
+
+      if (column.name === "Started") {
+
+         if (step.startTime) {
+
+            const displayTime = moment(step.startTime).tz(displayTimeZone());
+            const format = dashboard.display24HourClock ? "HH:mm:ss z" : "LT z";
+
+            let displayTimeStr = displayTime.format('MMM Do') + ` at ${displayTime.format(format)}`;
+
+
+            return <Stack horizontal horizontalAlign={"start"}>{displayTimeStr}</Stack>;
+
+         } else {
+            return "???";
+         }
+      }
+
+      if (column.name === "Duration") {
+
+         const start = moment(step.startTime);
+         let end = moment(Date.now());
+
+         if (step.finishTime) {
+            end = moment(step.finishTime);
+         }
+         if (step.startTime) {
+            const time = getElapsedString(start, end);
+            return <Stack horizontal horizontalAlign={"end"} style={{ paddingRight: 8 }}><Text>{time}</Text></Stack>;
+         } else {
+            return "???";
+         }
+      }
+
+
+      return null;
+
+   }
+
+   const onRenderGroupHeader = (props?: IGroupHeaderProps): JSX.Element | null => {
+
+      if (!props) {
+         return null;
+      }
+
+      const bisection = props.group?.data as GetBisectTaskResponse;
+
+      let markdown = `**Bisection ${bisection.id}** -  **State:** ${bisection.state} / **Initial Job:** [${bisection.initialJobId}](/job/${bisection.initialJobId}) - ${bisection.initialChange}`;
+      if (bisection.currentJobId) {
+         markdown += ` / **Current Job:** [${bisection.currentJobId}](/job/${bisection.currentJobId}) - ${bisection.currentChange}`
+      }
+
+      markdown += ` / **Outcome:** ${bisection.outcome}`
+
+      /*
+      if (bisection.owner) {
+         markdown += ` / **Owner:** ${bisection.owner.name}`
+      }
+      */
+
+      return <Stack verticalAlign="center" style={{ height: 48, padding: 18, backgroundColor: "rgb(233, 232, 231)" }}>
+         <Stack horizontal verticalAlign="center">
+            <Stack>
+               <Markdown>{markdown}</Markdown>
+            </Stack>
+            <Stack grow />
+            {bisection.state === BisectTaskState.Running && <Stack>
+               <DefaultButton style={{ color: "#FFFFFF", backgroundColor: dashboard.getStatusColors().get(StatusColor.Failure) }} text="Cancel" onClick={() => setShowCancelModal(bisection.id)} />
+            </Stack>
+            }
+         </Stack>
+      </Stack>
+   }
+
+   return <Stack>
+      {!!showCancelModal && <CancelBisectionModal jobDetails={jobDetails} dataView={dataView} bisectTaskId={showCancelModal} onClose={(canceled) => { if (canceled) { window.location.reload(); setShowCancelModal("") } }} />}
+      <DetailsList
+         className={customStyles.details}
+         compact={true}
+         isHeaderVisible={false}
+         indentWidth={0}
+         items={items}
+         groups={groups}
+         groupProps={{ onRenderHeader: onRenderGroupHeader, showEmptyGroups: true }}
+         columns={columns}
+         setKey="set"
+         selectionMode={SelectionMode.none}
+         layoutMode={DetailsListLayoutMode.justified}
+         onRenderItemColumn={renderItem}
+      />
+   </Stack>;
+}
+
+export const BisectionPanel: React.FC<{ jobDetails: JobDetailsV2, stepId?: string }> = observer(({ jobDetails, stepId }) => {
+
+   const dataView = jobDetails.getDataView<StepBisectionView>("StepBisectionView");
+
+   useEffect(() => {
+      return () => {
+         dataView?.clear();
+      };
+   }, [dataView]);
+
+   dataView.subscribe();
+
+   dataView.set(stepId);
+
+   if (!dataView.bisections?.length) {
+      return null;
+   }
+
+   return (<Stack id={sideRail.url} styles={{ root: { paddingTop: 18, paddingRight: 12 } }}>
+      <Stack className={hordeClasses.raised}>
+         <Stack tokens={{ childrenGap: 12 }}>
+            <Text variant="mediumPlus" styles={{ root: { fontFamily: "Horde Open Sans SemiBold" } }}>Bisection</Text>
+            <Stack styles={{ root: { paddingLeft: 4, paddingRight: 0, paddingTop: 8, paddingBottom: 4 } }}>
+               <div style={{ overflowY: 'auto', overflowX: 'hidden', minHeight: "400px", maxHeight: "400px" }} data-is-scrollable={true}>
+                  <StepBisectionList jobDetails={jobDetails} dataView={dataView} />
+               </div>
+            </Stack>
+         </Stack>
+      </Stack>
+   </Stack>);
+});
+
+
diff --git a/Engine/Source/Programs/Horde/HordeDashboard/src/components/jobDetailsV2/JobDetailViewV2.tsx b/Engine/Source/Programs/Horde/HordeDashboard/src/components/jobDetailsV2/JobDetailViewV2.tsx
--- a/Engine/Source/Programs/Horde/HordeDashboard/src/components/jobDetailsV2/JobDetailViewV2.tsx
+++ b/Engine/Source/Programs/Horde/HordeDashboard/src/components/jobDetailsV2/JobDetailViewV2.tsx
@@ -15,6 +15,7 @@ import { StepsPanelV2 } from "./JobDetailViewSteps";
 import { SummaryPanel } from "./JobDetailViewSummary";
 import { JobOperations } from "./JobOperationsBar";
 import { StepDetailView } from "./StepDetailView";
+import { BisectionPanel } from "./JobDetailBisection";
 
 class BreadcrumbDataView extends JobDataView {
 
@@ -302,7 +303,8 @@ const DetailsViewOverview: React.FC<{ jobDetails: JobDetailsV2 }> = ({ jobDetail
       <SummaryPanel jobDetails={details} />
       <StepsPanelV2 jobDetails={details} />
       <HealthPanel jobDetails={details} />
-      <TimelinePanel jobDetails={details} />
+      <TimelinePanel jobDetails={details} />      
+      <BisectionPanel jobDetails={details} />
    </Stack>
 };
 
diff --git a/Engine/Source/Programs/Horde/HordeDashboard/src/components/jobDetailsV2/JobOperationsBar.tsx b/Engine/Source/Programs/Horde/HordeDashboard/src/components/jobDetailsV2/JobOperationsBar.tsx
--- a/Engine/Source/Programs/Horde/HordeDashboard/src/components/jobDetailsV2/JobOperationsBar.tsx
+++ b/Engine/Source/Programs/Horde/HordeDashboard/src/components/jobDetailsV2/JobOperationsBar.tsx
@@ -17,6 +17,7 @@ import { JobDetailsV2 } from './JobDetailsViewCommon';
 import { RetryStepsModal, StepRetryModal, StepRetryType } from './StepRetryModal';
 import { JobArtifactsModal } from '../artifacts/ArtifactsModal';
 import { getSiteConfig } from '../../backend/Config';
+import { BisectionCreateModal } from '../bisection/CreateModal';
 
 enum ParameterState {
    Hidden,
@@ -59,12 +60,12 @@ export const JobOperations: React.FC<{ jobDetails: JobDetailsV2 }> = observer(({
 
       const retries = jobDetails.getStepRetries(s.id);
       const retryNumber = jobDetails.getStepRetryNumber(s.id);
-      if (retries.length && retryNumber < (retries.length - 1) ) {
+      if (retries.length && retryNumber < (retries.length - 1)) {
          return false;
       }
 
       return true;
-      
+
    });
    const retryFailedStepsDisabled = !failedSteps.length;
 
@@ -95,13 +96,13 @@ export const JobOperations: React.FC<{ jobDetails: JobDetailsV2 }> = observer(({
 
    //if (getSiteConfig().environment !== "production") {
 
-      opsList.push({
-         key: 'jobops_runfailedsteps',
-         text: "Retry Steps",
-         disabled: retryFailedStepsDisabled,
-         iconProps: { iconName: "Repeat" },
-         onClick: () => { setRetryStepsShown(true); }
-      });
+   opsList.push({
+      key: 'jobops_runfailedsteps',
+      text: "Retry Steps",
+      disabled: retryFailedStepsDisabled,
+      iconProps: { iconName: "Repeat" },
+      onClick: () => { setRetryStepsShown(true); }
+   });
    //}
 
    opsList.push({
@@ -291,8 +292,9 @@ const StepArtifactsOperations: React.FC<{ jobDetails: JobDetailsV2, stepId: stri
 
 const StepOperations: React.FC<{ jobDetails: JobDetailsV2, stepId: string }> = observer(({ jobDetails, stepId }) => {
 
-   const [shown, setShown] = useState<{ abortShown?: boolean, retryShown?: boolean, pauseShown?: boolean }>({});
+   const [shown, setShown] = useState<{ abortShown?: boolean, retryShown?: boolean, pauseShown?: boolean, bisectShown?: boolean }>({});
    const [runType, setRunType] = useState(StepRetryType.RunAgain);
+   const config = getSiteConfig();
 
    // subscribe
    if (dashboard.updated) { }
@@ -314,6 +316,7 @@ const StepOperations: React.FC<{ jobDetails: JobDetailsV2, stepId: string }> = o
 
    const canRunDisabled = !node?.allowRetry || !!step.retriedByUserInfo;
    const canTryFix = jobDetails.template?.allowPreflights;
+   const canBisect = (step.outcome === JobStepOutcome.Failure || step.outcome === JobStepOutcome.Warnings);
 
    const opsList: IContextualMenuItem[] = [];
 
@@ -342,6 +345,16 @@ const StepOperations: React.FC<{ jobDetails: JobDetailsV2, stepId: string }> = o
       onClick: () => { setShown({ retryShown: true }); setRunType(StepRetryType.TestFix); }
    });
 
+   if (config.environment !== "production") {
+      opsList.push({
+         key: 'stepops_bisect',
+         text: "Bisect",
+         iconProps: { iconName: "FlowReview" },
+         disabled: !canBisect,
+         onClick: () => { setShown({ bisectShown: true }); }
+      });
+   }
+
 
    opsList.push({
       key: 'stepops_buildlocally',
@@ -405,9 +418,15 @@ const StepOperations: React.FC<{ jobDetails: JobDetailsV2, stepId: string }> = o
    }
 
    return <Stack>
-      {shown.pauseShown && <PauseStepModal streamId={jobDetails.stream!.id} stepName={node!.name} templateName={jobDetails.template!.name} onClose={() => setShown({ pauseShown: false })} />}
-      <StepRetryModal stepId={stepId} jobDetails={jobDetails} type={runType} show={shown.retryShown ?? false} onClose={() => { setShown({}); }} />
-      <AbortJobModal stepId={stepId} jobDetails={jobDetails} show={shown.abortShown ?? false} onClose={() => { setShown({}); }} />
+      {!!shown.pauseShown && <PauseStepModal streamId={jobDetails.stream!.id} stepName={node!.name} templateName={jobDetails.template!.name} onClose={() => setShown({ pauseShown: false })} />}
+      {!!shown.retryShown && <StepRetryModal stepId={stepId} jobDetails={jobDetails} type={runType} show={true} onClose={() => { setShown({}); }} />}
+      {!!shown.abortShown && <AbortJobModal stepId={stepId} jobDetails={jobDetails} show={true} onClose={() => { setShown({}); }} />}
+      {!!shown.bisectShown && <BisectionCreateModal jobId={jobId} nodeName={node?.name ?? "Unknown Node"} onClose={(response) => {
+         if (response) {
+            window.location.reload();
+         }
+         setShown({});
+      }} />}
       <Stack horizontal styles={{ root: { paddingLeft: 0 } }}>
          <Stack grow />
          <Stack horizontal>
diff --git a/Engine/Source/Programs/Horde/HordeDashboard/src/components/jobDetailsV2/StepDetailView.tsx b/Engine/Source/Programs/Horde/HordeDashboard/src/components/jobDetailsV2/StepDetailView.tsx
--- a/Engine/Source/Programs/Horde/HordeDashboard/src/components/jobDetailsV2/StepDetailView.tsx
+++ b/Engine/Source/Programs/Horde/HordeDashboard/src/components/jobDetailsV2/StepDetailView.tsx
@@ -19,6 +19,7 @@ import { StepErrorPanel } from "./StepErrorPanel";
 import backend from "../../backend";
 import { GetArtifactResponseV2 } from "../../backend/Api";
 import { StepTrendsPanelV2 } from "./JobDetailStepTrendsV2";
+import { BisectionPanel } from "./JobDetailBisection";
 
 
 class StepDetailDataView extends JobDataView {
@@ -171,6 +172,9 @@ const StepDetailViewInner: React.FC<{ jobDetails: JobDetailsV2, stepId: string }
       {!!step && <Stack>
          <StepTrendsPanelV2 jobDetails={jobDetails} stepId={stepId}/>
       </Stack>}
+      {!!step && <Stack>
+         <BisectionPanel jobDetails={jobDetails} stepId={stepId} />
+      </Stack>}
    </Stack>
 });
 
diff --git a/Engine/Source/Programs/Horde/HordeDashboard/src/styles/Styles.tsx b/Engine/Source/Programs/Horde/HordeDashboard/src/styles/Styles.tsx
--- a/Engine/Source/Programs/Horde/HordeDashboard/src/styles/Styles.tsx
+++ b/Engine/Source/Programs/Horde/HordeDashboard/src/styles/Styles.tsx
@@ -132,7 +132,8 @@ registerIcons({
       "ArrowLeft": IconContents.ARROW_LEFT,
       "ArrowRight": IconContents.ARROW_RIGHT,
       "ArrowUp": IconContents.ARROW_UP,
-      "Repeat": IconContents.REPEAT
+      "Repeat": IconContents.REPEAT,
+      "FlowReview": IconContents.FLOW_REVIEW,
    }
 });
 
