Horde: Fix issue with order of checks on CL description for change button
#skipci
#rnx

diff --git a/Engine/Source/Programs/Horde/HordeDashboard/src/components/ChangeButton.tsx b/Engine/Source/Programs/Horde/HordeDashboard/src/components/ChangeButton.tsx
--- a/Engine/Source/Programs/Horde/HordeDashboard/src/components/ChangeButton.tsx
+++ b/Engine/Source/Programs/Horde/HordeDashboard/src/components/ChangeButton.tsx
@@ -48,23 +48,14 @@ export const ChangeContextMenu: React.FC<{ target: ChangeContextMenuTarget, onDi
 
    const menuItems: IContextualMenuItem[] = [];
 
-   if (commit) {
+   if (job?.preflightDescription) {
 
       menuItems.push({
          key: "commit", onRender: () => {
             return <Stack style={{ padding: 18, paddingTop: 18, paddingBottom: 18, maxWidth: 800 }} tokens={{ childrenGap: 12 }}>
-               <Stack horizontal tokens={{ childrenGap: 36 }}>
-                  <Text variant="small" style={{ fontFamily: "Horde Open Sans Bold" }}>Author:</Text>
-                  <Text variant="small" >{commit.authorInfo?.name}</Text>
-               </Stack>
-               <Stack horizontal tokens={{ childrenGap: 35 }}>
-                  <Text variant="small" style={{ fontFamily: "Horde Open Sans Bold" }}>Change:</Text>
-                  <Text variant="small" >{commit.number}</Text>
-               </Stack>
-
                <Stack horizontal tokens={{ childrenGap: 12 }}>
                   <Text variant="small" style={{ fontFamily: "Horde Open Sans Bold" }}>Description:</Text>
-                  <Text variant="small" style={{whiteSpace: "pre-wrap"}}>{commit.description}</Text>
+                  <Text variant="small" style={{ whiteSpace: "pre-wrap" }} >{job.preflightDescription}</Text>
                </Stack>
             </Stack>
          }
@@ -75,14 +66,23 @@ export const ChangeContextMenu: React.FC<{ target: ChangeContextMenuTarget, onDi
          itemType: ContextualMenuItemType.Divider
       });
 
+   } else if (commit) {
 
-   } else if (job.preflightDescription) {
       menuItems.push({
          key: "commit", onRender: () => {
             return <Stack style={{ padding: 18, paddingTop: 18, paddingBottom: 18, maxWidth: 800 }} tokens={{ childrenGap: 12 }}>
+               <Stack horizontal tokens={{ childrenGap: 36 }}>
+                  <Text variant="small" style={{ fontFamily: "Horde Open Sans Bold" }}>Author:</Text>
+                  <Text variant="small" >{commit.authorInfo?.name}</Text>
+               </Stack>
+               <Stack horizontal tokens={{ childrenGap: 35 }}>
+                  <Text variant="small" style={{ fontFamily: "Horde Open Sans Bold" }}>Change:</Text>
+                  <Text variant="small" >{commit.number}</Text>
+               </Stack>
+
                <Stack horizontal tokens={{ childrenGap: 12 }}>
                   <Text variant="small" style={{ fontFamily: "Horde Open Sans Bold" }}>Description:</Text>
-                  <Text variant="small" style={{whiteSpace: "pre-wrap"}} >{job.preflightDescription}</Text>
+                  <Text variant="small" style={{whiteSpace: "pre-wrap"}}>{commit.description}</Text>
                </Stack>
             </Stack>
          }
