Horde: Add some more debug logging to bisect service
#rnx

diff --git a/Engine/Source/Programs/Horde/Horde.Server/Jobs/Bisect/BisectService.cs b/Engine/Source/Programs/Horde/Horde.Server/Jobs/Bisect/BisectService.cs
--- a/Engine/Source/Programs/Horde/Horde.Server/Jobs/Bisect/BisectService.cs
+++ b/Engine/Source/Programs/Horde/Horde.Server/Jobs/Bisect/BisectService.cs
@@ -94,6 +94,8 @@ namespace Horde.Server.Jobs.Bisect
 			// Find the step that ran before this one, ignoring any excluded jobs
 			(IJobStepRef? currentJobStepRef, IJobStepRef? previousJobStepRef) = await GetBisectRangeAsync(bisectTask, cancellationToken);
 
+			_logger.LogDebug("Bisect task {BisectTaskId} ({StreamId}:{TemplateId}:{NodeName}): Current Job Id: {CurJobId}, Previous Job Id: {PrevJobId}", bisectTask.Id, bisectTask.StreamId, bisectTask.TemplateId, bisectTask.NodeName, currentJobStepRef?.Id.JobId, previousJobStepRef?.Id.JobId);
+
 			// Check if the task needs to be updated with the current bisect state
 			if (currentJobStepRef != null && bisectTask.CurrentJobId != currentJobStepRef.Id.JobId)
 			{
