Horde: Unifying C#/C++ implementations of remote execution API.

- Abstract out platform specific implementation details into ComputePlatform.cpp/h.
- Add an overridable transport implementation of ComputeSocket, similar to what is available in C#.
- Add some local tests for C++ implementation.
- Various refactoring to converge implementation details/names/etc...

diff --git a/Engine/Source/Programs/Horde/Horde.Agent/Leases/Handlers/ComputeHandler.cs b/Engine/Source/Programs/Horde/Horde.Agent/Leases/Handlers/ComputeHandler.cs
--- a/Engine/Source/Programs/Horde/Horde.Agent/Leases/Handlers/ComputeHandler.cs
+++ b/Engine/Source/Programs/Horde/Horde.Agent/Leases/Handlers/ComputeHandler.cs
@@ -38,9 +38,9 @@ namespace Horde.Agent.Leases.Handlers
 
 			public override ValueTask MarkCompleteAsync(CancellationToken cancellationToken) => _inner.MarkCompleteAsync(cancellationToken);
 
-			public override async ValueTask<int> ReadPartialAsync(Memory<byte> buffer, CancellationToken cancellationToken)
+			public override async ValueTask<int> RecvAsync(Memory<byte> buffer, CancellationToken cancellationToken)
 			{
-				int result = await _inner.ReadPartialAsync(buffer, cancellationToken);
+				int result = await _inner.RecvAsync(buffer, cancellationToken);
 				if (result > 0)
 				{
 					Interlocked.Exchange(ref _lastPingTicks, Stopwatch.GetTimestamp());
@@ -48,9 +48,9 @@ namespace Horde.Agent.Leases.Handlers
 				return result;
 			}
 
-			public override async ValueTask WriteAsync(ReadOnlySequence<byte> buffer, CancellationToken cancellationToken)
+			public override async ValueTask SendAsync(ReadOnlySequence<byte> buffer, CancellationToken cancellationToken)
 			{
-				await _inner.WriteAsync(buffer, cancellationToken);
+				await _inner.SendAsync(buffer, cancellationToken);
 				Interlocked.Exchange(ref _lastPingTicks, Stopwatch.GetTimestamp());
 			}
 		}
diff --git a/Engine/Source/Programs/Horde/Samples/RemoteWorker/RemoteWorkerApp.cs b/Engine/Source/Programs/Horde/Samples/RemoteWorker/RemoteWorkerApp.cs
--- a/Engine/Source/Programs/Horde/Samples/RemoteWorker/RemoteWorkerApp.cs
+++ b/Engine/Source/Programs/Horde/Samples/RemoteWorker/RemoteWorkerApp.cs
@@ -19,7 +19,7 @@ namespace RemoteWorker
 			byte[] data = new byte[4];
 			await channel.SendAsync(data.AsMemory(0, 1)); // Let the remote know we're ready
 
-			while(await channel.TryReceiveMessageAsync(data))
+			while(await channel.TryRecvMessageAsync(data))
 			{
 				int value = BinaryPrimitives.ReadInt32LittleEndian(data);
 				Console.WriteLine("Read value {0}", value);
diff --git a/Engine/Source/Programs/Horde/Samples/RemoteWorkerCpp/ComputeBuffer.cpp b/Engine/Source/Programs/Horde/Samples/RemoteWorkerCpp/ComputeBuffer.cpp
--- a/Engine/Source/Programs/Horde/Samples/RemoteWorkerCpp/ComputeBuffer.cpp
+++ b/Engine/Source/Programs/Horde/Samples/RemoteWorkerCpp/ComputeBuffer.cpp
@@ -1,15 +1,14 @@
 // Copyright Epic Games, Inc. All Rights Reserved.
 
-#include <windows.h>
+#include "ComputeBuffer.h"
+#include "ComputePlatform.h"
 #include <assert.h>
 #include <iostream>
-#include <rpcdcep.h>
 #include <assert.h>
-#include "ComputeBuffer.h"
 
-namespace
+struct FComputeBufferDetail
 {
-	enum class EWriteState : int
+	enum class EWriteState : unsigned long long
 	{
 		// Writer has moved to the next chunk
 		MovedToNext = 0b00,
@@ -55,41 +54,50 @@ namespace
 		// Read the state value from memory
 		FChunkState Get() const
 		{
-			return FChunkState(InterlockedCompareExchange64((volatile long long*)&Value, 0, 0));
+			return FChunkState(FComputePlatform::AtomicRead64((volatile long long*)&Value));
 		}
 
 		// Append data to the chunk
 		void Append(long long length)
 		{
-			InterlockedAdd64(&Value, length);
+			FComputePlatform::AtomicAdd64(&Value, length);
 		}
 
 		// Mark the chunk as being written to
 		void StartWriting(int numReaders)
 		{
-			InterlockedExchange64(&Value, FChunkState(EWriteState::Writing, (1 << numReaders) - 1, 0).Value);
+			FComputePlatform::AtomicWrite64(&Value, FChunkState(EWriteState::Writing, (1 << numReaders) - 1, 0).Value);
 		}
 
 		// Mark this chunk as the end of the stream
 		void MarkComplete()
 		{
-			InterlockedOr64(&Value, FChunkState(EWriteState::Complete, 0, 0).Value);
+			FComputePlatform::AtomicOr64(&Value, FChunkState(EWriteState::Complete, 0, 0).Value);
 		}
 
 		// Clear the reader flag
 		void FinishReading(int ReaderIdx)
 		{
-			InterlockedAnd64(&Value, ~FChunkState((EWriteState)0, 1 << ReaderIdx, 0).Value);
+			FComputePlatform::AtomicAnd64(&Value, ~FChunkState((EWriteState)0, 1 << ReaderIdx, 0).Value);
 		}
 
 		// Move to the next chunk
 		void FinishWriting()
 		{
-			InterlockedAnd64(&Value, ~FChunkState(EWriteState::Writing, 0, 0).Value);
+			FComputePlatform::AtomicAnd64(&Value, ~FChunkState(EWriteState::Writing, 0, 0).Value);
 		}
 
 	private:
-		volatile long long Value;
+		union
+		{
+			volatile long long Value;
+			struct
+			{
+				unsigned long long Length : 31;
+				unsigned long long ReaderFlags : 31;
+				EWriteState WriteState : 2;
+			};
+		};
 	};
 
 	struct FReaderState
@@ -108,251 +116,405 @@ namespace
 	public:
 		FReaderState Get() const
 		{
-			return FReaderState(InterlockedCompareExchange64((volatile long long*)&Value, 0, 0));
+			return FReaderState(FComputePlatform::AtomicRead64((volatile long long*)&Value));
 		}
 
 		void Set(FReaderState State)
 		{
-			InterlockedExchange64(&Value, State.Value);
+			FComputePlatform::AtomicWrite64(&Value, State.Value);
 		}
 
 		void Advance(size_t Length)
 		{
-			InterlockedAdd64(&Value, Length);
+			FComputePlatform::AtomicAdd64(&Value, Length);
 		}
 
 	private:
-		volatile long long Value;
+		union
+		{
+			volatile long long Value;
+			struct
+			{
+				int Offset;
+				int ChunkIdx;
+			};
+		};
 	};
-}
-
-struct FComputeBufferHeader
-{
-	int NumReaders;
-	int NumChunks;
-	int ChunkLength;
-	int WriteChunkIdx;
-	FChunkStatePtr Chunks[FComputeBuffer::MaxChunks];
-	FReaderStatePtr Readers[FComputeBuffer::MaxReaders];
-};
 
-struct FComputeBufferResources
-{
-	static unsigned int Counter;
+	struct FHeader
+	{
+		int NumReaders;
+		int NumChunks;
+		int ChunkLength;
+		int WriteChunkIdx;
+		FChunkStatePtr Chunks[FComputeBuffer::MaxChunks];
+		FReaderStatePtr Readers[FComputeBuffer::MaxReaders];
+	};
 
 	wchar_t Name[260];
 
-	FComputeBufferHeader* Header;
+	FHeader* Header;
 	unsigned char* ChunkPtrs[FComputeBuffer::MaxChunks];
 
-	void* WriterEvent;
-	void* ReaderEvents[FComputeBuffer::MaxReaders];
+	FComputeManualResetEvent WriterEvent;
+	FComputeManualResetEvent ReaderEvents[FComputeBuffer::MaxReaders];
 
 	FComputeBufferReader Reader;
 	FComputeBufferWriter Writer;
 
-	FComputeBufferResources(const wchar_t* Name)
-		: MemoryMappedFile(nullptr)
-		, Header(nullptr)
+	FComputeBufferDetail(const wchar_t* Name)
+		: Header(nullptr)
 		, ChunkPtrs{ nullptr, }
 		, RefCount(1)
-		, WriterEvent(nullptr)
-		, ReaderEvents{ nullptr, }
 	{
+		static_assert(sizeof(FChunkStatePtr) == sizeof(long long));
+		static_assert(sizeof(FReaderStatePtr) == sizeof(long long));
+
 		wcscpy_s(this->Name, Name);
 	}
 
-	~FComputeBufferResources()
+	~FComputeBufferDetail()
 	{
-		if (Header != nullptr)
-		{
-			UnmapViewOfFile(Header);
-			Header = nullptr;
-		}
-
-		if (MemoryMappedFile != nullptr)
-		{
-			CloseHandle(MemoryMappedFile);
-			MemoryMappedFile = nullptr;
-		}
-
-		if (WriterEvent != nullptr)
-		{
-			CloseHandle(WriterEvent);
-			WriterEvent = nullptr;
-		}
-
-		for (int ReaderIdx = 0; ReaderIdx < FComputeBuffer::MaxReaders; ReaderIdx++)
-		{
-			if (ReaderEvents[ReaderIdx] != nullptr)
-			{
-				CloseHandle(ReaderEvents[ReaderIdx]);
-				ReaderEvents[ReaderIdx] = nullptr;
-			}
-		}
+		Header = nullptr;
 	}
 
-	static std::shared_ptr<FComputeBufferResources> CreateNew(const FComputeBuffer::FParams& Params)
+	static std::shared_ptr<FComputeBufferDetail> CreateNew(const FComputeBuffer::FParams& Params)
 	{
-		unsigned long long Capacity = sizeof(FComputeBufferHeader) + (Params.NumChunks * sizeof(unsigned int)) + (Params.NumChunks * Params.ChunkLength);
+		long long Capacity = sizeof(FHeader) + (Params.NumChunks * sizeof(unsigned int)) + (Params.NumChunks * Params.ChunkLength);
 
 		const wchar_t* Name = Params.Name;
 
-		wchar_t BaseNameBuffer[MAX_PATH];
+		wchar_t BaseNameBuffer[FComputeBuffer::MaxNameLength];
 		if (Name == nullptr)
 		{
-			DWORD Pid = GetCurrentProcessId();
-			ULONGLONG TickCount = GetTickCount64();
-			swprintf_s(BaseNameBuffer, L"Local\\COMPUTE_%u_%llu_%u", Pid, TickCount, InterlockedIncrement(&Counter));
+			FComputePlatform::CreateUniqueName(BaseNameBuffer, FComputeBuffer::MaxNameLength);
 			Name = BaseNameBuffer;
 		}
 
-		wchar_t NameBuffer[MAX_PATH];
-		swprintf_s(NameBuffer, L"%s_M", Name);
+		wchar_t NameBuffer[FComputeBuffer::MaxNameLength];
+		swprintf(NameBuffer, FComputeBuffer::MaxNameLength, L"%s_M", Name);
 
-		std::shared_ptr<FComputeBufferResources> Resources = std::make_shared<FComputeBufferResources>(Name);
-
-		LARGE_INTEGER LargeInteger;
-		LargeInteger.QuadPart = Capacity;
-
-		Resources->MemoryMappedFile = CreateFileMappingW(INVALID_HANDLE_VALUE, NULL, PAGE_READWRITE, LargeInteger.HighPart, LargeInteger.LowPart, NameBuffer);
-		if (Resources->MemoryMappedFile == nullptr)
+		std::shared_ptr<FComputeBufferDetail> Detail = std::make_shared<FComputeBufferDetail>(Name);
+		if (!Detail->MemoryMappedFile.Create(NameBuffer, Capacity))
 		{
 			return nullptr;
 		}
 
-		Resources->Header = (FComputeBufferHeader*)MapViewOfFile(Resources->MemoryMappedFile, FILE_MAP_ALL_ACCESS, 0, 0, 0);
-		if (Resources->Header == nullptr)
+		Detail->Header = (FHeader*)Detail->MemoryMappedFile.GetPointer();
+		if (Detail->Header == nullptr)
 		{
 			return nullptr;
 		}
 
-		memset(Resources->Header, 0, sizeof(*Header));
-		Resources->Header->NumReaders = Params.NumReaders;
-		Resources->Header->NumChunks = Params.NumChunks;
-		Resources->Header->ChunkLength = Params.ChunkLength;
+		memset(Detail->Header, 0, sizeof(*Header));
+		Detail->Header->NumReaders = Params.NumReaders;
+		Detail->Header->NumChunks = Params.NumChunks;
+		Detail->Header->ChunkLength = Params.ChunkLength;
 
-		swprintf_s(NameBuffer, L"%s_W", Name);
-		Resources->WriterEvent = CreateEventW(NULL, TRUE, FALSE, NameBuffer);
-		if (Resources->WriterEvent == nullptr)
+		swprintf(NameBuffer, FComputeBuffer::MaxNameLength, L"%s_W", Name);
+		if (!Detail->WriterEvent.Create(NameBuffer))
 		{
 			return nullptr;
 		}
 
-		for (int ReaderIdx = 0; ReaderIdx < Resources->Header->NumReaders; ReaderIdx++)
+		for (int ReaderIdx = 0; ReaderIdx < Detail->Header->NumReaders; ReaderIdx++)
 		{
-			swprintf_s(NameBuffer, L"%s_R%d", Name, ReaderIdx);
-			Resources->ReaderEvents[ReaderIdx] = CreateEventW(NULL, TRUE, FALSE, NameBuffer);
-			if (Resources->ReaderEvents[ReaderIdx] == nullptr)
+			swprintf(NameBuffer, FComputeBuffer::MaxNameLength, L"%s_R%d", Name, ReaderIdx);
+			if (!Detail->ReaderEvents[ReaderIdx].Create(NameBuffer))
 			{
 				return nullptr;
 			}
 		}
 
-		Resources->Header->Chunks[0].StartWriting(Resources->Header->NumReaders);
+		Detail->Header->Chunks[0].StartWriting(Detail->Header->NumReaders);
 
-		InitShared(Resources);
-		return Resources;
+		InitShared(Detail);
+		return Detail;
 	}
 
-	static std::shared_ptr<FComputeBufferResources> OpenExisting(const wchar_t* Name)
+	static std::shared_ptr<FComputeBufferDetail> OpenExisting(const wchar_t* Name)
 	{
-		wchar_t NameBuffer[MAX_PATH];
-		swprintf_s(NameBuffer, L"%s_M", Name);
-
-		std::shared_ptr<FComputeBufferResources> Resources = std::make_shared<FComputeBufferResources>(Name);
+		wchar_t NameBuffer[FComputeBuffer::MaxNameLength];
+		swprintf(NameBuffer, FComputeBuffer::MaxNameLength, L"%s_M", Name);
 
-		Resources->MemoryMappedFile = OpenFileMappingW(FILE_MAP_ALL_ACCESS, TRUE, NameBuffer);
-		if (Resources->MemoryMappedFile == nullptr)
+		std::shared_ptr<FComputeBufferDetail> Detail = std::make_shared<FComputeBufferDetail>(Name);
+		if (!Detail->MemoryMappedFile.OpenExisting(NameBuffer))
 		{
 			return nullptr;
 		}
 
-		Resources->Header = (FComputeBufferHeader*)MapViewOfFile(Resources->MemoryMappedFile, FILE_MAP_ALL_ACCESS, 0, 0, 0);
-		if (Resources->Header == nullptr)
+		Detail->Header = (FHeader*)Detail->MemoryMappedFile.GetPointer();
+		if (Detail->Header == nullptr)
 		{
 			return nullptr;
 		}
 
-		swprintf_s(NameBuffer, L"%s_W", Name);
-		Resources->WriterEvent = OpenEventW(SYNCHRONIZE | EVENT_MODIFY_STATE, TRUE, NameBuffer);
-		if (Resources->WriterEvent == nullptr)
+		swprintf(NameBuffer, FComputeBuffer::MaxNameLength, L"%s_W", Name);
+		if (!Detail->WriterEvent.OpenExisting(NameBuffer))
 		{
 			return nullptr;
 		}
 
-		for (int ReaderIdx = 0; ReaderIdx < Resources->Header->NumReaders; ReaderIdx++)
+		for (int ReaderIdx = 0; ReaderIdx < Detail->Header->NumReaders; ReaderIdx++)
 		{
-			swprintf_s(NameBuffer, L"%s_R%d", Name, ReaderIdx);
-			Resources->ReaderEvents[ReaderIdx] = OpenEventW(SYNCHRONIZE | EVENT_MODIFY_STATE, TRUE, NameBuffer);
-			if (Resources->ReaderEvents[ReaderIdx] == nullptr)
+			swprintf(NameBuffer, FComputeBuffer::MaxNameLength, L"%s_R%d", Name, ReaderIdx);
+			if (!Detail->ReaderEvents[ReaderIdx].OpenExisting(NameBuffer))
 			{
 				return nullptr;
 			}
 		}
 
-		InitShared(Resources);
-		return Resources;
+		InitShared(Detail);
+		return Detail;
 	}
 
 	void AddRef()
 	{
-		InterlockedIncrement(&RefCount);
+		FComputePlatform::AtomicIncrement(&RefCount);
 	}
 
 	void Release()
 	{
-		if (InterlockedDecrement(&RefCount) == 0)
+		if (FComputePlatform::AtomicDecrement(&RefCount) == 0)
 		{
 			delete this;
 		}
 	}
 
+	bool IsComplete(int ReaderIdx) const
+	{
+		FReaderState ReaderState = Header->Readers[ReaderIdx].Get();
+		FChunkState ChunkState = Header->Chunks[ReaderState.GetChunkIdx()].Get();
+
+		return ChunkState.GetWriteState() == EWriteState::Complete && ReaderState.GetOffset() == ChunkState.GetLength();
+	}
+
+	void AdvanceReadPosition(int ReaderIdx, size_t Size)
+	{
+		Header->Readers[ReaderIdx].Advance(Size);
+	}
+
+	size_t GetMaxReadSize(int ReaderIdx) const
+	{
+		const FReaderStatePtr& readerStatePtr = Header->Readers[ReaderIdx];
+		FReaderState readerState = readerStatePtr.Get();
+
+		const FChunkStatePtr& chunkStatePtr = Header->Chunks[readerState.GetChunkIdx()];
+		FChunkState chunkState = chunkStatePtr.Get();
+
+		if (chunkState.HasReaderFlag(ReaderIdx))
+		{
+			return chunkState.GetLength() - readerState.GetOffset();
+		}
+		else
+		{
+			return 0;
+		}
+	}
+
+	const unsigned char* WaitToRead(int ReaderIdx, size_t MinSize, int TimeoutMs)
+	{
+		for (; ; )
+		{
+			FReaderStatePtr& readerStatePtr = Header->Readers[ReaderIdx];
+			FReaderState readerState = readerStatePtr.Get();
+
+			FChunkStatePtr& chunkStatePtr = Header->Chunks[readerState.GetChunkIdx()];
+			FChunkState chunkState = chunkStatePtr.Get();
+
+			if (!chunkState.HasReaderFlag(ReaderIdx))
+			{
+				// Wait until the current chunk is readable
+				ReaderEvents[ReaderIdx].Reset();
+				if (!chunkState.HasReaderFlag(ReaderIdx) && !ReaderEvents[ReaderIdx].Wait(TimeoutMs))
+				{
+					return nullptr;
+				}
+			}
+			else if (readerState.GetOffset() + MinSize <= chunkState.GetLength())
+			{
+				// We have enough data in the chunk to be able to read a message
+				return ChunkPtrs[readerState.GetChunkIdx()] + readerState.GetOffset();
+			}
+			else if (chunkState.GetWriteState() == EWriteState::Writing)
+			{
+				// Wait until there is more data in the chunk
+				ReaderEvents[ReaderIdx].Reset();
+				if (Header->Chunks[readerState.GetChunkIdx()].Get().Value == chunkState.Value && !ReaderEvents[ReaderIdx].Wait(TimeoutMs))
+				{
+					return nullptr;
+				}
+			}
+			else if (readerState.GetOffset() < chunkState.GetLength() || chunkState.GetWriteState() == EWriteState::Complete)
+			{
+				// Cannot read the requested amount of data from this chunk.
+				return nullptr;
+			}
+			else if (chunkState.GetWriteState() == EWriteState::MovedToNext)
+			{
+				// Move to the next chunk
+				chunkStatePtr.FinishReading(ReaderIdx);
+				WriterEvent.Set();
+
+				int chunkIdx = readerState.GetChunkIdx() + 1;
+				if (chunkIdx == Header->NumChunks)
+				{
+					chunkIdx = 0;
+				}
+
+				readerStatePtr.Set(FReaderState(chunkIdx, 0));
+			}
+			else
+			{
+				assert(false);
+			}
+		}
+	}
+
+	void MarkComplete()
+	{
+		Header->Chunks[Header->WriteChunkIdx].MarkComplete();
+		SetAllReaderEvents();
+	}
+
+	void AdvanceWritePosition(size_t Size)
+	{
+		Header->Chunks[Header->WriteChunkIdx].Append(Size);
+		SetAllReaderEvents();
+	}
+
+	size_t GetMaxWriteSize() const
+	{
+		FChunkState Value = Header->Chunks[Header->WriteChunkIdx].Get();
+		if (Value.GetWriteState() == EWriteState::Complete)
+		{
+			return 0;
+		}
+		else
+		{
+			return Header->ChunkLength - Value.GetLength();
+		}
+	}
+
+	size_t Write(const void* Buffer, size_t MaxSize, int TimeoutMs)
+	{
+		unsigned char* SendData = WaitToWrite(1, TimeoutMs);
+		if (SendData == nullptr)
+		{
+			return 0;
+		}
+
+		size_t SendSize = GetMaxWriteSize();
+		if (MaxSize < SendSize)
+		{
+			SendSize = MaxSize;
+		}
+
+		memcpy(SendData, Buffer, SendSize);
+		AdvanceWritePosition(SendSize);
+		return SendSize;
+	}
+
+	unsigned char* WaitToWrite(size_t MinSize, int TimeoutMs)
+	{
+		assert(MinSize <= Header->ChunkLength);
+
+		for (; ; )
+		{
+			FChunkStatePtr WriteChunkStatePtr = Header->Chunks[Header->WriteChunkIdx];
+
+			FChunkState ChunkState = WriteChunkStatePtr.Get();
+			if (ChunkState.GetWriteState() == EWriteState::Writing)
+			{
+				int Length = ChunkState.GetLength();
+				if (Length + MinSize <= Header->ChunkLength)
+				{
+					return ChunkPtrs[Header->WriteChunkIdx] + Length;
+				}
+
+				WriteChunkStatePtr.FinishWriting(); // STATE CHANGE
+				SetAllReaderEvents();
+			}
+
+			if (ChunkState.GetWriteState() == EWriteState::Complete)
+			{
+				return nullptr;
+			}
+
+			int NextWriteChunkIdx = Header->WriteChunkIdx + 1;
+			if (NextWriteChunkIdx == Header->NumChunks)
+			{
+				NextWriteChunkIdx = 0;
+			}
+
+			FChunkStatePtr NextWriteChunkStatePtr = Header->Chunks[NextWriteChunkIdx];
+			while (NextWriteChunkStatePtr.Get().GetReaderFlags() != 0)
+			{
+				if (!WriterEvent.Wait(TimeoutMs))
+				{
+					return nullptr;
+				}
+				WriterEvent.Reset();
+			}
+
+			Header->WriteChunkIdx = NextWriteChunkIdx;
+
+			NextWriteChunkStatePtr.StartWriting(Header->NumReaders);
+		}
+	}
+
 private:
-	void* MemoryMappedFile;
+	FComputeMemoryMappedFile MemoryMappedFile;
 	volatile long RefCount;
 
-	static void InitShared(std::shared_ptr<FComputeBufferResources> Resources)
+	static void InitShared(std::shared_ptr<FComputeBufferDetail> Detail)
 	{
-		Resources->Reader = FComputeBufferReader(Resources, 0);
-		Resources->Writer = FComputeBufferWriter(Resources);
+		Detail->Reader = FComputeBufferReader(Detail, 0);
+		Detail->Writer = FComputeBufferWriter(Detail);
 
-		FComputeBufferHeader* Header = Resources->Header;
+		FHeader* Header = Detail->Header;
 
 		unsigned char* NextPtr = (unsigned char*)(Header + 1);
 		for (int ChunkIdx = 0; ChunkIdx < Header->NumChunks; ChunkIdx++)
 		{
-			Resources->ChunkPtrs[ChunkIdx] = NextPtr;
+			Detail->ChunkPtrs[ChunkIdx] = NextPtr;
 			NextPtr += Header->ChunkLength;
 		}
 	}
+
+	void SetAllReaderEvents()
+	{
+		for (int Idx = 0; Idx < Header->NumReaders; Idx++)
+		{
+			ReaderEvents[Idx].Set();
+		}
+	}
 };
 
-unsigned int FComputeBufferResources::Counter = 0;
 
 
 
 //// FComputeBuffer /////
 
 FComputeBuffer::FComputeBuffer()
-	: Resources(nullptr)
+	: Detail(nullptr)
 {
 }
 
 FComputeBuffer::FComputeBuffer(const FComputeBuffer& Buffer)
 	: FComputeBuffer()
 {
-	if (Buffer.Resources != nullptr)
+	if (Buffer.Detail != nullptr)
 	{
-		Resources = Buffer.Resources;
-		Resources->AddRef();
+		Detail = Buffer.Detail;
+		Detail->AddRef();
 	}
 }
 
 FComputeBuffer::FComputeBuffer(FComputeBuffer&& Buffer) noexcept
-	: Resources(Buffer.Resources)
+	: Detail(Buffer.Detail)
 {
-	Buffer.Resources = nullptr;
+	Buffer.Detail = nullptr;
 }
 
 FComputeBuffer::~FComputeBuffer()
@@ -362,39 +524,39 @@ FComputeBuffer::~FComputeBuffer()
 
 bool FComputeBuffer::CreateNew(const FParams& Params)
 {
-	Resources = FComputeBufferResources::CreateNew(Params);
-	return Resources != nullptr;
+	Detail = FComputeBufferDetail::CreateNew(Params);
+	return Detail != nullptr;
 }
 
 bool FComputeBuffer::OpenExisting(const wchar_t* Name)
 {
-	Resources = FComputeBufferResources::OpenExisting(Name);
-	return Resources != nullptr;
+	Detail = FComputeBufferDetail::OpenExisting(Name);
+	return Detail != nullptr;
 }
 
 void FComputeBuffer::Close()
 {
-	Resources.reset();
+	Detail.reset();
 }
 
 FComputeBufferReader& FComputeBuffer::GetReader()
 {
-	return Resources->Reader;
+	return Detail->Reader;
 }
 
 const FComputeBufferReader& FComputeBuffer::GetReader() const
 {
-	return Resources->Reader;
+	return Detail->Reader;
 }
 
 FComputeBufferWriter& FComputeBuffer::GetWriter()
 {
-	return Resources->Writer;
+	return Detail->Writer;
 }
 
 const FComputeBufferWriter& FComputeBuffer::GetWriter() const
 {
-	return Resources->Writer;
+	return Detail->Writer;
 }
 
 
@@ -406,13 +568,13 @@ const FComputeBufferWriter& FComputeBuffer::GetWriter() const
 //// FComputeBufferReader /////
 
 FComputeBufferReader::FComputeBufferReader()
-	: Resources(nullptr)
+	: Detail(nullptr)
 	, ReaderIdx(-1)
 {
 }
 
-FComputeBufferReader::FComputeBufferReader(std::shared_ptr<FComputeBufferResources> Resources, int ReaderIdx)
-	: Resources(std::move(Resources))
+FComputeBufferReader::FComputeBufferReader(std::shared_ptr<FComputeBufferDetail> Detail, int ReaderIdx)
+	: Detail(std::move(Detail))
 	, ReaderIdx(ReaderIdx)
 {
 }
@@ -421,105 +583,53 @@ FComputeBufferReader::~FComputeBufferReader()
 {
 }
 
-bool FComputeBufferReader::IsComplete() const
+void FComputeBufferReader::ForceComplete()
 {
-	FComputeBufferHeader* Header = Resources->Header;
-
-	FReaderState ReaderState = Header->Readers[ReaderIdx].Get();
-	FChunkState ChunkState = Header->Chunks[ReaderState.GetChunkIdx()].Get();
+	Detail->MarkComplete();
+}
 
-	return ChunkState.GetWriteState() == EWriteState::Complete && ReaderState.GetOffset() == ChunkState.GetLength();
+bool FComputeBufferReader::IsComplete() const
+{
+	return Detail->IsComplete(ReaderIdx);
 }
 
 void FComputeBufferReader::AdvanceReadPosition(size_t Size)
 {
-	FComputeBufferHeader* Header = Resources->Header;
-	Header->Readers[ReaderIdx].Advance(Size);
+	Detail->AdvanceReadPosition(ReaderIdx, Size);
 }
 
 size_t FComputeBufferReader::GetMaxReadSize() const
 {
-	FComputeBufferHeader* Header = Resources->Header;
-
-	const FReaderStatePtr& readerStatePtr = Header->Readers[ReaderIdx];
-	FReaderState readerState = readerStatePtr.Get();
-
-	const FChunkStatePtr& chunkStatePtr = Header->Chunks[readerState.GetChunkIdx()];
-	FChunkState chunkState = chunkStatePtr.Get();
+	return Detail->GetMaxReadSize(ReaderIdx);
+}
 
-	if (chunkState.HasReaderFlag(ReaderIdx))
+size_t FComputeBufferReader::Read(void* Buffer, size_t MaxSize, int TimeoutMs)
+{
+	const unsigned char* RecvData = WaitToRead(1, TimeoutMs);
+	if (RecvData == nullptr)
 	{
-		return chunkState.GetLength() - readerState.GetOffset();
+		return 0;
 	}
-	else
+
+	size_t RecvSize = GetMaxReadSize();
+	if (MaxSize < RecvSize)
 	{
-		return 0;
+		RecvSize = MaxSize;
 	}
+
+	memcpy(Buffer, RecvData, RecvSize);
+	AdvanceReadPosition(RecvSize);
+	return RecvSize;
 }
 
 const unsigned char* FComputeBufferReader::WaitToRead(size_t MinSize, int TimeoutMs)
 {
-	FComputeBufferHeader* Header = Resources->Header;
-	for (; ; )
-	{
-		FReaderStatePtr& readerStatePtr = Header->Readers[ReaderIdx];
-		FReaderState readerState = readerStatePtr.Get();
-
-		FChunkStatePtr& chunkStatePtr = Header->Chunks[readerState.GetChunkIdx()];
-		FChunkState chunkState = chunkStatePtr.Get();
-
-		if (!chunkState.HasReaderFlag(ReaderIdx))
-		{
-			// Wait until the current chunk is readable
-			ResetEvent(Resources->ReaderEvents[ReaderIdx]);
-			if (!chunkState.HasReaderFlag(ReaderIdx) && WaitForSingleObject(Resources->ReaderEvents[ReaderIdx], TimeoutMs) == WAIT_TIMEOUT)
-			{
-				return nullptr;
-			}
-		}
-		else if (readerState.GetOffset() + MinSize <= chunkState.GetLength())
-		{
-			// We have enough data in the chunk to be able to read a message
-			return Resources->ChunkPtrs[readerState.GetChunkIdx()] + readerState.GetOffset();
-		}
-		else if (chunkState.GetWriteState() == EWriteState::Writing)
-		{
-			// Wait until there is more data in the chunk
-			ResetEvent(Resources->ReaderEvents[ReaderIdx]);
-			if (Header->Chunks[readerState.GetChunkIdx()].Get().Value == chunkState.Value && WaitForSingleObject(Resources->ReaderEvents[ReaderIdx], TimeoutMs) == WAIT_TIMEOUT)
-			{
-				return nullptr;
-			}
-		}
-		else if (readerState.GetOffset() < chunkState.GetLength() || chunkState.GetWriteState() == EWriteState::Complete)
-		{
-			// Cannot read the requested amount of data from this chunk.
-			return nullptr;
-		}
-		else if (chunkState.GetWriteState() == EWriteState::MovedToNext)
-		{
-			// Move to the next chunk
-			chunkStatePtr.FinishReading(ReaderIdx);
-			SetEvent(Resources->WriterEvent);
-
-			int chunkIdx = readerState.GetChunkIdx() + 1;
-			if (chunkIdx == Header->NumChunks)
-			{
-				chunkIdx = 0;
-			}
-
-			readerStatePtr.Set(FReaderState(chunkIdx, 0));
-		}
-		else
-		{
-			assert(false);
-		}
-	}
+	return Detail->WaitToRead(ReaderIdx, MinSize, TimeoutMs);
 }
 
 const wchar_t* FComputeBufferReader::GetName() const
 {
-	return Resources->Name;
+	return Detail->Name;
 }
 
 
@@ -529,12 +639,12 @@ const wchar_t* FComputeBufferReader::GetName() const
 //// FComputeBufferWriter /////
 
 FComputeBufferWriter::FComputeBufferWriter()
-	: Resources(nullptr)
+	: Detail(nullptr)
 {
 }
 
-FComputeBufferWriter::FComputeBufferWriter(std::shared_ptr<FComputeBufferResources> Resources)
-	: Resources(std::move(Resources))
+FComputeBufferWriter::FComputeBufferWriter(std::shared_ptr<FComputeBufferDetail> Detail)
+	: Detail(std::move(Detail))
 {
 }
 
@@ -544,81 +654,44 @@ FComputeBufferWriter::~FComputeBufferWriter()
 
 void FComputeBufferWriter::MarkComplete()
 {
-	FComputeBufferHeader* Header = Resources->Header;
-	Header->Chunks[Header->WriteChunkIdx].MarkComplete();
-	SetAllReaderEvents();
+	Detail->MarkComplete();
 }
 
 void FComputeBufferWriter::AdvanceWritePosition(size_t Size)
 {
-	FComputeBufferHeader* Header = Resources->Header;
-	Header->Chunks[Header->WriteChunkIdx].Append(Size);
-	SetAllReaderEvents();
+	Detail->AdvanceWritePosition(Size);
 }
 
 size_t FComputeBufferWriter::GetMaxWriteSize() const
 {
-	FComputeBufferHeader* Header = Resources->Header;
-	FChunkState Value = Header->Chunks[Header->WriteChunkIdx].Get();
-	return Header->ChunkLength - Value.GetLength();
+	return Detail->GetMaxWriteSize();
 }
 
-unsigned char* FComputeBufferWriter::WaitToWrite(size_t MinSize, int TimeoutMs)
+size_t FComputeBufferWriter::Write(const void* Buffer, size_t MaxSize, int TimeoutMs)
 {
-	FComputeBufferHeader* Header = Resources->Header;
-	assert(MinSize <= Header->ChunkLength);
-
-	DWORD WaitParam = (TimeoutMs == -1) ? INFINITE : (unsigned int)TimeoutMs;
-	for (; ; )
+	unsigned char* SendData = WaitToWrite(1, TimeoutMs);
+	if (SendData == nullptr)
 	{
-		FChunkStatePtr WriteChunkStatePtr = Header->Chunks[Header->WriteChunkIdx];
-
-		FChunkState ChunkState = WriteChunkStatePtr.Get();
-		if (ChunkState.GetWriteState() == EWriteState::Writing)
-		{
-			int Length = ChunkState.GetLength();
-			if (Length + MinSize <= Header->ChunkLength)
-			{
-				return Resources->ChunkPtrs[Header->WriteChunkIdx] + Length;
-			}
-
-			WriteChunkStatePtr.FinishWriting(); // STATE CHANGE
-			SetAllReaderEvents();
-		}
-
-		int NextWriteChunkIdx = Header->WriteChunkIdx + 1;
-		if (NextWriteChunkIdx == Header->NumChunks)
-		{
-			NextWriteChunkIdx = 0;
-		}
-
-		FChunkStatePtr NextWriteChunkStatePtr = Header->Chunks[NextWriteChunkIdx];
-		while (NextWriteChunkStatePtr.Get().GetReaderFlags() != 0)
-		{
-			if (WaitForSingleObject(Resources->WriterEvent, WaitParam) == WAIT_TIMEOUT)
-			{
-				return nullptr;
-			}
-			ResetEvent(Resources->WriterEvent);
-		}
-
-		Header->WriteChunkIdx = NextWriteChunkIdx;
+		return 0;
+	}
 
-		NextWriteChunkStatePtr.StartWriting(Header->NumReaders);
+	size_t SendSize = GetMaxWriteSize();
+	if (MaxSize < SendSize)
+	{
+		SendSize = MaxSize;
 	}
+
+	memcpy(SendData, Buffer, SendSize);
+	AdvanceWritePosition(SendSize);
+	return SendSize;
 }
 
-const wchar_t* FComputeBufferWriter::GetName() const
+unsigned char* FComputeBufferWriter::WaitToWrite(size_t MinSize, int TimeoutMs)
 {
-	return Resources->Name;
+	return Detail->WaitToWrite(MinSize, TimeoutMs);
 }
 
-void FComputeBufferWriter::SetAllReaderEvents()
+const wchar_t* FComputeBufferWriter::GetName() const
 {
-	FComputeBufferHeader* Header = Resources->Header;
-	for (int Idx = 0; Idx < Header->NumReaders; Idx++)
-	{
-		SetEvent(Resources->ReaderEvents[Idx]);
-	}
+	return Detail->Name;
 }
-
diff --git a/Engine/Source/Programs/Horde/Samples/RemoteWorkerCpp/ComputeBuffer.h b/Engine/Source/Programs/Horde/Samples/RemoteWorkerCpp/ComputeBuffer.h
--- a/Engine/Source/Programs/Horde/Samples/RemoteWorkerCpp/ComputeBuffer.h
+++ b/Engine/Source/Programs/Horde/Samples/RemoteWorkerCpp/ComputeBuffer.h
@@ -7,8 +7,8 @@
 class FComputeBufferReader;
 class FComputeBufferWriter;
 
-struct FComputeBufferHeader;
-struct FComputeBufferResources;
+struct FHeader;
+struct FComputeBufferDetail;
 
 //
 // Implements a ring buffer using shared memory, with one writer and multiple readers. 
@@ -29,6 +29,9 @@ public:
 	// Maximum allowed number of chunks
 	static const int MaxChunks = 16;
 
+	// Maximum length of a compute buffer name
+	static const size_t MaxNameLength = 256;
+
 	// Parameters for constructing a compute buffer
 	struct FParams
 	{
@@ -53,7 +56,7 @@ public:
 	void Close();
 
 	// Test if the buffer is currently open
-	bool IsOpen() const { return Resources != nullptr; }
+	bool IsOpen() const { return Detail != nullptr; }
 
 	// Gets a reference to the reader instance
 	FComputeBufferReader& GetReader();
@@ -64,7 +67,7 @@ public:
 	const FComputeBufferWriter& GetWriter() const;
 
 private:
-	std::shared_ptr<FComputeBufferResources> Resources;
+	std::shared_ptr<FComputeBufferDetail> Detail;
 };
 
 //
@@ -74,21 +77,27 @@ class FComputeBufferReader
 {
 public:
 	FComputeBufferReader();
-	FComputeBufferReader(std::shared_ptr<FComputeBufferResources> Resources, int ReaderIdx);
+	FComputeBufferReader(std::shared_ptr<FComputeBufferDetail> Detail, int ReaderIdx);
 	~FComputeBufferReader();
 
 	// Test if the reader is valid
-	bool IsValid() const { return Resources.get() != nullptr; }
+	bool IsValid() const { return Detail.get() != nullptr; }
 
 	// Test whether the buffer has finished being written to (ie. MarkComplete() has been called by the writer) and all data has been read from it.
 	bool IsComplete() const;
 
+	// Mark this buffer as complete. This is technically a "write" operation, but is needed to release blocking waits on shutdown.
+	void ForceComplete();
+
 	// Move the read cursor forwards by the given number of bytes
 	void AdvanceReadPosition(size_t Size);
 
 	// Gets the amount of data that is ready to be read from a contiguous block of memory.
 	size_t GetMaxReadSize() const;
 
+	// Reads data into the given buffer
+	size_t Read(void* Buffer, size_t MaxSize, int TimeoutMs = -1);
+
 	// Waits until the given amount of data has been read, and returns a pointer to it. Returns nullptr if the timeout expires, or
 	// if the requested amount of data is not in a contiguous block of memory.
 	const unsigned char* WaitToRead(size_t MinSize, int TimeoutMs = -1);
@@ -96,7 +105,7 @@ public:
 private:
 	friend class FWorkerComputeSocket;
 
-	std::shared_ptr<FComputeBufferResources> Resources;
+	std::shared_ptr<FComputeBufferDetail> Detail;
 	int ReaderIdx;
 
 	const wchar_t* GetName() const;
@@ -109,11 +118,11 @@ class FComputeBufferWriter
 {
 public:
 	FComputeBufferWriter();
-	FComputeBufferWriter(std::shared_ptr<FComputeBufferResources> Resources);
+	FComputeBufferWriter(std::shared_ptr<FComputeBufferDetail> Detail);
 	~FComputeBufferWriter();
 
 	// Test if the writer is valid
-	bool IsValid() const { return Resources.get() != nullptr; }
+	bool IsValid() const { return Detail.get() != nullptr; }
 
 	// Signal that we've finished writing to this buffer
 	void MarkComplete();
@@ -124,6 +133,9 @@ public:
 	// Gets the length of the current write buffer.
 	size_t GetMaxWriteSize() const;
 
+	// Writes data into the compute buffer
+	size_t Write(const void* Buffer, size_t MaxSize, int TimeoutMs = -1);
+
 	// Waits until a write buffer of the requested size is available, and returns a pointer to it. Returns nullptr if the 
 	// timeout expires before enough data has been flushed to return a buffer of the given length.
 	unsigned char* WaitToWrite(size_t MinSize, int TimeoutMs = -1);
@@ -131,8 +143,7 @@ public:
 private:
 	friend class FWorkerComputeSocket;
 
-	std::shared_ptr<FComputeBufferResources> Resources;
+	std::shared_ptr<FComputeBufferDetail> Detail;
 
 	const wchar_t* GetName() const;
-	void SetAllReaderEvents();
 };
diff --git a/Engine/Source/Programs/Horde/Samples/RemoteWorkerCpp/ComputeChannel.cpp b/Engine/Source/Programs/Horde/Samples/RemoteWorkerCpp/ComputeChannel.cpp
--- a/Engine/Source/Programs/Horde/Samples/RemoteWorkerCpp/ComputeChannel.cpp
+++ b/Engine/Source/Programs/Horde/Samples/RemoteWorkerCpp/ComputeChannel.cpp
@@ -1,101 +1,40 @@
 // Copyright Epic Games, Inc. All Rights Reserved.
 
+#include "ComputeChannel.h"
 #include <string.h>
 #include <stdio.h>
-#include <windows.h>
-#include "ComputeChannel.h"
 #include "ComputeSocket.h"
 
 FComputeChannel::FComputeChannel()
 {
 }
 
-FComputeChannel::~FComputeChannel()
+FComputeChannel::FComputeChannel(FComputeBufferReader InReader, FComputeBufferWriter InWriter)
+	: Reader(std::move(InReader))
+	, Writer(std::move(InWriter))
 {
-	if (SendBufferWriter.IsValid())
-	{
-		SendBufferWriter.MarkComplete();
-	}
 }
 
-void FComputeChannel::Attach(FComputeSocket& Socket, int ChannelId, FComputeBuffer SendBuffer, FComputeBuffer RecvBuffer)
+FComputeChannel::~FComputeChannel()
 {
-	RecvBufferReader = RecvBuffer.GetReader();
-	SendBufferWriter = SendBuffer.GetWriter();
-	Socket.AttachRecvBuffer(ChannelId, RecvBuffer.GetWriter());
-	Socket.AttachSendBuffer(ChannelId, SendBuffer.GetReader());
 }
 
-bool FComputeChannel::Attach(FComputeSocket& Socket, int ChannelId, const FComputeBuffer::FParams& Params)
+bool FComputeChannel::IsValid() const
 {
-	return Attach(Socket, ChannelId, Params, Params);
+	return Reader.IsValid();
 }
 
-bool FComputeChannel::Attach(FComputeSocket& Socket, int ChannelId, const FComputeBuffer::FParams& SendBufferParams, const FComputeBuffer::FParams& RecvBufferParams)
+size_t FComputeChannel::Send(const void* Data, size_t Size, int TimeoutMs)
 {
-	Detach();
-
-	FComputeBuffer SendBuffer;
-	if (!SendBuffer.CreateNew(SendBufferParams))
-	{
-		return false;
-	}
-
-	FComputeBuffer RecvBuffer;
-	if (!RecvBuffer.CreateNew(RecvBufferParams))
-	{
-		return false;
-	}
-
-	Attach(Socket, ChannelId, SendBuffer, RecvBuffer);
-	return true;
+	return Writer.Write(Data, Size, TimeoutMs);
 }
 
-void FComputeChannel::Detach()
+size_t FComputeChannel::Recv(void* Data, size_t Size, int TimeoutMs)
 {
-	RecvBufferReader = FComputeBufferReader();
-	SendBufferWriter = FComputeBufferWriter();
+	return Reader.Read(Data, Size, TimeoutMs);
 }
 
 void FComputeChannel::MarkComplete()
 {
-	SendBufferWriter.MarkComplete();
-}
-
-size_t FComputeChannel::Send(const void* Data, size_t Size, int TimeoutMs)
-{
-	unsigned char* SendData = SendBufferWriter.WaitToWrite(1, TimeoutMs);
-	if (SendData == nullptr)
-	{
-		return 0;
-	}
-
-	size_t SendSize = SendBufferWriter.GetMaxWriteSize();
-	if (Size < SendSize)
-	{
-		SendSize = Size;
-	}
-
-	memcpy(SendData, Data, SendSize);
-	SendBufferWriter.AdvanceWritePosition(SendSize);
-	return SendSize;
-}
-
-size_t FComputeChannel::Recv(void* Data, size_t Size, int TimeoutMs)
-{
-	const unsigned char* RecvData = RecvBufferReader.WaitToRead(1, TimeoutMs);
-	if (RecvData == nullptr)
-	{
-		return 0;
-	}
-
-	size_t RecvSize = RecvBufferReader.GetMaxReadSize();
-	if (Size < RecvSize)
-	{
-		RecvSize = Size;
-	}
-
-	memcpy(Data, RecvData, RecvSize);
-	RecvBufferReader.AdvanceReadPosition(RecvSize);
-	return RecvSize;
+	Writer.MarkComplete();
 }
diff --git a/Engine/Source/Programs/Horde/Samples/RemoteWorkerCpp/ComputeChannel.h b/Engine/Source/Programs/Horde/Samples/RemoteWorkerCpp/ComputeChannel.h
--- a/Engine/Source/Programs/Horde/Samples/RemoteWorkerCpp/ComputeChannel.h
+++ b/Engine/Source/Programs/Horde/Samples/RemoteWorkerCpp/ComputeChannel.h
@@ -12,23 +12,18 @@ class FComputeSocket;
 class FComputeChannel
 {
 public:
-	FComputeChannel();
-	~FComputeChannel();
-
-	// Creates a new channel using default parameters
-	void Attach(FComputeSocket& Socket, int ChannelId, FComputeBuffer SendBuffer, FComputeBuffer RecvBuffer);
+	// Reader for the channel
+	FComputeBufferReader Reader;
 
-	// Creates a new channel using the same parameters for the send and receive buffers
-	bool Attach(FComputeSocket& Socket, int ChannelId, const FComputeBuffer::FParams& Params);
+	// Writer for the channel
+	FComputeBufferWriter Writer;
 
-	// Creates a new channel using the same parameters for the send and receive buffers
-	bool Attach(FComputeSocket& Socket, int ChannelId, const FComputeBuffer::FParams& SendParams, const FComputeBuffer::FParams& RecvParams);
-
-	// Close the current buffer and release all allocated resources
-	void Detach();
+	FComputeChannel();
+	FComputeChannel(FComputeBufferReader InReader, FComputeBufferWriter InWriter);
+	~FComputeChannel();
 
-	// Indicate to the remote that no more data will be sent.
-	void MarkComplete();
+	// Tests whether the channel is valid
+	bool IsValid() const;
 
 	// Sends bytes to the remote. 
 	size_t Send(const void* Data, size_t Size, int TimeoutMs = -1);
@@ -36,7 +31,6 @@ public:
 	// Reads as many bytes as are available from the socket.
 	size_t Recv(void* Data, size_t Size, int TimeoutMs = -1);
 
-private:
-	FComputeBufferReader RecvBufferReader;
-	FComputeBufferWriter SendBufferWriter;
+	// Indicate to the remote that no more data will be sent.
+	void MarkComplete();
 };
diff --git a/Engine/Source/Programs/Horde/Samples/RemoteWorkerCpp/ComputePlatform.cpp b/Engine/Source/Programs/Horde/Samples/RemoteWorkerCpp/ComputePlatform.cpp
new file mode 100644
--- /dev/null
+++ b/Engine/Source/Programs/Horde/Samples/RemoteWorkerCpp/ComputePlatform.cpp
@@ -0,0 +1,239 @@
+#include "ComputePlatform.h"
+#include <assert.h>
+#include <wchar.h>
+
+#ifdef _MSC_VER
+#include <Windows.h>
+#undef GetEnvironmentVariable
+#endif
+
+/////////////////////////////////////////////////// 
+
+FComputeManualResetEvent::FComputeManualResetEvent()
+	: Handle(nullptr)
+{
+}
+
+FComputeManualResetEvent::~FComputeManualResetEvent()
+{
+	Close();
+}
+
+bool FComputeManualResetEvent::Create(const wchar_t* Name)
+{
+	Close();
+
+	Handle = CreateEventW(NULL, TRUE, FALSE, Name);
+	return Handle != nullptr;
+}
+
+bool FComputeManualResetEvent::OpenExisting(const wchar_t* Name)
+{
+	Close();
+
+	Handle = OpenEventW(SYNCHRONIZE | EVENT_MODIFY_STATE, TRUE, Name);
+	return Handle != nullptr;
+}
+
+void FComputeManualResetEvent::Close()
+{
+	if (Handle != nullptr)
+	{
+		CloseHandle(Handle);
+		Handle = nullptr;
+	}
+}
+
+void FComputeManualResetEvent::Set()
+{
+	SetEvent(Handle);
+}
+
+void FComputeManualResetEvent::Reset()
+{
+	ResetEvent(Handle);
+}
+
+bool FComputeManualResetEvent::Wait(int timeoutMs)
+{
+	DWORD WaitParam = (timeoutMs < 0) ? INFINITE : (DWORD)timeoutMs;
+	return WaitForSingleObject(Handle, WaitParam) != WAIT_TIMEOUT;
+}
+
+/////////////////////////////////////////////////// 
+
+FComputeMemoryMappedFile::FComputeMemoryMappedFile()
+	: Handle(nullptr)
+	, Pointer(nullptr)
+{
+}
+
+FComputeMemoryMappedFile::~FComputeMemoryMappedFile()
+{
+	Close();
+}
+
+bool FComputeMemoryMappedFile::Create(const wchar_t* Name, long long Capacity)
+{
+	Close();
+
+	LARGE_INTEGER LargeInteger;
+	LargeInteger.QuadPart = Capacity;
+
+	Handle = CreateFileMappingW(INVALID_HANDLE_VALUE, NULL, PAGE_READWRITE, LargeInteger.HighPart, LargeInteger.LowPart, Name);
+	if (Handle == nullptr)
+	{
+		return false;
+	}
+
+	Pointer = MapViewOfFile(Handle, FILE_MAP_ALL_ACCESS, 0, 0, 0);
+	if (Pointer == nullptr)
+	{
+		return false;
+	}
+
+	return true;
+}
+
+bool FComputeMemoryMappedFile::OpenExisting(const wchar_t* Name)
+{
+	Close();
+
+	Handle = OpenFileMappingW(FILE_MAP_ALL_ACCESS, TRUE, Name);
+	if (Handle == nullptr)
+	{
+		return false;
+	}
+
+	Pointer = MapViewOfFile(Handle, FILE_MAP_ALL_ACCESS, 0, 0, 0);
+	if (Pointer == nullptr)
+	{
+		return false;
+	}
+
+	return true;
+}
+
+void FComputeMemoryMappedFile::Close()
+{
+	if (Pointer != nullptr)
+	{
+		UnmapViewOfFile(Pointer);
+		Pointer = nullptr;
+	}
+
+	if (Handle != nullptr)
+	{
+		CloseHandle(Handle);
+		Handle = nullptr;
+	}
+}
+
+void* FComputeMemoryMappedFile::GetPointer() const
+{
+	return Pointer;
+}
+
+/////////////////////////////////////////////////// 
+
+bool FComputePlatform::GetEnvironmentVariable(const wchar_t* Name, wchar_t* Buffer, size_t BufferLen)
+{
+	int Length = GetEnvironmentVariableW(Name, Buffer, (DWORD)BufferLen);
+	return Length > 0 && Length < BufferLen;
+}
+
+void FComputePlatform::CreateUniqueName(wchar_t* NameBuffer, size_t NameBufferLen)
+{
+	static long Counter = 0;
+
+	DWORD Pid = GetCurrentProcessId();
+	ULONGLONG TickCount = GetTickCount64();
+	swprintf(NameBuffer, NameBufferLen, L"Local\\COMPUTE_%u_%llu_%lu", Pid, TickCount, AtomicIncrement(&Counter));
+}
+
+unsigned int FComputePlatform::FloorLog2(unsigned int Value)
+{
+	// Use BSR to return the log2 of the integer
+	// return 0 if value is 0
+	unsigned long BitIndex;
+	return _BitScanReverse(&BitIndex, Value) ? BitIndex : 0;
+}
+
+unsigned int FComputePlatform::CountLeadingZeros(unsigned int Value)
+{
+	// return 32 if value is zero
+	unsigned long BitIndex;
+	_BitScanReverse64(&BitIndex, (unsigned long long)(Value) * 2 + 1);
+	return 32 - BitIndex;
+}
+
+size_t FComputePlatform::Utf8ToWchar(const char* Source, size_t SourceLen, wchar_t* Dest, size_t DestMaxLen)
+{
+	size_t DecodedLen = MultiByteToWideChar(CP_UTF8, 0, Source, (int)SourceLen, Dest, (int)DestMaxLen);
+	Dest[DecodedLen] = 0;
+	return DecodedLen;
+}
+
+size_t FComputePlatform::WcharToUtf8(const wchar_t* Source, size_t SourceLen, char* Dest, size_t DestMaxLen)
+{
+	if (SourceLen == 0)
+	{
+		return 0;
+	}
+
+	int Result = WideCharToMultiByte(CP_UTF8, 0, Source, (int)SourceLen, Dest, (int)DestMaxLen, nullptr, nullptr);
+	assert(Result > 0);
+
+	return (size_t)Result;
+}
+
+long long FComputePlatform::AtomicRead64(volatile long long* Ptr)
+{
+	return InterlockedCompareExchange64(Ptr, 0, 0);
+}
+
+void FComputePlatform::AtomicWrite64(volatile long long* Ptr, long long Value)
+{
+	InterlockedExchange64(Ptr, Value);
+}
+
+long FComputePlatform::AtomicIncrement(volatile long* Ptr)
+{
+	return InterlockedIncrement(Ptr);
+}
+
+long FComputePlatform::AtomicDecrement(volatile long* Ptr)
+{
+	return InterlockedDecrement(Ptr);
+}
+
+long long FComputePlatform::AtomicAdd64(volatile long long* Ptr, long long Value)
+{
+	return InterlockedAdd64(Ptr, Value);
+}
+
+long long FComputePlatform::AtomicAnd64(volatile long long* Ptr, long long Value)
+{
+	return InterlockedAnd64(Ptr, Value);
+}
+
+long long FComputePlatform::AtomicOr64(volatile long long* Ptr, long long Value)
+{
+	return InterlockedOr64(Ptr, Value);
+}
+
+long long FComputePlatform::AtomicIncrement64(volatile long long* Ptr)
+{
+	return InterlockedIncrement64(Ptr);
+}
+
+long long FComputePlatform::AtomicExchange64(volatile long long* Ptr, long long Exchange)
+{
+	return InterlockedExchange64(Ptr, Exchange);
+}
+
+bool FComputePlatform::AtomicCompareExchange64(volatile long long* Ptr, long long Exchange, long long Comperand)
+{
+	return InterlockedCompareExchange64(Ptr, Exchange, Comperand) == Comperand;
+}
+
diff --git a/Engine/Source/Programs/Horde/Samples/RemoteWorkerCpp/ComputePlatform.h b/Engine/Source/Programs/Horde/Samples/RemoteWorkerCpp/ComputePlatform.h
new file mode 100644
--- /dev/null
+++ b/Engine/Source/Programs/Horde/Samples/RemoteWorkerCpp/ComputePlatform.h
@@ -0,0 +1,130 @@
+// Copyright Epic Games, Inc. All Rights Reserved.
+
+#pragma once
+
+//
+// IPC manually-reset event object
+//
+class FComputeManualResetEvent
+{
+public:
+	FComputeManualResetEvent();
+	~FComputeManualResetEvent();
+
+	// Creates a new event with the given name
+	bool Create(const wchar_t* Name);
+
+	// Opens an existing event created elsewhere
+	bool OpenExisting(const wchar_t* Name);
+
+	// Close the event and release its resources
+	void Close();
+
+	// Signal the event, releasing any waiters
+	void Set();
+
+	// Reset the event after a call to Set()
+	void Reset();
+
+	// Wait for the event to be signalled or timeout. Pass -1 for timeoutMs to wait infinitely.
+	bool Wait(int timeoutMs);
+
+private:
+	void* Handle;
+};
+
+//
+// IPC memory mapped file
+//
+class FComputeMemoryMappedFile
+{
+public:
+	FComputeMemoryMappedFile();
+	~FComputeMemoryMappedFile();
+
+	// Creates a new memory mapped file with the given capacity
+	bool Create(const wchar_t* Name, long long Capacity);
+
+	// Opens an existing memory mapped file
+	bool OpenExisting(const wchar_t* Name);
+
+	// Close the memory mapped file handle
+	void Close();
+
+	// Gets a pointer to the mapped memory
+	void* GetPointer() const;
+
+private:
+	void* Handle;
+	void* Pointer;
+};
+
+//
+// Platform-specific utility functions
+//
+struct FComputePlatform
+{
+	//
+	// General
+	//
+	// 
+	// Reads an environment variable
+	static bool GetEnvironmentVariable(const wchar_t* Name, wchar_t* Buffer, size_t BufferLen);
+
+	// Creates a unique object name
+	static void CreateUniqueName(wchar_t* NameBuffer, size_t NameBufferLen);
+
+	//
+	// Math
+	//
+
+	// Find the log2 of the given value, returning 0 if the value is zero.
+	static unsigned int FloorLog2(unsigned int Value);
+
+	// Count the number of leading zeros in the given value.
+	static unsigned int CountLeadingZeros(unsigned int Value);
+
+	//
+	// Strings
+	//
+
+	// Translate a UTF8 string to a wchar_t string. Returns the length of the converted string in characters, even if the supplied buffer is not large enough to hold it.
+	static size_t Utf8ToWchar(const char* Source, size_t SourceLen, wchar_t* Dest, size_t DestMaxLen);
+
+	// Translate a wchar_t string to a UTF8 string. Returns the length of the converted string in characters, even if the supplied buffer is not large enough to hold it.
+	static size_t WcharToUtf8(const wchar_t* Source, size_t SourceLen, char* Dest, size_t DestMaxLen);
+
+	//
+	// Atomics
+	//
+
+	// Reads a 64-bit value from the given memory location 
+	static long long AtomicRead64(volatile long long* Ptr);
+
+	// Writes a 64-bit value to the given memory location
+	static void AtomicWrite64(volatile long long* Ptr, long long Value);
+
+	// Increments an integer at the given memory location and returns the incremented value.
+	static long AtomicIncrement(volatile long* Ptr);
+
+	// Increments an integer at the given memory location and returns the incremented value.
+	static long long AtomicIncrement64(volatile long long* Ptr);
+
+	// Decrements an integer at the given memory location and returns the decremented value.
+	static long AtomicDecrement(volatile long* Ptr);
+
+	// Adds a value to the given memory location, and returns the resulting value.
+	static long long AtomicAdd64(volatile long long* Ptr, long long Value);
+
+	// ANDs a value with a given memory location, and returns the contents of the memory location BEFORE the operation.
+	static long long AtomicAnd64(volatile long long* Ptr, long long Value);
+
+	// ORs a value with a given memory location, and returns the contents of the memory location BEFORE the operation.
+	static long long AtomicOr64(volatile long long* Ptr, long long Value);
+
+	// Sets the contents of a pointer to the given value, and returns the original value in that memory location.
+	static long long AtomicExchange64(volatile long long* Ptr, long long Exchange);
+
+	// Compares the contents of a memory location to a value, and exchanges it for another value if they are equal. Returns true if the operation succeeded.
+	static bool AtomicCompareExchange64(volatile long long* Ptr, long long Exchange, long long Comperand);
+};
diff --git a/Engine/Source/Programs/Horde/Samples/RemoteWorkerCpp/ComputeSocket.cpp b/Engine/Source/Programs/Horde/Samples/RemoteWorkerCpp/ComputeSocket.cpp
--- a/Engine/Source/Programs/Horde/Samples/RemoteWorkerCpp/ComputeSocket.cpp
+++ b/Engine/Source/Programs/Horde/Samples/RemoteWorkerCpp/ComputeSocket.cpp
@@ -1,22 +1,52 @@
 // Copyright Epic Games, Inc. All Rights Reserved.
 
-#include <windows.h>
+#include "ComputeSocket.h"
+#include "ComputePlatform.h"
 #include <iostream>
 #include <assert.h>
 #include <uchar.h>
 #include <unordered_set>
 #include <unordered_map>
 #include <vector>
-#include "ComputeSocket.h"
+#include <thread>
+#include <mutex>
+
+FComputeSocket::FComputeSocket()
+{
+}
 
 FComputeSocket::~FComputeSocket()
 {
 }
 
+FComputeChannel FComputeSocket::CreateChannel(int ChannelId)
+{
+	FComputeBuffer RecvBuffer;
+	if (!RecvBuffer.CreateNew(FComputeBuffer::FParams()))
+	{
+		return FComputeChannel();
+	}
+
+	FComputeBuffer SendBuffer;
+	if (!SendBuffer.CreateNew(FComputeBuffer::FParams()))
+	{
+		return FComputeChannel();
+	}
+
+	return CreateChannel(ChannelId, RecvBuffer, SendBuffer);
+}
+
+FComputeChannel FComputeSocket::CreateChannel(int ChannelId, FComputeBuffer RecvBuffer, FComputeBuffer SendBuffer)
+{
+	AttachRecvBuffer(ChannelId, RecvBuffer.GetWriter());
+	AttachSendBuffer(ChannelId, SendBuffer.GetReader());
+
+	return FComputeChannel(RecvBuffer.GetReader(), SendBuffer.GetWriter());
+}
 
 //////////////////////////////////////////////////////
 
-const wchar_t* const FWorkerComputeSocket::EnvVarName = L"UE_HORDE_COMPUTE_IPC";
+const wchar_t* const FWorkerComputeSocket::IpcEnvVar = L"UE_HORDE_COMPUTE_IPC";
 
 enum class FWorkerComputeSocket::EMessageType
 {
@@ -35,10 +65,8 @@ FWorkerComputeSocket::~FWorkerComputeSocket()
 
 bool FWorkerComputeSocket::Open()
 {
-	wchar_t EnvVar[MAX_PATH];
-	int Length = GetEnvironmentVariableW(EnvVarName, EnvVar, sizeof(EnvVar) / sizeof(EnvVar[0]));
-
-	if (Length <= 0 || Length >= sizeof(EnvVar))
+	wchar_t EnvVar[FComputeBuffer::MaxNameLength];
+	if (!FComputePlatform::GetEnvironmentVariable(IpcEnvVar, EnvVar, sizeof(EnvVar) / sizeof(EnvVar[0])))
 	{
 		return false;
 	}
@@ -93,43 +121,43 @@ void FWorkerComputeSocket::RunServer(FComputeBufferReader& CommandBufferReader,
 		switch (MessageType)
 		{
 		case EMessageType::AttachSendBuffer:
-		{
-			unsigned int ChannelId;
-			Len += ReadVarUInt(Message + Len, &ChannelId);
-
-			wchar_t Name[MAX_PATH];
-			Len += ReadString(Message + Len, Name, MAX_PATH);
-
-			FComputeBuffer Buffer;
-			if (Buffer.OpenExisting(Name))
 			{
-				Socket.AttachSendBuffer(ChannelId, Buffer.GetReader());
+				unsigned int ChannelId;
+				Len += ReadVarUInt(Message + Len, &ChannelId);
+
+				wchar_t Name[FComputeBuffer::MaxNameLength];
+				Len += ReadString(Message + Len, Name, FComputeBuffer::MaxNameLength);
+
+				FComputeBuffer Buffer;
+				if (Buffer.OpenExisting(Name))
+				{
+					Socket.AttachSendBuffer(ChannelId, Buffer.GetReader());
+				}
+				else
+				{
+					assert(false);
+				}
 			}
-			else
-			{
-				assert(false);
-			}
-		}
-		break;
+			break;
 		case EMessageType::AttachRecvBuffer:
-		{
-			unsigned int ChannelId;
-			Len += ReadVarUInt(Message + Len, &ChannelId);
-
-			wchar_t Name[MAX_PATH];
-			Len += ReadString(Message + Len, Name, MAX_PATH);
-
-			FComputeBuffer Buffer;
-			if (Buffer.OpenExisting(Name))
 			{
-				Socket.AttachRecvBuffer(ChannelId, Buffer.GetWriter());
+				unsigned int ChannelId;
+				Len += ReadVarUInt(Message + Len, &ChannelId);
+
+				wchar_t Name[FComputeBuffer::MaxNameLength];
+				Len += ReadString(Message + Len, Name, FComputeBuffer::MaxNameLength);
+
+				FComputeBuffer Buffer;
+				if (Buffer.OpenExisting(Name))
+				{
+					Socket.AttachRecvBuffer(ChannelId, Buffer.GetWriter());
+				}
+				else
+				{
+					assert(false);
+				}
 			}
-			else
-			{
-				assert(false);
-			}
-		}
-		break;
+			break;
 		default:
 			assert(false);
 			return;
@@ -141,7 +169,7 @@ void FWorkerComputeSocket::RunServer(FComputeBufferReader& CommandBufferReader,
 
 size_t FWorkerComputeSocket::ReadVarUInt(const unsigned char* Pos, unsigned int* OutValue)
 {
-	size_t ByteCount = CountLeadingZeros((unsigned char)(~*static_cast<const unsigned char*>(Pos))) - 23;
+	size_t ByteCount = FComputePlatform::CountLeadingZeros((unsigned char)(~*static_cast<const unsigned char*>(Pos))) - 23;
 
 	unsigned int Value = *Pos++ & (unsigned char)(0xff >> ByteCount);
 	switch (ByteCount - 1)
@@ -167,8 +195,7 @@ size_t FWorkerComputeSocket::ReadString(const unsigned char* Pos, wchar_t* OutTe
 	unsigned int TextLen;
 	size_t Len = ReadVarUInt(Pos, &TextLen);
 
-	int DecodedLen = MultiByteToWideChar(CP_UTF8, 0, (const char*)Pos + Len, TextLen, OutText, (int)OutTextMaxLen);
-	OutText[DecodedLen] = 0;
+	FComputePlatform::Utf8ToWchar((const char*)Pos + Len, TextLen, OutText, OutTextMaxLen);
 
 	return Len + TextLen;
 }
@@ -177,7 +204,7 @@ size_t FWorkerComputeSocket::WriteVarUInt(unsigned char* Pos, unsigned int Value
 {
 	// Use BSR to return the log2 of the integer
 	// return 0 if value is 0
-	unsigned int ByteCount = (unsigned int)(int(FloorLog2(Value)) / 7 + 1);
+	unsigned int ByteCount = (unsigned int)(int(FComputePlatform::FloorLog2(Value)) / 7 + 1);
 
 	unsigned char* OutBytes = Pos + ByteCount - 1;
 	switch (ByteCount - 1)
@@ -195,194 +222,118 @@ size_t FWorkerComputeSocket::WriteVarUInt(unsigned char* Pos, unsigned int Value
 
 size_t FWorkerComputeSocket::WriteString(unsigned char* Pos, const wchar_t* Text)
 {
-	int EncodedLen = WideCharToMultiByte(CP_UTF8, 0, Text, -1, nullptr, 0, nullptr, nullptr);
-	EncodedLen--; // Ignore null terminator
+	size_t TextLen = wcslen(Text);
 
-	size_t Len = WriteVarUInt(Pos, EncodedLen);
-	WideCharToMultiByte(CP_UTF8, 0, Text, -1, (char*)Pos + Len, EncodedLen, nullptr, nullptr);
+	size_t EncodedLen = FComputePlatform::WcharToUtf8(Text, TextLen, nullptr, 0);
 
-	return Len + EncodedLen;
-}
+	size_t Len = WriteVarUInt(Pos, (int)EncodedLen);
+	FComputePlatform::WcharToUtf8(Text, TextLen, (char*)Pos + Len, EncodedLen);
 
-unsigned int FWorkerComputeSocket::FloorLog2(unsigned int Value)
-{
-	// Use BSR to return the log2 of the integer
-	// return 0 if value is 0
-	unsigned long BitIndex;
-	return _BitScanReverse(&BitIndex, Value) ? BitIndex : 0;
-}
-
-unsigned int FWorkerComputeSocket::CountLeadingZeros(unsigned int Value)
-{
-	// return 32 if value is zero
-	unsigned long BitIndex;
-	_BitScanReverse64(&BitIndex, (unsigned long long)(Value) * 2 + 1);
-	return 32 - BitIndex;
+	return Len + EncodedLen;
 }
 
 //////////////////////////////////////////////////////
 
-enum class FRemoteComputeSocket::EControlMessageType
+class FRemoteComputeSocket : public FComputeSocket
 {
-	Attach = -1,
-	Detach = -2,
-};
-
-struct FRemoteComputeSocket::FFrameHeader
-{
-	int Channel;
-	int Size;
-};
+public:
+	enum class EControlMessageType
+	{
+		Detach = -2,
+	};
 
-struct FRemoteComputeSocket::FDetail
-{
-	class FCriticalSectionLock
+	struct FFrameHeader
 	{
-	public:
-		FCriticalSectionLock(CRITICAL_SECTION* InCriticalSection) : CriticalSection(InCriticalSection)
-		{
-			EnterCriticalSection(CriticalSection);
-		}
+		int Channel;
+		int Size;
+	};
 
-		FCriticalSectionLock(const FCriticalSectionLock& Other) = delete;
+	std::unique_ptr<FComputeTransport> Transport;
+	const EComputeSocketEndpoint Endpoint;
+	std::mutex CriticalSection;
 
-		~FCriticalSectionLock()
-		{
-			LeaveCriticalSection(CriticalSection);
-		}
+	std::thread RecvThread;
 
-	private:
-		CRITICAL_SECTION* const CriticalSection;
-	};
+	std::unordered_map<int, FComputeBufferWriter> Writers;
+	std::vector<FComputeBufferReader> Readers;
+	std::unordered_map<int, std::thread> SendThreads;
 
-	class FSendThread
+	FRemoteComputeSocket(std::unique_ptr<FComputeTransport> InTransport, EComputeSocketEndpoint InEndpoint)
+		: Transport(std::move(InTransport))
+		, Endpoint(InEndpoint)
+		, CriticalSection()
+		, RecvThread(&FRemoteComputeSocket::RecvThreadProc, this)
 	{
-	public:
-		FSendThread(FDetail& InDetail, int InChannel, FComputeBufferReader& InReader)
-			: Detail(InDetail)
-			, Channel(InChannel)
-			, Reader(InReader)
-			, ThreadHandle(nullptr)
-		{
-		}
-
-		~FSendThread()
-		{
-			Stop();
-		}
-
-		void Start()
-		{
-			Stop();
-			ThreadHandle = CreateThread(nullptr, 0, &SendThreadProc, this, 0, nullptr);
-		}
+	}
 
-		void Stop()
+	~FRemoteComputeSocket()
+	{
+		for (FComputeBufferReader& Reader : Readers)
 		{
-			if (ThreadHandle != nullptr)
-			{
-				WaitForSingleObject(ThreadHandle, INFINITE);
-				ThreadHandle = nullptr;
-			}
+			Reader.ForceComplete();
 		}
 
-	private:
-		FDetail& Detail;
-		const int Channel;
-		FComputeBufferReader& Reader;
-		HANDLE ThreadHandle;
-
-		static DWORD __stdcall SendThreadProc(void* Param)
+		for (std::pair<const int, std::thread>& Pair : SendThreads)
 		{
-			FSendThread* SendBufferInfo = (FSendThread*)Param;
-			FRemoteComputeSocket& Socket = SendBufferInfo->Detail.Socket;
-			FComputeBufferReader& Reader = SendBufferInfo->Reader;
-
-			FFrameHeader Header;
-			Header.Channel = SendBufferInfo->Channel;
-
-			const unsigned char* Data;
-			while ((Data = Reader.WaitToRead(1)) != nullptr)
-			{
-				FCriticalSectionLock Lock(&SendBufferInfo->Detail.CriticalSection);
-				Header.Size = (int)Reader.GetMaxReadSize();
-				Socket.SendFull(&Header, sizeof(Header));
-				Socket.SendFull(Data, Header.Size);
-			}
-
-			if (Reader.IsComplete())
-			{
-				FCriticalSectionLock Lock(&SendBufferInfo->Detail.CriticalSection);
-				Header.Size = (int)EControlMessageType::Detach;
-				Socket.SendFull(&Header, sizeof(Header));
-			}
-
-			delete SendBufferInfo;
-			return 0;
+			Pair.second.join();
 		}
-	};
 
-	FRemoteComputeSocket& Socket;
-	CRITICAL_SECTION CriticalSection;
-
-	HANDLE RecvThreadHandle;
-
-	std::unordered_map<int, FComputeBufferWriter*> Writers;
-	std::unordered_map<int, std::unique_ptr<FSendThread>> SendThreads;
-
-	std::unordered_set<int> AttachedRemoteBuffers;
-	std::unordered_map<int, HANDLE> AttachEvents;
-
-	FDetail(FRemoteComputeSocket& InSocket)
-		: Socket(InSocket)
-	{
-		InitializeCriticalSection(&CriticalSection);
-		RecvThreadHandle = CreateThread(nullptr, 0, &RecvThread, this, 0, nullptr);
+		Transport->Close();
+		RecvThread.join();
 	}
 
-	~FDetail()
+	void RecvThreadProc()
 	{
-		WaitForSingleObject(RecvThreadHandle, INFINITE);
-		DeleteCriticalSection(&CriticalSection);
-	}
-
-	static DWORD _stdcall RecvThread(void* Param)
-	{
-		FDetail& Detail = *(FDetail*)Param;
-		FRemoteComputeSocket& Socket = Detail.Socket;
-
-		std::unordered_map<int, FComputeBufferWriter*> CachedWriters;
+		std::unordered_map<int, FComputeBufferWriter> CachedWriters;
 
 		// Process messages from the remote
 		FFrameHeader Header;
-		while (Socket.RecvFull(&Header, sizeof(Header)))
+		while (Transport->RecvMessage(&Header, sizeof(Header)))
 		{
 			if (Header.Size >= 0)
 			{
-				Detail.ReadFrame(CachedWriters, Header.Channel, Header.Size);
-			}
-			else if (Header.Size == (int)EControlMessageType::Attach)
-			{
-				Detail.AttachRemoteRecvBuffer(Header.Channel);
+				ReadFrame(CachedWriters, Header.Channel, Header.Size);
 			}
 			else if (Header.Size == (int)EControlMessageType::Detach)
 			{
-				Detail.DetachRecvBuffer(CachedWriters, Header.Channel);
+				DetachRecvBuffer(CachedWriters, Header.Channel);
 			}
 			else
 			{
 				assert(false);
 			}
 		}
-		return 0;
 	}
 
-	bool ReadFrame(std::unordered_map<int, FComputeBufferWriter*>& CachedWriters, int Channel, int Size)
+	void SendThreadProc(int Channel, FComputeBufferReader Reader)
+	{
+		FFrameHeader Header;
+		Header.Channel = Channel;
+
+		const unsigned char* Data;
+		while ((Data = Reader.WaitToRead(1)) != nullptr)
+		{
+			std::lock_guard<std::mutex> Lock(CriticalSection);
+			Header.Size = (int)Reader.GetMaxReadSize();
+			Transport->SendMessage(&Header, sizeof(Header));
+			Transport->SendMessage(Data, Header.Size);
+			Reader.AdvanceReadPosition(Header.Size);
+		}
+
+		if (Reader.IsComplete())
+		{
+			std::lock_guard<std::mutex> Lock(CriticalSection);
+			Header.Size = (int)EControlMessageType::Detach;
+			Transport->SendMessage(&Header, sizeof(Header));
+		}
+	}
+
+	bool ReadFrame(std::unordered_map<int, FComputeBufferWriter>& CachedWriters, int Channel, int Size)
 	{
-		std::unordered_map<int, FComputeBufferWriter*>::iterator Iter = CachedWriters.find(Channel);
+		std::unordered_map<int, FComputeBufferWriter>::iterator Iter = CachedWriters.find(Channel);
 		if (Iter == CachedWriters.end())
 		{
-			FCriticalSectionLock Lock(&CriticalSection);
+			std::lock_guard<std::mutex> Lock(CriticalSection);
 
 			Iter = Writers.find(Channel);
 			if (Iter == Writers.end())
@@ -393,122 +344,47 @@ struct FRemoteComputeSocket::FDetail
 			Iter = CachedWriters.insert(*Iter).first;
 		}
 
-		FComputeBufferWriter* Writer = Iter->second;
+		FComputeBufferWriter& Writer = Iter->second;
 
-		unsigned char* Data = Writer->WaitToWrite(Size);
-		if (!Socket.RecvFull(Data, Size))
+		unsigned char* Data = Writer.WaitToWrite(Size);
+		if (!Transport->RecvMessage(Data, Size))
 		{
 			return false;
 		}
 
-		Writer->AdvanceWritePosition(Size);
+		Writer.AdvanceWritePosition(Size);
 		return true;
 	}
 
-	void AttachRecvBuffer(int ChannelId, FComputeBufferWriter& Writer)
+	void AttachRecvBuffer(int ChannelId, FComputeBufferWriter Writer)
 	{
-		FCriticalSectionLock Lock(&CriticalSection);
-		Writers.insert(std::pair<int, FComputeBufferWriter*>(ChannelId, &Writer));
-
-		FFrameHeader Header;
-		Header.Channel = ChannelId;
-		Header.Size = (int)EControlMessageType::Attach;
-		Socket.SendFull(&Header, sizeof(Header));
+		std::lock_guard<std::mutex> Lock(CriticalSection);
+		Writers.insert(std::pair<int, FComputeBufferWriter>(ChannelId, std::move(Writer)));
 	}
 
-	void AttachSendBuffer(int ChannelId, FComputeBufferReader& Reader)
+	void AttachSendBuffer(int ChannelId, FComputeBufferReader Reader)
 	{
-		FCriticalSectionLock Lock(&CriticalSection);
-		SendThreads.insert(std::make_pair(ChannelId, std::make_unique<FSendThread>(*this, ChannelId, Reader)));
+		std::lock_guard<std::mutex> Lock(CriticalSection);
+		Readers.push_back(Reader);
+		SendThreads.insert(std::make_pair(ChannelId, std::thread(&FRemoteComputeSocket::SendThreadProc, this, ChannelId, std::move(Reader))));
 	}
 
-	void AttachRemoteRecvBuffer(int Channel)
-	{
-		FCriticalSectionLock Lock(&CriticalSection);
-		AttachedRemoteBuffers.insert(Channel);
-
-		std::unordered_map<int, HANDLE>::iterator AttachEventIter = AttachEvents.find(Channel);
-		if (AttachEventIter != AttachEvents.end())
-		{
-			SetEvent(AttachEventIter->second);
-			AttachEvents.erase(AttachEventIter);
-		}
-	}
-
-	void DetachRecvBuffer(std::unordered_map<int, FComputeBufferWriter*>& CachedWriters, int Channel)
+	void DetachRecvBuffer(std::unordered_map<int, FComputeBufferWriter>& CachedWriters, int Channel)
 	{
 		CachedWriters.erase(Channel);
 
-		FCriticalSectionLock Lock(&CriticalSection);
+		std::lock_guard<std::mutex> Lock(CriticalSection);
 
-		std::unordered_map<int, FComputeBufferWriter*>::iterator Iter = Writers.find(Channel);
+		std::unordered_map<int, FComputeBufferWriter>::iterator Iter = Writers.find(Channel);
 		if (Iter != Writers.end())
 		{
-			Iter->second->MarkComplete();
+			Iter->second.MarkComplete();
 			Writers.erase(Iter);
 		}
 	}
 };
 
-FRemoteComputeSocket::FRemoteComputeSocket()
-	: Detail(nullptr)
-{
-}
-
-FRemoteComputeSocket::~FRemoteComputeSocket()
-{
-	delete Detail;
-}
-
-void FRemoteComputeSocket::Start()
-{
-	if (Detail == nullptr)
-	{
-		Detail = new FDetail(*this);
-	}
-}
-
-void FRemoteComputeSocket::AttachRecvBuffer(int ChannelId, FComputeBufferWriter Writer)
+std::unique_ptr<FComputeSocket> CreateComputeSocket(std::unique_ptr<FComputeTransport> Transport, EComputeSocketEndpoint Endpoint)
 {
-	Detail->AttachRecvBuffer(ChannelId, Writer);
+	return std::unique_ptr<FComputeSocket>(new FRemoteComputeSocket(std::move(Transport), Endpoint));
 }
-
-void FRemoteComputeSocket::AttachSendBuffer(int ChannelId, FComputeBufferReader Reader)
-{
-	Detail->AttachSendBuffer(ChannelId, Reader);
-}
-
-bool FRemoteComputeSocket::SendFull(const void* Data, size_t Size)
-{
-	const unsigned char* RemainingData = (const unsigned char*)Data;
-	for (size_t RemainingSize = Size; RemainingSize > 0; )
-	{
-		size_t SentSize = Send(RemainingData, RemainingSize);
-		if (SentSize == 0)
-		{
-			return false;
-		}
-
-		RemainingData += SentSize;
-		RemainingSize -= SentSize;
-	}
-	return true;
-}
-
-bool FRemoteComputeSocket::RecvFull(void* Data, size_t Size)
-{
-	unsigned char* RemainingData = (unsigned char*)Data;
-	for (size_t RemainingSize = Size; RemainingSize > 0; )
-	{
-		size_t RecvSize = Recv(RemainingData, RemainingSize);
-		if (RecvSize == 0)
-		{
-			return false;
-		}
-
-		RemainingData += RecvSize;
-		RemainingSize -= RecvSize;
-	}
-	return true;
-}
-
diff --git a/Engine/Source/Programs/Horde/Samples/RemoteWorkerCpp/ComputeSocket.h b/Engine/Source/Programs/Horde/Samples/RemoteWorkerCpp/ComputeSocket.h
--- a/Engine/Source/Programs/Horde/Samples/RemoteWorkerCpp/ComputeSocket.h
+++ b/Engine/Source/Programs/Horde/Samples/RemoteWorkerCpp/ComputeSocket.h
@@ -4,6 +4,7 @@
 
 #include "ComputeBuffer.h"
 #include "ComputeChannel.h"
+#include "ComputeTransport.h"
 
 //
 // Connection to a remote machine that multiplexes data into and out-of multiple buffers
@@ -12,13 +13,23 @@
 class FComputeSocket
 {
 public:
+	FComputeSocket();
 	virtual ~FComputeSocket();
 
+	FComputeSocket(const FComputeSocket&) = delete;
+	FComputeSocket& operator=(const FComputeSocket&) = delete;
+
 	// Attaches a new buffer for receiving data
 	virtual void AttachRecvBuffer(int ChannelId, FComputeBufferWriter Writer) = 0;
 
 	// Attaches a new buffer for sending data */
 	virtual void AttachSendBuffer(int ChannelId, FComputeBufferReader Reader) = 0;
+
+	// Attaches a channel to this socket
+	FComputeChannel CreateChannel(int ChannelId);
+
+	// Attaches a channel to this socket
+	FComputeChannel CreateChannel(int ChannelId, FComputeBuffer RecvBuffer, FComputeBuffer SendBuffer);
 };
 
 //
@@ -28,7 +39,7 @@ public:
 class FWorkerComputeSocket final : public FComputeSocket
 {
 public:
-	static const wchar_t* const EnvVarName;
+	static const wchar_t* const IpcEnvVar;
 
 	FWorkerComputeSocket();
 	~FWorkerComputeSocket();
@@ -63,43 +74,19 @@ private:
 
 	static size_t WriteVarUInt(unsigned char* Pos, unsigned int Value);
 	static size_t WriteString(unsigned char* Pos, const wchar_t* Text);
-
-	static unsigned int FloorLog2(unsigned int Value);
-	static unsigned int CountLeadingZeros(unsigned int Value);
 };
 
 //
-// Implementation of FComputeSocket that communicates via an overridable transport mechanism.
+// Enum identifying which end of the socket a particular machine is
 //
-class FRemoteComputeSocket : public FComputeSocket
+enum class EComputeSocketEndpoint
 {
-public:
-	FRemoteComputeSocket();
-	~FRemoteComputeSocket();
-
-	// Attaches a new buffer for receiving data
-	virtual void AttachRecvBuffer(int ChannelId, FComputeBufferWriter Writer) override final;
-
-	// Attaches a new buffer for sending data */
-	virtual void AttachSendBuffer(int ChannelId, FComputeBufferReader Reader) override final;
-
-protected:
-	// Starts reading and dispatching data in a background thread
-	void Start();
-
-	// Sends data to the remote
-	virtual size_t Send(const void* Data, size_t Size) = 0;
-
-	// Receives data from the remote
-	virtual size_t Recv(void* Data, size_t Size) = 0;
+	// The initiating machine
+	Local,
 
-private:
-	struct FDetail;
-	struct FFrameHeader;
-	enum class EControlMessageType;
-
-	bool SendFull(const void* Data, size_t Size);
-	bool RecvFull(void* Data, size_t Size);
-
-	FDetail* Detail;
+	// The remote machine
+	Remote
 };
+
+// Creates a socket using a custom transport
+std::unique_ptr<FComputeSocket> CreateComputeSocket(std::unique_ptr<FComputeTransport> Transport, EComputeSocketEndpoint Endpoint);
diff --git a/Engine/Source/Programs/Horde/Samples/RemoteWorkerCpp/ComputeTransport.cpp b/Engine/Source/Programs/Horde/Samples/RemoteWorkerCpp/ComputeTransport.cpp
new file mode 100644
--- /dev/null
+++ b/Engine/Source/Programs/Horde/Samples/RemoteWorkerCpp/ComputeTransport.cpp
@@ -0,0 +1,80 @@
+// Copyright Epic Games, Inc. All Rights Reserved.
+
+#include "ComputeTransport.h"
+
+bool FComputeTransport::SendMessage(const void* Data, size_t Size)
+{
+	const unsigned char* RemainingData = (const unsigned char*)Data;
+	for (size_t RemainingSize = Size; RemainingSize > 0; )
+	{
+		size_t SentSize = Send(RemainingData, RemainingSize);
+		if (SentSize == 0)
+		{
+			return false;
+		}
+
+		RemainingData += SentSize;
+		RemainingSize -= SentSize;
+	}
+	return true;
+}
+
+bool FComputeTransport::RecvMessage(void* Data, size_t Size)
+{
+	unsigned char* RemainingData = (unsigned char*)Data;
+	for (size_t RemainingSize = Size; RemainingSize > 0; )
+	{
+		size_t RecvSize = Recv(RemainingData, RemainingSize);
+		if (RecvSize == 0)
+		{
+			return false;
+		}
+
+		RemainingData += RecvSize;
+		RemainingSize -= RecvSize;
+	}
+	return true;
+}
+
+///////////////////////////////////
+
+FBufferTransport::FBufferTransport(FComputeBuffer InSendBuffer, FComputeBuffer InRecvBuffer)
+	: SendBuffer(std::move(InSendBuffer))
+	, RecvBuffer(std::move(InRecvBuffer))
+{
+}
+
+size_t FBufferTransport::Send(const void* Data, size_t Size)
+{
+	FComputeBufferWriter& Writer = SendBuffer.GetWriter();
+	unsigned char* Buffer = Writer.WaitToWrite(1);
+
+	size_t WriteSize = std::min(Size, SendBuffer.GetWriter().GetMaxWriteSize());
+	memcpy(Buffer, Data, WriteSize);
+	Writer.AdvanceWritePosition(WriteSize);
+
+	return WriteSize;
+}
+
+size_t FBufferTransport::Recv(void* Data, size_t Size)
+{
+	FComputeBufferReader& Reader = RecvBuffer.GetReader();
+	const unsigned char* Buffer = Reader.WaitToRead(1);
+
+	size_t ReadSize = std::min(Size, Reader.GetMaxReadSize());
+	memcpy(Data, Buffer, ReadSize);
+	Reader.AdvanceReadPosition(ReadSize);
+
+	return ReadSize;
+}
+
+void FBufferTransport::MarkComplete()
+{
+	SendBuffer.GetWriter().MarkComplete();
+}
+
+void FBufferTransport::Close()
+{
+	SendBuffer.GetWriter().MarkComplete();
+	RecvBuffer.GetWriter().MarkComplete();
+}
diff --git a/Engine/Source/Programs/Horde/Samples/RemoteWorkerCpp/ComputeTransport.h b/Engine/Source/Programs/Horde/Samples/RemoteWorkerCpp/ComputeTransport.h
new file mode 100644
--- /dev/null
+++ b/Engine/Source/Programs/Horde/Samples/RemoteWorkerCpp/ComputeTransport.h
@@ -0,0 +1,49 @@
+// Copyright Epic Games, Inc. All Rights Reserved.
+
+#pragma once
+
+#include "ComputeBuffer.h"
+
+//
+// Interface for a transport mechanism that can be used by a compute socket
+//
+class FComputeTransport
+{
+public:
+	// Sends data to the remote
+	virtual size_t Send(const void* Data, size_t Size) = 0;
+
+	// Receives data from the remote
+	virtual size_t Recv(void* Data, size_t Size) = 0;
+
+	// Indicates to the remote that no more data will be sent.
+	virtual void MarkComplete() = 0;
+
+	// Indicates that no more data will be sent or received, and that any blocking reads/writes should stop.
+	virtual void Close() = 0;
+
+	// Sends data to the remote, blocking until the entire message has been sent.
+	bool SendMessage(const void* Data, size_t Size);
+
+	// Receives a fixed length block of data from the remote, blocking until the entire length has been received.
+	bool RecvMessage(void* Data, size_t Size);
+};
+
+//
+// Implementation of FComputeTransport which uses in-memory buffers to transport data
+//
+class FBufferTransport final : public FComputeTransport
+{
+public:
+	FBufferTransport(FComputeBuffer InSendBuffer, FComputeBuffer InRecvBuffer);
+
+protected:
+	virtual size_t Send(const void* Data, size_t Size) override final;
+	virtual size_t Recv(void* Data, size_t Size) override final;
+	virtual void MarkComplete() override final;
+	virtual void Close() override final;
+
+private:
+	FComputeBuffer SendBuffer;
+	FComputeBuffer RecvBuffer;
+};
diff --git a/Engine/Source/Programs/Horde/Samples/RemoteWorkerCpp/Main.cpp b/Engine/Source/Programs/Horde/Samples/RemoteWorkerCpp/Main.cpp
--- a/Engine/Source/Programs/Horde/Samples/RemoteWorkerCpp/Main.cpp
+++ b/Engine/Source/Programs/Horde/Samples/RemoteWorkerCpp/Main.cpp
@@ -1,14 +1,21 @@
 // Copyright Epic Games, Inc. All Rights Reserved.
 
-#include <iostream>
 #include "ComputeChannel.h"
 #include "ComputeBuffer.h"
 #include "ComputeSocket.h"
+#include <iostream>
+#include <thread>
+#include <assert.h>
 
-#include <Windows.h>
-
-int main()
+int main(int argc, const char* argv[])
 {
+	if (argc >= 2 && !_stricmp(argv[1], "-Test"))
+	{
+		void RunTests();
+		RunTests();
+		return 0;
+	}
+
 	const int ChannelId = 100;
 
 	FWorkerComputeSocket Socket;
@@ -18,8 +25,8 @@ int main()
 		return 1;
 	}
 
-	FComputeChannel Channel;
-	if (!Channel.Attach(Socket, ChannelId, FComputeBuffer::FParams()))
+	FComputeChannel Channel = Socket.CreateChannel(ChannelId);
+	if(!Channel.IsValid())
 	{
 		std::cout << "Unable to create channel to initiator" << std::endl;
 		return 1;
@@ -53,4 +60,3 @@ int main()
 	Channel.MarkComplete();
 	return 0;
 }
-
diff --git a/Engine/Source/Programs/Horde/Samples/RemoteWorkerCpp/Tests.cpp b/Engine/Source/Programs/Horde/Samples/RemoteWorkerCpp/Tests.cpp
new file mode 100644
--- /dev/null
+++ b/Engine/Source/Programs/Horde/Samples/RemoteWorkerCpp/Tests.cpp
@@ -0,0 +1,74 @@
+// Copyright Epic Games, Inc. All Rights Reserved.
+
+#include "ComputeBuffer.h"
+#include "ComputeSocket.h"
+#include <assert.h>
+#include <iostream>
+
+#define verify(x) ((x) || (throw #x, true))
+
+void RunTests()
+{
+	void ComputeBufferTest();
+	ComputeBufferTest();
+
+	void ComputeSocketTest();
+	ComputeSocketTest();
+}
+
+void ComputeBufferTest()
+{
+	char TestData[] = "hello world!";
+
+	FComputeBuffer Buffer;
+	verify(Buffer.CreateNew(FComputeBuffer::FParams()));
+
+	FComputeBufferWriter& Writer = Buffer.GetWriter();
+	unsigned char* WriteBuffer = Writer.WaitToWrite(sizeof(TestData));
+	memcpy(WriteBuffer, TestData, sizeof(TestData));
+	Writer.AdvanceWritePosition(sizeof(TestData));
+
+	FComputeBufferReader& Reader = Buffer.GetReader();
+	const unsigned char* ReadBuffer = Reader.WaitToRead(sizeof(TestData));
+	verify(memcmp(ReadBuffer, TestData, sizeof(TestData)) == 0);
+}
+
+template<size_t TestDataSize> void CheckChannelSendRecv(FComputeChannel& SendChannel, FComputeChannel& RecvChannel, const char(&TestData)[TestDataSize])
+{
+	SendChannel.Send(TestData, TestDataSize);
+
+	char RecvTestData[TestDataSize];
+	verify(RecvChannel.Recv(RecvTestData, TestDataSize) == TestDataSize);
+
+	verify(memcmp(TestData, RecvTestData, TestDataSize) == 0);
+}
+
+void ComputeSocketTest()
+{
+	// Buffers for transferring between client and server
+	FComputeBuffer ClientToServerBuffer;
+	verify(ClientToServerBuffer.CreateNew(FComputeBuffer::FParams()));
+
+	FComputeBuffer ServerToClientBuffer;
+	verify(ServerToClientBuffer.CreateNew(FComputeBuffer::FParams()));
+
+	// Client transport
+	std::unique_ptr<FComputeSocket> ClientSocket = CreateComputeSocket(std::make_unique<FBufferTransport>(ClientToServerBuffer, ServerToClientBuffer), EComputeSocketEndpoint::Local);
+	FComputeChannel ClientChannel1 = ClientSocket->CreateChannel(1);
+	FComputeChannel ClientChannel2 = ClientSocket->CreateChannel(2);
+
+	// Server socket
+	std::unique_ptr<FComputeSocket> ServerSocket = CreateComputeSocket(std::make_unique<FBufferTransport>(ServerToClientBuffer, ClientToServerBuffer), EComputeSocketEndpoint::Remote);
+	FComputeChannel ServerChannel1 = ServerSocket->CreateChannel(1);
+	FComputeChannel ServerChannel2 = ServerSocket->CreateChannel(2);
+
+	// Send data over channel 1
+	CheckChannelSendRecv(ClientChannel1, ServerChannel1, "Channel 1: Client -> Server");
+
+	CheckChannelSendRecv(ServerChannel2, ClientChannel2, "Channel 2: Server -> Client");
+	CheckChannelSendRecv(ClientChannel2, ServerChannel2, "Channel 2: Client -> Server");
+
+	// Send data over channel 1 again
+	CheckChannelSendRecv(ClientChannel1, ServerChannel1, "Channel 1: Client -> Server (2)");
+	CheckChannelSendRecv(ServerChannel1, ClientChannel1, "Channel 1: Server -> Client (2)");
+}
diff --git a/Engine/Source/Programs/Shared/EpicGames.Horde.Tests/TransportTests.cs b/Engine/Source/Programs/Shared/EpicGames.Horde.Tests/TransportTests.cs
--- a/Engine/Source/Programs/Shared/EpicGames.Horde.Tests/TransportTests.cs
+++ b/Engine/Source/Programs/Shared/EpicGames.Horde.Tests/TransportTests.cs
@@ -25,7 +25,7 @@ namespace EpicGames.Horde.Tests
 			{
 				using MemoryStream memoryStream = new MemoryStream();
 				StreamTransport streamTransport = new StreamTransport(memoryStream);
-				await streamTransport.WriteAsync(input, CancellationToken.None);
+				await streamTransport.SendAsync(input, CancellationToken.None);
 				payload = memoryStream.ToArray();
 			}
 
@@ -62,7 +62,7 @@ namespace EpicGames.Horde.Tests
 				using MemoryStream memoryStream = new MemoryStream();
 				StreamTransport streamTransport = new StreamTransport(memoryStream);
 				await using AesTransport aesTransport = new AesTransport(streamTransport, key, nonce);
-				await aesTransport.WriteAsync(input, CancellationToken.None);
+				await aesTransport.SendAsync(input, CancellationToken.None);
 				encrypted = memoryStream.ToArray();
 			}
 
@@ -74,7 +74,7 @@ namespace EpicGames.Horde.Tests
 
 				for (int offset = 0; offset < output.Length;)
 				{
-					offset += await aesTransport.ReadPartialAsync(output.AsMemory(offset), CancellationToken.None);
+					offset += await aesTransport.RecvAsync(output.AsMemory(offset), CancellationToken.None);
 				}
 			}
 
@@ -89,7 +89,7 @@ namespace EpicGames.Horde.Tests
 		{
 			for (int offset = 0; offset < buffer.Length;)
 			{
-				int size = await transport.ReadPartialAsync(buffer.Slice(offset), CancellationToken.None);
+				int size = await transport.RecvAsync(buffer.Slice(offset), CancellationToken.None);
 				Assert.IsTrue(size > 0);
 				offset += size;
 			}
diff --git a/Engine/Source/Programs/Shared/EpicGames.Horde/Compute/ComputeChannel.cs b/Engine/Source/Programs/Shared/EpicGames.Horde/Compute/ComputeChannel.cs
--- a/Engine/Source/Programs/Shared/EpicGames.Horde/Compute/ComputeChannel.cs
+++ b/Engine/Source/Programs/Shared/EpicGames.Horde/Compute/ComputeChannel.cs
@@ -4,7 +4,6 @@ using System;
 using System.IO;
 using System.Threading;
 using System.Threading.Tasks;
-using EpicGames.Horde.Compute.Buffers;
 
 namespace EpicGames.Horde.Compute
 {
@@ -38,8 +37,6 @@ namespace EpicGames.Horde.Compute
 		public void Dispose()
 		{
 			Reader.Dispose();
-
-			Writer.MarkComplete();
 			Writer.Dispose();
 		}
 
@@ -55,16 +52,16 @@ namespace EpicGames.Horde.Compute
 		/// </summary>
 		/// <param name="buffer">Buffer to receive the data</param>
 		/// <param name="cancellationToken">Cancellation token for the operation</param>
-		public ValueTask<int> ReceiveAsync(Memory<byte> buffer, CancellationToken cancellationToken = default) => Reader.ReadAsync(buffer, cancellationToken);
+		public ValueTask<int> RecvAsync(Memory<byte> buffer, CancellationToken cancellationToken = default) => Reader.ReadAsync(buffer, cancellationToken);
 
 		/// <summary>
 		/// Reads a complete message from the given socket, retrying reads until the buffer is full.
 		/// </summary>
 		/// <param name="buffer">Buffer to store the data</param>
 		/// <param name="cancellationToken">Cancellation token for the operation</param>
-		public async ValueTask ReceiveMessageAsync(Memory<byte> buffer, CancellationToken cancellationToken = default)
+		public async ValueTask RecvMessageAsync(Memory<byte> buffer, CancellationToken cancellationToken = default)
 		{
-			if (!await TryReceiveMessageAsync(buffer, cancellationToken))
+			if (!await TryRecvMessageAsync(buffer, cancellationToken))
 			{
 				throw new EndOfStreamException();
 			}
@@ -75,11 +72,11 @@ namespace EpicGames.Horde.Compute
 		/// </summary>
 		/// <param name="buffer">Buffer to store the data</param>
 		/// <param name="cancellationToken">Cancellation token for the operation</param>
-		public async ValueTask<bool> TryReceiveMessageAsync(Memory<byte> buffer, CancellationToken cancellationToken = default)
+		public async ValueTask<bool> TryRecvMessageAsync(Memory<byte> buffer, CancellationToken cancellationToken = default)
 		{
 			for (int offset = 0; offset < buffer.Length;)
 			{
-				int read = await ReceiveAsync(buffer.Slice(offset), cancellationToken);
+				int read = await RecvAsync(buffer.Slice(offset), cancellationToken);
 				if (read == 0)
 				{
 					return false;
diff --git a/Engine/Source/Programs/Shared/EpicGames.Horde/Compute/ComputeSocket.cs b/Engine/Source/Programs/Shared/EpicGames.Horde/Compute/ComputeSocket.cs
--- a/Engine/Source/Programs/Shared/EpicGames.Horde/Compute/ComputeSocket.cs
+++ b/Engine/Source/Programs/Shared/EpicGames.Horde/Compute/ComputeSocket.cs
@@ -94,7 +94,16 @@ namespace EpicGames.Horde.Compute
 				throw new InvalidOperationException($"Environment variable {IpcEnvVar} is not defined; cannot connect as worker.");
 			}
 
-			SharedMemoryBuffer commandBuffer = SharedMemoryBuffer.OpenExisting(baseName);
+			return Open(baseName);
+		}
+
+		/// <summary>
+		/// Opens a socket which allows a worker to communicate with the Horde Agent
+		/// </summary>
+		/// <param name="commandBufferName">Name of the command buffer</param>
+		public static WorkerComputeSocket Open(string commandBufferName)
+		{
+			SharedMemoryBuffer commandBuffer = SharedMemoryBuffer.OpenExisting(commandBufferName);
 			return new WorkerComputeSocket(commandBuffer);
 		}
 
@@ -247,7 +256,7 @@ namespace EpicGames.Horde.Compute
 				for (; ; )
 				{
 					// Read the next packet header
-					if (!await transport.ReadOptionalAsync(header, cancellationToken))
+					if (!await transport.RecvOptionalAsync(header, cancellationToken))
 					{
 						_logger.LogTrace("[{Tag}] End of socket", Tag);
 						break;
@@ -303,7 +312,7 @@ namespace EpicGames.Horde.Compute
 
 			for (int offset = 0; offset < size;)
 			{
-				int read = await transport.ReadPartialAsync(memory.Slice(offset, size - offset), cancellationToken);
+				int read = await transport.RecvAsync(memory.Slice(offset, size - offset), cancellationToken);
 				offset += read;
 			}
 
@@ -369,7 +378,7 @@ namespace EpicGames.Horde.Compute
 				_bodySegment.Set(memory, null, _header.Length);
 
 				ReadOnlySequence<byte> sequence = new ReadOnlySequence<byte>(_headerSegment, 0, _bodySegment, memory.Length);
-				await _transport.WriteAsync(sequence, cancellationToken);
+				await _transport.SendAsync(sequence, cancellationToken);
 			}
 			finally
 			{
diff --git a/Engine/Source/Programs/Shared/EpicGames.Horde/Compute/ComputeTransport.cs b/Engine/Source/Programs/Shared/EpicGames.Horde/Compute/ComputeTransport.cs
--- a/Engine/Source/Programs/Shared/EpicGames.Horde/Compute/ComputeTransport.cs
+++ b/Engine/Source/Programs/Shared/EpicGames.Horde/Compute/ComputeTransport.cs
@@ -14,21 +14,21 @@ namespace EpicGames.Horde.Compute
 	public abstract class ComputeTransport
 	{
 		/// <summary>
-		/// Reads data from the underlying transport into an output buffer
+		/// Writes data to the underlying transport
 		/// </summary>
-		/// <param name="buffer">Buffer to read into</param>
+		/// <param name="buffer">Buffer to be written</param>
 		/// <param name="cancellationToken">Cancellation token for the operation</param>
-		public abstract ValueTask<int> ReadPartialAsync(Memory<byte> buffer, CancellationToken cancellationToken);
+		public abstract ValueTask SendAsync(ReadOnlySequence<byte> buffer, CancellationToken cancellationToken);
 
 		/// <summary>
-		/// Writes data to the underlying transport
+		/// Reads data from the underlying transport into an output buffer
 		/// </summary>
-		/// <param name="buffer">Buffer to be written</param>
+		/// <param name="buffer">Buffer to read into</param>
 		/// <param name="cancellationToken">Cancellation token for the operation</param>
-		public abstract ValueTask WriteAsync(ReadOnlySequence<byte> buffer, CancellationToken cancellationToken);
+		public abstract ValueTask<int> RecvAsync(Memory<byte> buffer, CancellationToken cancellationToken);
 
 		/// <summary>
-		/// Indicate that all data has been written to the transport layer, and that there will be no more calls to <see cref="WriteAsync(ReadOnlySequence{Byte}, CancellationToken)"/>
+		/// Indicate that all data has been written to the transport layer, and that there will be no more calls to <see cref="SendAsync(ReadOnlySequence{Byte}, CancellationToken)"/>
 		/// </summary>
 		/// <param name="cancellationToken">Cancellation token for the operation</param>
 		public abstract ValueTask MarkCompleteAsync(CancellationToken cancellationToken);
@@ -38,12 +38,12 @@ namespace EpicGames.Horde.Compute
 		/// </summary>
 		/// <param name="buffer">Buffer to read into</param>
 		/// <param name="cancellationToken">Cancellation token for the operation</param>
-		public async ValueTask ReadFullAsync(Memory<byte> buffer, CancellationToken cancellationToken)
+		public async ValueTask RecvFullAsync(Memory<byte> buffer, CancellationToken cancellationToken)
 		{
 			int read = 0;
 			while (read < buffer.Length)
 			{
-				int partialRead = await ReadPartialAsync(buffer.Slice(read, buffer.Length - read), cancellationToken);
+				int partialRead = await RecvAsync(buffer.Slice(read, buffer.Length - read), cancellationToken);
 				if (partialRead == 0)
 				{
 					throw new EndOfStreamException();
@@ -57,16 +57,16 @@ namespace EpicGames.Horde.Compute
 		/// </summary>
 		/// <param name="buffer">Buffer to read into</param>
 		/// <param name="cancellationToken">Cancellation token for the operation</param>
-		public async ValueTask<bool> ReadOptionalAsync(Memory<byte> buffer, CancellationToken cancellationToken)
+		public async ValueTask<bool> RecvOptionalAsync(Memory<byte> buffer, CancellationToken cancellationToken)
 		{
-			int read = await ReadPartialAsync(buffer, cancellationToken);
+			int read = await RecvAsync(buffer, cancellationToken);
 			if (read == 0)
 			{
 				return false;
 			}
 			if (read < buffer.Length)
 			{
-				await ReadFullAsync(buffer.Slice(read), cancellationToken);
+				await RecvFullAsync(buffer.Slice(read), cancellationToken);
 			}
 			return true;
 		}
@@ -76,8 +76,8 @@ namespace EpicGames.Horde.Compute
 		/// </summary>
 		/// <param name="buffer">Buffer to be written</param>
 		/// <param name="cancellationToken">Cancellation token for the operation</param>
-		public ValueTask WriteAsync(ReadOnlyMemory<byte> buffer, CancellationToken cancellationToken)
-			=> WriteAsync(new ReadOnlySequence<byte>(buffer), cancellationToken);
+		public ValueTask SendAsync(ReadOnlyMemory<byte> buffer, CancellationToken cancellationToken)
+			=> SendAsync(new ReadOnlySequence<byte>(buffer), cancellationToken);
 	}
 }
 
diff --git a/Engine/Source/Programs/Shared/EpicGames.Horde/Compute/Transports/AesTransport.cs b/Engine/Source/Programs/Shared/EpicGames.Horde/Compute/Transports/AesTransport.cs
--- a/Engine/Source/Programs/Shared/EpicGames.Horde/Compute/Transports/AesTransport.cs
+++ b/Engine/Source/Programs/Shared/EpicGames.Horde/Compute/Transports/AesTransport.cs
@@ -100,7 +100,7 @@ namespace EpicGames.Horde.Compute.Transports
 			for (; ; )
 			{
 				// Read more data into the buffer
-				int read = await _inner.ReadPartialAsync(memory.Slice(size), cancellationToken);
+				int read = await _inner.RecvAsync(memory.Slice(size), cancellationToken);
 				if (read == 0)
 				{
 					break;
@@ -143,7 +143,7 @@ namespace EpicGames.Horde.Compute.Transports
 		}
 
 		/// <inheritdoc/>
-		public override async ValueTask<int> ReadPartialAsync(Memory<byte> buffer, CancellationToken cancellationToken)
+		public override async ValueTask<int> RecvAsync(Memory<byte> buffer, CancellationToken cancellationToken)
 		{
 			int sizeRead = 0;
 			while (sizeRead == 0)
@@ -167,7 +167,7 @@ namespace EpicGames.Horde.Compute.Transports
 		}
 
 		/// <inheritdoc/>
-		public override async ValueTask WriteAsync(ReadOnlySequence<byte> buffer, CancellationToken cancellationToken)
+		public override async ValueTask SendAsync(ReadOnlySequence<byte> buffer, CancellationToken cancellationToken)
 		{
 			while (buffer.Length > 0)
 			{
@@ -185,7 +185,7 @@ namespace EpicGames.Horde.Compute.Transports
 				Position += plainText.Length;
 
 				await _lastWriteTask;
-				_lastWriteTask = _inner.WriteAsync(writeBuffer, cancellationToken).AsTask();
+				_lastWriteTask = _inner.SendAsync(writeBuffer, cancellationToken).AsTask();
 
 				buffer = buffer.Slice(plainText.Length);
 				_writeBufferOffset ^= WritePacketSize;
diff --git a/Engine/Source/Programs/Shared/EpicGames.Horde/Compute/Transports/PipeTransport.cs b/Engine/Source/Programs/Shared/EpicGames.Horde/Compute/Transports/PipeTransport.cs
--- a/Engine/Source/Programs/Shared/EpicGames.Horde/Compute/Transports/PipeTransport.cs
+++ b/Engine/Source/Programs/Shared/EpicGames.Horde/Compute/Transports/PipeTransport.cs
@@ -35,7 +35,7 @@ namespace EpicGames.Horde.Compute.Transports
 		}
 
 		/// <inheritdoc/>
-		public override async ValueTask<int> ReadPartialAsync(Memory<byte> buffer, CancellationToken cancellationToken)
+		public override async ValueTask<int> RecvAsync(Memory<byte> buffer, CancellationToken cancellationToken)
 		{
 			int sizeRead = 0;
 			while (sizeRead < buffer.Length)
@@ -78,7 +78,7 @@ namespace EpicGames.Horde.Compute.Transports
 		}
 
 		/// <inheritdoc/>
-		public override async ValueTask WriteAsync(ReadOnlySequence<byte> buffer, CancellationToken cancellationToken)
+		public override async ValueTask SendAsync(ReadOnlySequence<byte> buffer, CancellationToken cancellationToken)
 		{
 			foreach (ReadOnlyMemory<byte> memory in buffer)
 			{
diff --git a/Engine/Source/Programs/Shared/EpicGames.Horde/Compute/Transports/StreamTransport.cs b/Engine/Source/Programs/Shared/EpicGames.Horde/Compute/Transports/StreamTransport.cs
--- a/Engine/Source/Programs/Shared/EpicGames.Horde/Compute/Transports/StreamTransport.cs
+++ b/Engine/Source/Programs/Shared/EpicGames.Horde/Compute/Transports/StreamTransport.cs
@@ -25,7 +25,7 @@ namespace EpicGames.Horde.Compute.Transports
 		public StreamTransport(Stream stream) => _stream = stream;
 
 		/// <inheritdoc/>
-		public override async ValueTask<int> ReadPartialAsync(Memory<byte> buffer, CancellationToken cancellationToken)
+		public override async ValueTask<int> RecvAsync(Memory<byte> buffer, CancellationToken cancellationToken)
 		{
 			int length = await _stream.ReadAsync(buffer, cancellationToken);
 			Position += length;
@@ -33,7 +33,7 @@ namespace EpicGames.Horde.Compute.Transports
 		}
 
 		/// <inheritdoc/>
-		public override async ValueTask WriteAsync(ReadOnlySequence<byte> buffer, CancellationToken cancellationToken)
+		public override async ValueTask SendAsync(ReadOnlySequence<byte> buffer, CancellationToken cancellationToken)
 		{
 			foreach (ReadOnlyMemory<byte> memory in buffer)
 			{
diff --git a/Engine/Source/Programs/Shared/EpicGames.Horde/Compute/Transports/TcpTransport.cs b/Engine/Source/Programs/Shared/EpicGames.Horde/Compute/Transports/TcpTransport.cs
--- a/Engine/Source/Programs/Shared/EpicGames.Horde/Compute/Transports/TcpTransport.cs
+++ b/Engine/Source/Programs/Shared/EpicGames.Horde/Compute/Transports/TcpTransport.cs
@@ -26,7 +26,7 @@ namespace EpicGames.Horde.Compute.Transports
 		public TcpTransport(Socket socket) => _socket = socket;
 
 		/// <inheritdoc/>
-		public override async ValueTask<int> ReadPartialAsync(Memory<byte> buffer, CancellationToken cancellationToken)
+		public override async ValueTask<int> RecvAsync(Memory<byte> buffer, CancellationToken cancellationToken)
 		{
 			int read = await _socket.ReceiveAsync(buffer, SocketFlags.None, cancellationToken);
 			Position += read;
@@ -34,7 +34,7 @@ namespace EpicGames.Horde.Compute.Transports
 		}
 
 		/// <inheritdoc/>
-		public override async ValueTask WriteAsync(ReadOnlySequence<byte> buffer, CancellationToken cancellationToken)
+		public override async ValueTask SendAsync(ReadOnlySequence<byte> buffer, CancellationToken cancellationToken)
 		{
 			foreach (ReadOnlyMemory<byte> memory in buffer)
 			{
