Horde: Add missing copyright notice.

diff --git a/Engine/Source/Programs/Horde/Samples/RemoteWorkerCpp/ComputePlatform.cpp b/Engine/Source/Programs/Horde/Samples/RemoteWorkerCpp/ComputePlatform.cpp
--- a/Engine/Source/Programs/Horde/Samples/RemoteWorkerCpp/ComputePlatform.cpp
+++ b/Engine/Source/Programs/Horde/Samples/RemoteWorkerCpp/ComputePlatform.cpp
@@ -1,3 +1,5 @@
+// Copyright Epic Games, Inc. All Rights Reserved.
+
 #include "ComputePlatform.h"
 #include <assert.h>
 #include <wchar.h>
