Horde: Fix issue with detecting running bisection task jobs
#rnx

diff --git a/Engine/Source/Programs/Horde/Horde.Server/Jobs/JobCollection.cs b/Engine/Source/Programs/Horde/Horde.Server/Jobs/JobCollection.cs
--- a/Engine/Source/Programs/Horde/Horde.Server/Jobs/JobCollection.cs
+++ b/Engine/Source/Programs/Horde/Horde.Server/Jobs/JobCollection.cs
@@ -594,20 +594,18 @@ namespace Horde.Server.Jobs
 		public async IAsyncEnumerable<IJob> FindBisectTaskJobsAsync(BisectTaskId bisectTaskId, bool? running, [EnumeratorCancellation] CancellationToken cancellationToken)
 		{
 			FilterDefinition<JobDocument> filter = Builders<JobDocument>.Filter.Eq(x => x.StartedByBisectTaskId, bisectTaskId);
-			if (running.HasValue)
+
+			await foreach (JobDocument jobDoc in _jobs.Find(filter).ToAsyncEnumerable(cancellationToken))
 			{
-				if (running.Value)
-				{
-					filter &= Builders<JobDocument>.Filter.Gt(x => x.SchedulePriority, 0);
-				}
-				else
+				if (running.HasValue && running.Value)
 				{
-					filter &= Builders<JobDocument>.Filter.Lte(x => x.SchedulePriority, 0);
+					JobState state = jobDoc.GetState();
+					if (state == JobState.Complete)
+					{
+						continue;
+					}					
 				}
-			}
 
-			await foreach (JobDocument jobDoc in _jobs.Find(filter).ToAsyncEnumerable(cancellationToken))
-			{
 				await PostLoadAsync(jobDoc);
 				yield return jobDoc;
 			}
