Horde: Add bitfields to C++ compute buffer states, to help with debugging.

diff --git a/Engine/Source/Programs/Horde/Samples/RemoteWorkerCpp/ComputeBuffer.cpp b/Engine/Source/Programs/Horde/Samples/RemoteWorkerCpp/ComputeBuffer.cpp
--- a/Engine/Source/Programs/Horde/Samples/RemoteWorkerCpp/ComputeBuffer.cpp
+++ b/Engine/Source/Programs/Horde/Samples/RemoteWorkerCpp/ComputeBuffer.cpp
@@ -20,9 +20,9 @@ struct FComputeBufferDetail
 		Complete = 0b11,
 	};
 
-	struct FChunkState
+	union FChunkState
 	{
-		const long long Value;
+		long long Value;
 
 		FChunkState(long long InValue)
 			: Value(InValue)
@@ -46,6 +46,14 @@ struct FComputeBufferDetail
 
 		// Test whether a particular reader is still referencing the chunk
 		bool HasReaderFlag(int ReaderIdx) const { return (Value & (1ULL << (31 + ReaderIdx))) != 0; }
+
+	private:
+		struct // For debugging purposes only; non-portable assumption of bitfield layout
+		{
+			unsigned long long Length : 31;
+			unsigned long long ReaderFlags : 31;
+			EWriteState WriteState : 2;
+		};
 	};
 
 	struct FChunkStatePtr
@@ -54,61 +62,59 @@ struct FComputeBufferDetail
 		// Read the state value from memory
 		FChunkState Get() const
 		{
-			return FChunkState(FComputePlatform::AtomicRead64((volatile long long*)&Value));
+			return FChunkState(FComputePlatform::AtomicRead64(&State.Value));
 		}
 
 		// Append data to the chunk
-		void Append(long long length)
+		void Append(long long Length)
 		{
-			FComputePlatform::AtomicAdd64(&Value, length);
+			FComputePlatform::AtomicAdd64(&State.Value, Length);
 		}
 
 		// Mark the chunk as being written to
 		void StartWriting(int numReaders)
 		{
-			FComputePlatform::AtomicWrite64(&Value, FChunkState(EWriteState::Writing, (1 << numReaders) - 1, 0).Value);
+			FComputePlatform::AtomicWrite64(&State.Value, FChunkState(EWriteState::Writing, (1 << numReaders) - 1, 0).Value);
 		}
 
 		// Mark this chunk as the end of the stream
 		void MarkComplete()
 		{
-			FComputePlatform::AtomicOr64(&Value, FChunkState(EWriteState::Complete, 0, 0).Value);
+			FComputePlatform::AtomicOr64(&State.Value, FChunkState(EWriteState::Complete, 0, 0).Value);
 		}
 
 		// Clear the reader flag
 		void FinishReading(int ReaderIdx)
 		{
-			FComputePlatform::AtomicAnd64(&Value, ~FChunkState((EWriteState)0, 1 << ReaderIdx, 0).Value);
+			FComputePlatform::AtomicAnd64(&State.Value, ~FChunkState((EWriteState)0, 1 << ReaderIdx, 0).Value);
 		}
 
 		// Move to the next chunk
 		void FinishWriting()
 		{
-			FComputePlatform::AtomicAnd64(&Value, ~FChunkState(EWriteState::Writing, 0, 0).Value);
+			FComputePlatform::AtomicAnd64(&State.Value, ~FChunkState(EWriteState::Writing, 0, 0).Value);
 		}
 
 	private:
-		union
-		{
-			volatile long long Value;
-			struct
-			{
-				unsigned long long Length : 31;
-				unsigned long long ReaderFlags : 31;
-				EWriteState WriteState : 2;
-			};
-		};
+		volatile FChunkState State;
 	};
 
-	struct FReaderState
+	union FReaderState
 	{
-		const long long Value;
+		long long Value;
 
 		FReaderState(long long InValue) : Value(InValue) { }
 		FReaderState(int chunkIdx, int offset) : Value(((unsigned long long)chunkIdx << 32) | offset) { }
 
 		int GetChunkIdx() const { return (int)(Value >> 32); }
 		int GetOffset() const { return (int)(unsigned int)Value; }
+
+	private:
+		struct // For debugging purposes only; non-portable assumption of bitfield layout
+		{
+			unsigned int Offset : 32;
+			unsigned int ChunkIdx : 32;
+		};
 	};
 
 	struct FReaderStatePtr
@@ -116,29 +122,21 @@ struct FComputeBufferDetail
 	public:
 		FReaderState Get() const
 		{
-			return FReaderState(FComputePlatform::AtomicRead64((volatile long long*)&Value));
+			return FReaderState(FComputePlatform::AtomicRead64((volatile long long*)&State.Value));
 		}
 
 		void Set(FReaderState State)
 		{
-			FComputePlatform::AtomicWrite64(&Value, State.Value);
+			FComputePlatform::AtomicWrite64(&State.Value, State.Value);
 		}
 
 		void Advance(size_t Length)
 		{
-			FComputePlatform::AtomicAdd64(&Value, Length);
+			FComputePlatform::AtomicAdd64(&State.Value, Length);
 		}
 
 	private:
-		union
-		{
-			volatile long long Value;
-			struct
-			{
-				int Offset;
-				int ChunkIdx;
-			};
-		};
+		volatile FReaderState State;
 	};
 
 	struct FHeader
@@ -167,8 +165,8 @@ struct FComputeBufferDetail
 		, ChunkPtrs{ nullptr, }
 		, RefCount(1)
 	{
-		static_assert(sizeof(FChunkStatePtr) == sizeof(long long));
-		static_assert(sizeof(FReaderStatePtr) == sizeof(long long));
+		static_assert(sizeof(FChunkStatePtr) == sizeof(long long), "Incorrect size of FChunkStatePtr; check union is declared correctly.");
+		static_assert(sizeof(FReaderStatePtr) == sizeof(long long), "Incorrect size of FReaderStatePtr; check union is declared correctly.");
 
 		wcscpy_s(this->Name, Name);
 	}
@@ -296,15 +294,15 @@ struct FComputeBufferDetail
 
 	size_t GetMaxReadSize(int ReaderIdx) const
 	{
-		const FReaderStatePtr& readerStatePtr = Header->Readers[ReaderIdx];
-		FReaderState readerState = readerStatePtr.Get();
+		const FReaderStatePtr& ReaderStatePtr = Header->Readers[ReaderIdx];
+		FReaderState ReaderState = ReaderStatePtr.Get();
 
-		const FChunkStatePtr& chunkStatePtr = Header->Chunks[readerState.GetChunkIdx()];
-		FChunkState chunkState = chunkStatePtr.Get();
+		const FChunkStatePtr& ChunkStatePtr = Header->Chunks[ReaderState.GetChunkIdx()];
+		FChunkState ChunkState = ChunkStatePtr.Get();
 
-		if (chunkState.HasReaderFlag(ReaderIdx))
+		if (ChunkState.HasReaderFlag(ReaderIdx))
 		{
-			return chunkState.GetLength() - readerState.GetOffset();
+			return ChunkState.GetLength() - ReaderState.GetOffset();
 		}
 		else
 		{
@@ -316,53 +314,53 @@ struct FComputeBufferDetail
 	{
 		for (; ; )
 		{
-			FReaderStatePtr& readerStatePtr = Header->Readers[ReaderIdx];
-			FReaderState readerState = readerStatePtr.Get();
+			FReaderStatePtr& ReaderStatePtr = Header->Readers[ReaderIdx];
+			FReaderState ReaderState = ReaderStatePtr.Get();
 
-			FChunkStatePtr& chunkStatePtr = Header->Chunks[readerState.GetChunkIdx()];
-			FChunkState chunkState = chunkStatePtr.Get();
+			FChunkStatePtr& ChunkStatePtr = Header->Chunks[ReaderState.GetChunkIdx()];
+			FChunkState ChunkState = ChunkStatePtr.Get();
 
-			if (!chunkState.HasReaderFlag(ReaderIdx))
+			if (!ChunkState.HasReaderFlag(ReaderIdx))
 			{
 				// Wait until the current chunk is readable
 				ReaderEvents[ReaderIdx].Reset();
-				if (!chunkState.HasReaderFlag(ReaderIdx) && !ReaderEvents[ReaderIdx].Wait(TimeoutMs))
+				if (!ChunkState.HasReaderFlag(ReaderIdx) && !ReaderEvents[ReaderIdx].Wait(TimeoutMs))
 				{
 					return nullptr;
 				}
 			}
-			else if (readerState.GetOffset() + MinSize <= chunkState.GetLength())
+			else if (ReaderState.GetOffset() + MinSize <= ChunkState.GetLength())
 			{
 				// We have enough data in the chunk to be able to read a message
-				return ChunkPtrs[readerState.GetChunkIdx()] + readerState.GetOffset();
+				return ChunkPtrs[ReaderState.GetChunkIdx()] + ReaderState.GetOffset();
 			}
-			else if (chunkState.GetWriteState() == EWriteState::Writing)
+			else if (ChunkState.GetWriteState() == EWriteState::Writing)
 			{
 				// Wait until there is more data in the chunk
 				ReaderEvents[ReaderIdx].Reset();
-				if (Header->Chunks[readerState.GetChunkIdx()].Get().Value == chunkState.Value && !ReaderEvents[ReaderIdx].Wait(TimeoutMs))
+				if (Header->Chunks[ReaderState.GetChunkIdx()].Get().Value == ChunkState.Value && !ReaderEvents[ReaderIdx].Wait(TimeoutMs))
 				{
 					return nullptr;
 				}
 			}
-			else if (readerState.GetOffset() < chunkState.GetLength() || chunkState.GetWriteState() == EWriteState::Complete)
+			else if (ReaderState.GetOffset() < ChunkState.GetLength() || ChunkState.GetWriteState() == EWriteState::Complete)
 			{
 				// Cannot read the requested amount of data from this chunk.
 				return nullptr;
 			}
-			else if (chunkState.GetWriteState() == EWriteState::MovedToNext)
+			else if (ChunkState.GetWriteState() == EWriteState::MovedToNext)
 			{
 				// Move to the next chunk
-				chunkStatePtr.FinishReading(ReaderIdx);
+				ChunkStatePtr.FinishReading(ReaderIdx);
 				WriterEvent.Set();
 
-				int chunkIdx = readerState.GetChunkIdx() + 1;
+				int chunkIdx = ReaderState.GetChunkIdx() + 1;
 				if (chunkIdx == Header->NumChunks)
 				{
 					chunkIdx = 0;
 				}
 
-				readerStatePtr.Set(FReaderState(chunkIdx, 0));
+				ReaderStatePtr.Set(FReaderState(chunkIdx, 0));
 			}
 			else
 			{
@@ -421,7 +419,7 @@ struct FComputeBufferDetail
 
 		for (; ; )
 		{
-			FChunkStatePtr WriteChunkStatePtr = Header->Chunks[Header->WriteChunkIdx];
+			FChunkStatePtr& WriteChunkStatePtr = Header->Chunks[Header->WriteChunkIdx];
 
 			FChunkState ChunkState = WriteChunkStatePtr.Get();
 			if (ChunkState.GetWriteState() == EWriteState::Writing)
@@ -447,7 +445,7 @@ struct FComputeBufferDetail
 				NextWriteChunkIdx = 0;
 			}
 
-			FChunkStatePtr NextWriteChunkStatePtr = Header->Chunks[NextWriteChunkIdx];
+			FChunkStatePtr& NextWriteChunkStatePtr = Header->Chunks[NextWriteChunkIdx];
 			while (NextWriteChunkStatePtr.Get().GetReaderFlags() != 0)
 			{
 				if (!WriterEvent.Wait(TimeoutMs))
diff --git a/Engine/Source/Programs/Horde/Samples/RemoteWorkerCpp/ComputePlatform.cpp b/Engine/Source/Programs/Horde/Samples/RemoteWorkerCpp/ComputePlatform.cpp
--- a/Engine/Source/Programs/Horde/Samples/RemoteWorkerCpp/ComputePlatform.cpp
+++ b/Engine/Source/Programs/Horde/Samples/RemoteWorkerCpp/ComputePlatform.cpp
@@ -189,9 +189,9 @@ size_t FComputePlatform::WcharToUtf8(const wchar_t* Source, size_t SourceLen, ch
 	return (size_t)Result;
 }
 
-long long FComputePlatform::AtomicRead64(volatile long long* Ptr)
+long long FComputePlatform::AtomicRead64(const volatile long long* Ptr)
 {
-	return InterlockedCompareExchange64(Ptr, 0, 0);
+	return InterlockedCompareExchange64(const_cast<volatile long long*>(Ptr), 0, 0);
 }
 
 void FComputePlatform::AtomicWrite64(volatile long long* Ptr, long long Value)
diff --git a/Engine/Source/Programs/Horde/Samples/RemoteWorkerCpp/ComputePlatform.h b/Engine/Source/Programs/Horde/Samples/RemoteWorkerCpp/ComputePlatform.h
--- a/Engine/Source/Programs/Horde/Samples/RemoteWorkerCpp/ComputePlatform.h
+++ b/Engine/Source/Programs/Horde/Samples/RemoteWorkerCpp/ComputePlatform.h
@@ -99,7 +99,7 @@ struct FComputePlatform
 	//
 
 	// Reads a 64-bit value from the given memory location 
-	static long long AtomicRead64(volatile long long* Ptr);
+	static long long AtomicRead64(const volatile long long* Ptr);
 
 	// Writes a 64-bit value to the given memory location
 	static void AtomicWrite64(volatile long long* Ptr, long long Value);
