Horde: Adding query to find bisections by job, user id, etc
#rnx none

diff --git a/Engine/Source/Programs/Horde/Horde.Server.Tests/BisectTests.cs b/Engine/Source/Programs/Horde/Horde.Server.Tests/BisectTests.cs
--- a/Engine/Source/Programs/Horde/Horde.Server.Tests/BisectTests.cs
+++ b/Engine/Source/Programs/Horde/Horde.Server.Tests/BisectTests.cs
@@ -65,6 +65,10 @@ namespace Horde.Server.Tests
 			Assert.AreEqual(14, jobs[0].Change);
 			await SetJobOutcome(jobs[0], graph, JobStepOutcome.Failure);
 
+			List<GetBisectTaskResponse> bisectTasks = Deref(await BisectTasksController!.FindBisectTasksAsync(bisectTask.Owner.Id));
+			Assert.AreEqual(1, bisectTasks.Count);
+			Assert.AreEqual(bisectTask.Owner.Id, bisectTasks[0].Owner.Id);
+
 			await Clock.AdvanceAsync(TimeSpan.FromMinutes(30));
 
 			bisectTask = Deref(await BisectTasksController!.GetAsync(task!.BisectTaskId));
diff --git a/Engine/Source/Programs/Horde/Horde.Server/Jobs/Bisect/BisectController.cs b/Engine/Source/Programs/Horde/Horde.Server/Jobs/Bisect/BisectController.cs
--- a/Engine/Source/Programs/Horde/Horde.Server/Jobs/Bisect/BisectController.cs
+++ b/Engine/Source/Programs/Horde/Horde.Server/Jobs/Bisect/BisectController.cs
@@ -14,7 +14,6 @@ using Horde.Server.Users;
 using Horde.Server.Utilities;
 using HordeCommon;
 using Microsoft.AspNetCore.Authorization;
-using Microsoft.AspNetCore.Components.Web;
 using Microsoft.AspNetCore.Mvc;
 using Microsoft.Extensions.Options;
 
@@ -290,7 +289,7 @@ namespace Horde.Server.Jobs.Bisect
 		[Route("/api/v1/bisect/{bisectTaskId}")]
 		public async Task<ActionResult> UpdateAsync([FromRoute] BisectTaskId bisectTaskId, [FromBody] UpdateBisectTaskRequest request, CancellationToken cancellationToken)
 		{
-			for(; ;)
+			for (; ; )
 			{
 				IBisectTask? bisectTask = await _bisectTaskCollection.GetAsync(bisectTaskId, cancellationToken);
 				if (bisectTask == null)
@@ -324,8 +323,8 @@ namespace Horde.Server.Jobs.Bisect
 						IJob? existingJob = await _jobCollection.FindBisectTaskJobsAsync(bisectTask.Id, true, cancellationToken).FirstOrDefaultAsync(cancellationToken);
 						if (existingJob != null && existingJob.AbortedByUserId == null)
 						{
-							await _jobService.UpdateJobAsync(existingJob, null, null, null,  User.GetUserId() ?? KnownUsers.System, null, null, null);
-						}						
+							await _jobService.UpdateJobAsync(existingJob, null, null, null, User.GetUserId() ?? KnownUsers.System, null, null, null);
+						}
 					}
 					return Ok();
 				}
@@ -348,7 +347,7 @@ namespace Horde.Server.Jobs.Bisect
 				return NotFound(jobId);
 			}
 
-			IReadOnlyList<IBisectTask> tasks = await _bisectTaskCollection.FindAsync(jobId, cancellationToken);
+			IReadOnlyList<IBisectTask> tasks = await _bisectTaskCollection.FindAsync(jobId, null, null, null, null, null, cancellationToken);
 
 			List<GetBisectTaskResponse> response = new List<GetBisectTaskResponse>();
 
@@ -370,11 +369,55 @@ namespace Horde.Server.Jobs.Bisect
 				List<IJobStepRef> steps = await _jobStepRefs.GetStepsForNodeAsync(job.StreamId, job.TemplateId, bisectTask.NodeName, null, true, 1024, bisectTask.Id, cancellationToken);
 
 				response.Add(new GetBisectTaskResponse(bisectTask, new GetThinUserInfoResponse(user), steps));
-				
+
 			}
 
 			return response;
 		}
-	}
 
+		/// <summary>
+		/// 
+		/// </summary>
+		/// <returns></returns>
+		[HttpGet]
+		[Route("/api/v1/bisect")]
+		[ProducesResponseType(typeof(List<GetBisectTaskResponse>), 200)]
+		public async Task<ActionResult<List<GetBisectTaskResponse>>> FindBisectTasksAsync(
+		[FromQuery] string? ownerId = null,
+		[FromQuery] string? jobId = null,
+		[FromQuery] DateTimeOffset? minCreateTime = null,
+		[FromQuery] DateTimeOffset? maxCreateTime = null,
+		[FromQuery] int index = 0,
+		[FromQuery] int count = 100,
+		CancellationToken cancellationToken = default)
+		{
+			List<GetBisectTaskResponse> responses = new List<GetBisectTaskResponse>();
+
+			JobId? JobIdValue = !String.IsNullOrEmpty(jobId) ? JobId.Parse(jobId) : null;
+			UserId? OwnerIdValue = !String.IsNullOrEmpty(ownerId) ? UserId.Parse(ownerId) : null;
+
+			IReadOnlyList<IBisectTask> tasks = await _bisectTaskCollection.FindAsync(JobIdValue, OwnerIdValue, minCreateTime?.UtcDateTime, maxCreateTime?.UtcDateTime, index, count, cancellationToken);
+
+			if (tasks.Count == 0) 
+			{
+				return responses;
+			}
+
+			List<IJob> jobs = await _jobCollection.FindAsync(new HashSet<JobId>(tasks.Select(x => x.InitialJobId)).ToArray());
+
+			for (int i = 0; i < tasks.Count; i++)
+			{
+				IBisectTask task = tasks[i];
+				IUser? user = await _userCollection.GetUserAsync(task.OwnerId);
+				IJob? job = jobs.FirstOrDefault(x => x.Id == task.InitialJobId);
+				if (job != null && user != null)
+				{
+					List<IJobStepRef> steps = await _jobStepRefs.GetStepsForNodeAsync(job.StreamId, job.TemplateId, task.NodeName, null, true, 1024,task.Id, cancellationToken);
+					responses.Add(new GetBisectTaskResponse(task, new GetThinUserInfoResponse(user), steps));
+				}
+			}
+
+			return responses;
+		}
+	}
 }
diff --git a/Engine/Source/Programs/Horde/Horde.Server/Jobs/Bisect/BisectTaskCollection.cs b/Engine/Source/Programs/Horde/Horde.Server/Jobs/Bisect/BisectTaskCollection.cs
--- a/Engine/Source/Programs/Horde/Horde.Server/Jobs/Bisect/BisectTaskCollection.cs
+++ b/Engine/Source/Programs/Horde/Horde.Server/Jobs/Bisect/BisectTaskCollection.cs
@@ -11,6 +11,7 @@ using Horde.Server.Streams;
 using Horde.Server.Users;
 using Horde.Server.Utilities;
 using HordeCommon;
+using MongoDB.Bson;
 using MongoDB.Bson.Serialization.Attributes;
 using MongoDB.Driver;
 
@@ -144,15 +145,37 @@ namespace Horde.Server.Jobs.Bisect
 		}
 
 		/// <inheritdoc/>
-		public async Task<IReadOnlyList<IBisectTask>> FindAsync(JobId? jobId = null, CancellationToken cancellationToken = default)
+		public async Task<IReadOnlyList<IBisectTask>> FindAsync(JobId? jobId = null, UserId? ownerId = null, DateTime? minCreateTime = null, DateTime? maxCreateTime = null, int? index = null, int? count = null, CancellationToken cancellationToken = default)
 		{
 			// Find all the bisection tasks matching the given criteria
 			FilterDefinitionBuilder<BisectTaskDoc> filterBuilder = Builders<BisectTaskDoc>.Filter;
 
 			FilterDefinition<BisectTaskDoc> filter = FilterDefinition<BisectTaskDoc>.Empty;
-			filter &= filterBuilder.Eq(x => x.InitialJobId, jobId);
 
-			List<BisectTaskDoc> steps = await _bisectTasks.Find(filter).SortByDescending(x => x.InitialChange).ToListAsync(cancellationToken);
+			if (jobId != null)
+			{
+				filter &= filterBuilder.Eq(x => x.InitialJobId, jobId);
+			}
+
+			if (ownerId != null)
+			{
+				filter &= filterBuilder.Eq(x => x.OwnerId, ownerId);
+			}
+
+			if (minCreateTime != null)
+			{
+				BisectTaskId minTime = new BisectTaskId(ObjectId.GenerateNewId(minCreateTime.Value));
+				filter &= filterBuilder.Gte(x => x.Id!, minTime);
+
+			}
+			if (maxCreateTime != null)
+			{
+				BisectTaskId maxTime = new BisectTaskId(ObjectId.GenerateNewId(maxCreateTime.Value));
+				filter &= filterBuilder.Lte(x => x.Id!, maxTime);
+			}
+
+
+			List<BisectTaskDoc> steps = await _bisectTasks.Find(filter).SortByDescending(x => x.InitialChange).Range(index, count).ToListAsync(cancellationToken);
 			return steps.ConvertAll<IBisectTask>(x => x);
 		}
 
diff --git a/Engine/Source/Programs/Horde/Horde.Server/Jobs/Bisect/IBisectTaskCollection.cs b/Engine/Source/Programs/Horde/Horde.Server/Jobs/Bisect/IBisectTaskCollection.cs
--- a/Engine/Source/Programs/Horde/Horde.Server/Jobs/Bisect/IBisectTaskCollection.cs
+++ b/Engine/Source/Programs/Horde/Horde.Server/Jobs/Bisect/IBisectTaskCollection.cs
@@ -1,5 +1,6 @@
 // Copyright Epic Games, Inc. All Rights Reserved.
 
+using System;
 using System.Collections.Generic;
 using System.Threading;
 using System.Threading.Tasks;
@@ -37,9 +38,14 @@ namespace Horde.Server.Jobs.Bisect
 		/// Finds all bisect tasks for the provided criteria
 		/// </summary>
 		/// <param name="jobId">jobId of initial bisection</param>
+		/// <param name="ownerId">ownerId of bisection</param>
+		/// <param name="minCreateTime">min creation time of the bisection task</param>
+		/// <param name="maxCreateTime">max creation time of the bisection task</param>
+		/// <param name="index"></param>
+		/// <param name="count"></param>
 		/// <param name="cancellationToken">Cancellation token for the operation</param>
 		/// <returns>List of bisect tasks</returns>
-		public Task<IReadOnlyList<IBisectTask>> FindAsync(JobId? jobId = null, CancellationToken cancellationToken = default);
+		public Task<IReadOnlyList<IBisectTask>> FindAsync(JobId? jobId = null, UserId? ownerId = null, DateTime? minCreateTime = null, DateTime? maxCreateTime = null, int? index = null, int? count = null, CancellationToken cancellationToken = default);
 
 		/// <summary>
 		/// Gets a bisect task by id
