Horde: C++ Remote Execution improvements.

- Multiple readers can now be allocated from a single buffer dynamically. Readers are ref-counted separately to the underlying buffer, and are released once their ref count reaches zero.
- Detaching readers is now supported independently to marking the buffer as complete.

diff --git a/Engine/Source/Programs/Horde/Samples/RemoteWorkerCpp/ComputeBuffer.cpp b/Engine/Source/Programs/Horde/Samples/RemoteWorkerCpp/ComputeBuffer.cpp
--- a/Engine/Source/Programs/Horde/Samples/RemoteWorkerCpp/ComputeBuffer.cpp
+++ b/Engine/Source/Programs/Horde/Samples/RemoteWorkerCpp/ComputeBuffer.cpp
@@ -65,16 +65,22 @@ struct FComputeBufferDetail
 			return FChunkState(FComputePlatform::AtomicRead64(&State.Value));
 		}
 
-		// Append data to the chunk
-		void Append(long long Length)
+		// Read the state value from memory
+		void Set(FChunkState NewState)
 		{
-			FComputePlatform::AtomicAdd64(&State.Value, Length);
+			FComputePlatform::AtomicWrite64(&State.Value, NewState.Value);
 		}
 
-		// Mark the chunk as being written to
-		void StartWriting(int numReaders)
+		// Attempt to update the chunk state
+		bool TryUpdate(FChunkState PrevState, FChunkState NextState)
 		{
-			FComputePlatform::AtomicWrite64(&State.Value, FChunkState(EWriteState::Writing, (1 << numReaders) - 1, 0).Value);
+			return FComputePlatform::AtomicCompareExchange64(&State.Value, NextState.Value, PrevState.Value);
+		}
+
+		// Append data to the chunk
+		void Append(long long Length)
+		{
+			FComputePlatform::AtomicAdd64(&State.Value, Length);
 		}
 
 		// Mark this chunk as the end of the stream
@@ -83,6 +89,12 @@ struct FComputeBufferDetail
 			FComputePlatform::AtomicOr64(&State.Value, FChunkState(EWriteState::Complete, 0, 0).Value);
 		}
 
+		// Start reading the chunk with the given reader
+		void StartReading(int ReaderIdx)
+		{
+			FComputePlatform::AtomicOr64(&State.Value, FChunkState((EWriteState)0, 1 << ReaderIdx, 0).Value);
+		}
+
 		// Clear the reader flag
 		void FinishReading(int ReaderIdx)
 		{
@@ -99,44 +111,62 @@ struct FComputeBufferDetail
 		volatile FChunkState State;
 	};
 
-	union FReaderState
+	struct FReaderState
+	{
+		const int ReaderIdx;
+
+		int Offset;
+		int ChunkIdx;
+		bool Detached;
+
+		FReaderState(int InReaderIdx)
+			: ReaderIdx(InReaderIdx)
+			, Offset(0)
+			, ChunkIdx(0)
+			, Detached(false)
+		{ }
+	};
+
+	union FWriterState
 	{
 		long long Value;
 
-		FReaderState(long long InValue) : Value(InValue) { }
-		FReaderState(int chunkIdx, int offset) : Value(((unsigned long long)chunkIdx << 32) | offset) { }
+		FWriterState(long long InValue) : Value(InValue) { }
+		FWriterState(int ChunkIdx, int ReaderFlags, bool HasWrapped) : Value(ChunkIdx | ((long long)ReaderFlags << 32) | (HasWrapped? (1ULL << 63) : 0)) { }
 
-		int GetChunkIdx() const { return (int)(Value >> 32); }
-		int GetOffset() const { return (int)(unsigned int)Value; }
+		int GetReaderFlags() const { return ((unsigned long long)Value >> 32) & 0xffff; }
+		int GetChunkIdx() const { return Value & 0x7fffffff; }
+		bool HasWrapped() const { return (Value & (1ULL << 63)) != 0; }
 
 	private:
 		struct // For debugging purposes only; non-portable assumption of bitfield layout
 		{
-			unsigned int Offset : 32;
-			unsigned int ChunkIdx : 32;
+			int ChunkIdx : 32;
+			int ReaderFlags : 31;
+			int Wrapped : 1;
 		};
 	};
 
-	struct FReaderStatePtr
+	struct FWriterStatePtr
 	{
 	public:
-		FReaderState Get() const
+		FWriterState Get() const
 		{
-			return FReaderState(FComputePlatform::AtomicRead64((volatile long long*)&State.Value));
+			return FWriterState(FComputePlatform::AtomicRead64(&State.Value));
 		}
 
-		void Set(FReaderState State)
+		void Set(FWriterState State)
 		{
 			FComputePlatform::AtomicWrite64(&State.Value, State.Value);
 		}
 
-		void Advance(size_t Length)
+		bool TryUpdate(FWriterState PrevValue, FWriterState NextValue)
 		{
-			FComputePlatform::AtomicAdd64(&State.Value, Length);
+			return FComputePlatform::AtomicCompareExchange64(&State.Value, NextValue.Value, PrevValue.Value);
 		}
 
 	private:
-		volatile FReaderState State;
+		volatile FWriterState State;
 	};
 
 	struct FHeader
@@ -144,9 +174,9 @@ struct FComputeBufferDetail
 		int NumReaders;
 		int NumChunks;
 		int ChunkLength;
-		int WriteChunkIdx;
+		long AllocatedReaders;
+		FWriterStatePtr Writer;
 		FChunkStatePtr Chunks[FComputeBuffer::MaxChunks];
-		FReaderStatePtr Readers[FComputeBuffer::MaxReaders];
 	};
 
 	wchar_t Name[260];
@@ -157,16 +187,14 @@ struct FComputeBufferDetail
 	FComputeManualResetEvent WriterEvent;
 	FComputeManualResetEvent ReaderEvents[FComputeBuffer::MaxReaders];
 
-	FComputeBufferReader Reader;
 	FComputeBufferWriter Writer;
 
 	FComputeBufferDetail(const wchar_t* Name)
 		: Header(nullptr)
 		, ChunkPtrs{ nullptr, }
-		, RefCount(1)
 	{
 		static_assert(sizeof(FChunkStatePtr) == sizeof(long long), "Incorrect size of FChunkStatePtr; check union is declared correctly.");
-		static_assert(sizeof(FReaderStatePtr) == sizeof(long long), "Incorrect size of FReaderStatePtr; check union is declared correctly.");
+		static_assert(sizeof(FWriterStatePtr) == sizeof(long long), "Incorrect size of FWriterStatePtr; check union is declared correctly.");
 
 		wcscpy_s(this->Name, Name);
 	}
@@ -208,6 +236,7 @@ struct FComputeBufferDetail
 		Detail->Header->NumReaders = Params.NumReaders;
 		Detail->Header->NumChunks = Params.NumChunks;
 		Detail->Header->ChunkLength = Params.ChunkLength;
+		Detail->Header->Chunks[0].Set(FChunkState(EWriteState::Writing, 0, 0));
 
 		swprintf(NameBuffer, FComputeBuffer::MaxNameLength, L"%s_W", Name);
 		if (!Detail->WriterEvent.Create(NameBuffer))
@@ -224,8 +253,6 @@ struct FComputeBufferDetail
 			}
 		}
 
-		Detail->Header->Chunks[0].StartWriting(Detail->Header->NumReaders);
-
 		InitShared(Detail);
 		return Detail;
 	}
@@ -266,58 +293,87 @@ struct FComputeBufferDetail
 		return Detail;
 	}
 
-	void AddRef()
-	{
-		FComputePlatform::AtomicIncrement(&RefCount);
-	}
-
-	void Release()
+	bool IsComplete(const FReaderState& ReaderState) const
 	{
-		if (FComputePlatform::AtomicDecrement(&RefCount) == 0)
+		if (ReaderState.Detached)
 		{
-			delete this;
+			return true;
 		}
+
+		FChunkState ChunkState = Header->Chunks[ReaderState.ChunkIdx].Get();
+		return ChunkState.GetWriteState() == EWriteState::Complete && ReaderState.Offset == ChunkState.GetLength();
 	}
 
-	bool IsComplete(int ReaderIdx) const
+	int AllocateReader(int ChunkIdx, int Offset)
 	{
-		FReaderState ReaderState = Header->Readers[ReaderIdx].Get();
-		FChunkState ChunkState = Header->Chunks[ReaderState.GetChunkIdx()].Get();
+		for (;;)
+		{
+			int AllocatedReaders = Header->AllocatedReaders;
+			int ReaderFlag = (AllocatedReaders + 1) ^ AllocatedReaders;
 
-		return ChunkState.GetWriteState() == EWriteState::Complete && ReaderState.GetOffset() == ChunkState.GetLength();
+			int ReaderIdx = FComputePlatform::FloorLog2(ReaderFlag);
+			if (ReaderIdx >= Header->NumReaders)
+			{
+				return -1;
+			}
+
+			if (FComputePlatform::AtomicCompareExchange(&Header->AllocatedReaders, AllocatedReaders | ReaderFlag, AllocatedReaders))
+			{
+				for (;;)
+				{
+					FWriterState State = Header->Writer.Get();
+					if (Header->Writer.TryUpdate(State, FWriterState(State.GetChunkIdx(), State.GetReaderFlags() | ReaderFlag, State.HasWrapped())))
+					{
+						for (int WriteChunkIdx = 0; WriteChunkIdx <= State.GetChunkIdx(); WriteChunkIdx++)
+						{
+							Header->Chunks[WriteChunkIdx].StartReading(ReaderIdx);
+						}
+						return ReaderIdx;
+					}
+				}
+			}
+		}
 	}
 
-	void AdvanceReadPosition(int ReaderIdx, size_t Size)
+	void ReleaseReader(int ReaderIdx)
 	{
-		Header->Readers[ReaderIdx].Advance(Size);
+		int ReaderFlag = 1 << ReaderIdx;
+		for (int Idx = 0; Idx < Header->NumChunks; Idx++)
+		{
+			Header->Chunks[Idx].FinishReading(ReaderFlag);
+		}
+		FComputePlatform::AtomicAnd(&Header->AllocatedReaders, (long)~ReaderFlag);
 	}
 
-	size_t GetMaxReadSize(int ReaderIdx) const
+	size_t GetMaxReadSize(const FReaderState& ReaderState) const
 	{
-		const FReaderStatePtr& ReaderStatePtr = Header->Readers[ReaderIdx];
-		FReaderState ReaderState = ReaderStatePtr.Get();
+		if (ReaderState.Detached)
+		{
+			return 0;
+		}
 
-		const FChunkStatePtr& ChunkStatePtr = Header->Chunks[ReaderState.GetChunkIdx()];
+		const FChunkStatePtr& ChunkStatePtr = Header->Chunks[ReaderState.ChunkIdx];
 		FChunkState ChunkState = ChunkStatePtr.Get();
 
-		if (ChunkState.HasReaderFlag(ReaderIdx))
-		{
-			return ChunkState.GetLength() - ReaderState.GetOffset();
-		}
-		else
+		if (!ChunkState.HasReaderFlag(ReaderState.ReaderIdx))
 		{
 			return 0;
 		}
+
+		return ChunkState.GetLength() - ReaderState.Offset;
 	}
 
-	const unsigned char* WaitToRead(int ReaderIdx, size_t MinSize, int TimeoutMs)
+	const unsigned char* WaitToRead(FReaderState& ReaderState, size_t MinSize, int TimeoutMs)
 	{
+		int ReaderIdx = ReaderState.ReaderIdx;
 		for (; ; )
 		{
-			FReaderStatePtr& ReaderStatePtr = Header->Readers[ReaderIdx];
-			FReaderState ReaderState = ReaderStatePtr.Get();
+			if (ReaderState.Detached)
+			{
+				return nullptr;
+			}
 
-			FChunkStatePtr& ChunkStatePtr = Header->Chunks[ReaderState.GetChunkIdx()];
+			FChunkStatePtr& ChunkStatePtr = Header->Chunks[ReaderState.ChunkIdx];
 			FChunkState ChunkState = ChunkStatePtr.Get();
 
 			if (!ChunkState.HasReaderFlag(ReaderIdx))
@@ -329,21 +385,21 @@ struct FComputeBufferDetail
 					return nullptr;
 				}
 			}
-			else if (ReaderState.GetOffset() + MinSize <= ChunkState.GetLength())
+			else if (ReaderState.Offset + MinSize <= ChunkState.GetLength())
 			{
 				// We have enough data in the chunk to be able to read a message
-				return ChunkPtrs[ReaderState.GetChunkIdx()] + ReaderState.GetOffset();
+				return ChunkPtrs[ReaderState.ChunkIdx] + ReaderState.Offset;
 			}
 			else if (ChunkState.GetWriteState() == EWriteState::Writing)
 			{
 				// Wait until there is more data in the chunk
 				ReaderEvents[ReaderIdx].Reset();
-				if (Header->Chunks[ReaderState.GetChunkIdx()].Get().Value == ChunkState.Value && !ReaderEvents[ReaderIdx].Wait(TimeoutMs))
+				if (Header->Chunks[ReaderState.ChunkIdx].Get().Value == ChunkState.Value && !ReaderEvents[ReaderIdx].Wait(TimeoutMs))
 				{
 					return nullptr;
 				}
 			}
-			else if (ReaderState.GetOffset() < ChunkState.GetLength() || ChunkState.GetWriteState() == EWriteState::Complete)
+			else if (ReaderState.Offset < ChunkState.GetLength() || ChunkState.GetWriteState() == EWriteState::Complete)
 			{
 				// Cannot read the requested amount of data from this chunk.
 				return nullptr;
@@ -354,13 +410,12 @@ struct FComputeBufferDetail
 				ChunkStatePtr.FinishReading(ReaderIdx);
 				WriterEvent.Set();
 
-				int chunkIdx = ReaderState.GetChunkIdx() + 1;
-				if (chunkIdx == Header->NumChunks)
+				if (++ReaderState.ChunkIdx == Header->NumChunks)
 				{
-					chunkIdx = 0;
+					ReaderState.ChunkIdx = 0;
 				}
 
-				ReaderStatePtr.Set(FReaderState(chunkIdx, 0));
+				ReaderState.Offset = 0;
 			}
 			else
 			{
@@ -371,19 +426,23 @@ struct FComputeBufferDetail
 
 	void MarkComplete()
 	{
-		Header->Chunks[Header->WriteChunkIdx].MarkComplete();
+		FWriterState State = Header->Writer.Get();
+		Header->Chunks[State.GetChunkIdx()].MarkComplete();
 		SetAllReaderEvents();
 	}
 
 	void AdvanceWritePosition(size_t Size)
 	{
-		Header->Chunks[Header->WriteChunkIdx].Append(Size);
+		FWriterState State = Header->Writer.Get();
+		Header->Chunks[State.GetChunkIdx()].Append(Size);
 		SetAllReaderEvents();
 	}
 
 	size_t GetMaxWriteSize() const
 	{
-		FChunkState Value = Header->Chunks[Header->WriteChunkIdx].Get();
+		FWriterState State = Header->Writer.Get();
+
+		FChunkState Value = Header->Chunks[State.GetChunkIdx()].Get();
 		if (Value.GetWriteState() == EWriteState::Complete)
 		{
 			return 0;
@@ -417,36 +476,44 @@ struct FComputeBufferDetail
 	{
 		assert(MinSize <= Header->ChunkLength);
 
-		for (; ; )
-		{
-			FChunkStatePtr& WriteChunkStatePtr = Header->Chunks[Header->WriteChunkIdx];
+		// Get the current chunk we're writing to
+		FWriterState WriterState = Header->Writer.Get();
+		int WriteChunkIdx = WriterState.GetChunkIdx();
 
-			FChunkState ChunkState = WriteChunkStatePtr.Get();
-			if (ChunkState.GetWriteState() == EWriteState::Writing)
-			{
-				int Length = ChunkState.GetLength();
-				if (Length + MinSize <= Header->ChunkLength)
-				{
-					return ChunkPtrs[Header->WriteChunkIdx] + Length;
-				}
+		FChunkStatePtr& WriteChunkStatePtr = Header->Chunks[WriteChunkIdx];
 
-				WriteChunkStatePtr.FinishWriting(); // STATE CHANGE
-				SetAllReaderEvents();
-			}
-
-			if (ChunkState.GetWriteState() == EWriteState::Complete)
+		// Check if we can append to this chunk
+		FChunkState ChunkState = WriteChunkStatePtr.Get();
+		if (ChunkState.GetWriteState() == EWriteState::Writing)
+		{
+			int Length = ChunkState.GetLength();
+			if (Length + MinSize <= Header->ChunkLength)
 			{
-				return nullptr;
+				return ChunkPtrs[WriteChunkIdx] + Length;
 			}
 
-			int NextWriteChunkIdx = Header->WriteChunkIdx + 1;
-			if (NextWriteChunkIdx == Header->NumChunks)
-			{
-				NextWriteChunkIdx = 0;
-			}
+			WriteChunkStatePtr.FinishWriting(); // STATE CHANGE
+			SetAllReaderEvents();
+		}
+
+		if (ChunkState.GetWriteState() == EWriteState::Complete)
+		{
+			return nullptr;
+		}
+
+		// Otherwise get the next chunk to write to
+		int NextWriteChunkIdx = WriteChunkIdx + 1;
+		if (NextWriteChunkIdx == Header->NumChunks)
+		{
+			NextWriteChunkIdx = 0;
+		}
 
-			FChunkStatePtr& NextWriteChunkStatePtr = Header->Chunks[NextWriteChunkIdx];
-			while (NextWriteChunkStatePtr.Get().GetReaderFlags() != 0)
+		// Wait until all readers have finished with the chunk, and we can update the writer to match
+		FChunkStatePtr& NextWriteChunkStatePtr = Header->Chunks[NextWriteChunkIdx];
+		for (;;)
+		{
+			FChunkState NextWriteChunkState = NextWriteChunkStatePtr.Get();
+			if (NextWriteChunkState.GetReaderFlags() != 0)
 			{
 				if (!WriterEvent.Wait(TimeoutMs))
 				{
@@ -454,20 +521,26 @@ struct FComputeBufferDetail
 				}
 				WriterEvent.Reset();
 			}
-
-			Header->WriteChunkIdx = NextWriteChunkIdx;
-
-			NextWriteChunkStatePtr.StartWriting(Header->NumReaders);
+			else if (NextWriteChunkStatePtr.TryUpdate(NextWriteChunkState, FChunkState(EWriteState::Writing, WriterState.GetReaderFlags(), 0)))
+			{
+				if (Header->Writer.TryUpdate(WriterState, FWriterState(NextWriteChunkIdx, WriterState.GetReaderFlags(), NextWriteChunkIdx == 0)))
+				{
+					break;
+				}
+				else
+				{
+					WriterState = Header->Writer.Get();
+				}
+			}
 		}
+		return ChunkPtrs[NextWriteChunkIdx];
 	}
 
 private:
 	FComputeMemoryMappedFile MemoryMappedFile;
-	volatile long RefCount;
 
 	static void InitShared(std::shared_ptr<FComputeBufferDetail> Detail)
 	{
-		Detail->Reader = FComputeBufferReader(Detail, 0);
 		Detail->Writer = FComputeBufferWriter(Detail);
 
 		FHeader* Header = Detail->Header;
@@ -492,27 +565,60 @@ private:
 
 
 
-//// FComputeBuffer /////
+//// FComputeBufferReaderDetail /////
 
-FComputeBuffer::FComputeBuffer()
-	: Detail(nullptr)
+struct FComputeBufferReaderDetail
 {
-}
+	std::shared_ptr<FComputeBufferDetail> Buffer;
+	FComputeBufferDetail::FReaderState ReaderState;
 
-FComputeBuffer::FComputeBuffer(const FComputeBuffer& Buffer)
-	: FComputeBuffer()
-{
-	if (Buffer.Detail != nullptr)
+	FComputeBufferReaderDetail(std::shared_ptr<FComputeBufferDetail> InBuffer, int InReaderIdx)
+		: Buffer(InBuffer)
+		, ReaderState(InReaderIdx)
 	{
-		Detail = Buffer.Detail;
-		Detail->AddRef();
 	}
-}
 
-FComputeBuffer::FComputeBuffer(FComputeBuffer&& Buffer) noexcept
-	: Detail(Buffer.Detail)
+	~FComputeBufferReaderDetail()
+	{
+		Buffer->ReleaseReader(ReaderState.ReaderIdx);
+	}
+
+	bool IsComplete() const
+	{
+		return Buffer->IsComplete(ReaderState);
+	}
+
+	void Detach()
+	{
+		ReaderState.Detached = true;
+		Buffer->ReaderEvents[ReaderState.ReaderIdx].Set();
+	}
+
+	void AdvanceReadPosition(size_t Size)
+	{
+		ReaderState.Offset += (long)Size;
+	}
+
+	size_t GetMaxReadSize() const
+	{
+		return Buffer->GetMaxReadSize(ReaderState);
+	}
+
+	const unsigned char* WaitToRead(size_t MinSize, int TimeoutMs)
+	{
+		return Buffer->WaitToRead(ReaderState, MinSize, TimeoutMs);
+	}
+};
+
+
+
+
+
+//// FComputeBuffer /////
+
+FComputeBuffer::FComputeBuffer()
+	: Detail(nullptr)
 {
-	Buffer.Detail = nullptr;
 }
 
 FComputeBuffer::~FComputeBuffer()
@@ -537,14 +643,15 @@ void FComputeBuffer::Close()
 	Detail.reset();
 }
 
-FComputeBufferReader& FComputeBuffer::GetReader()
-{
-	return Detail->Reader;
-}
-
-const FComputeBufferReader& FComputeBuffer::GetReader() const
+FComputeBufferReader FComputeBuffer::CreateReader()
 {
-	return Detail->Reader;
+	int ReaderIdx = Detail->AllocateReader(0, 0);
+	if (ReaderIdx == -1)
+	{
+		assert(false);
+		return FComputeBufferReader();
+	}
+	return FComputeBufferReader(std::make_shared<FComputeBufferReaderDetail>(Detail, ReaderIdx));
 }
 
 FComputeBufferWriter& FComputeBuffer::GetWriter()
@@ -566,39 +673,42 @@ const FComputeBufferWriter& FComputeBuffer::GetWriter() const
 //// FComputeBufferReader /////
 
 FComputeBufferReader::FComputeBufferReader()
-	: Detail(nullptr)
-	, ReaderIdx(-1)
 {
 }
 
-FComputeBufferReader::FComputeBufferReader(std::shared_ptr<FComputeBufferDetail> Detail, int ReaderIdx)
+FComputeBufferReader::FComputeBufferReader(std::shared_ptr<FComputeBufferReaderDetail> Detail)
 	: Detail(std::move(Detail))
-	, ReaderIdx(ReaderIdx)
 {
 }
 
 FComputeBufferReader::~FComputeBufferReader()
 {
+	Close();
 }
 
-void FComputeBufferReader::ForceComplete()
+void FComputeBufferReader::Close()
 {
-	Detail->MarkComplete();
+	Detail.reset();
+}
+
+void FComputeBufferReader::Detach()
+{
+	Detail->Detach();
 }
 
 bool FComputeBufferReader::IsComplete() const
 {
-	return Detail->IsComplete(ReaderIdx);
+	return Detail->IsComplete();
 }
 
 void FComputeBufferReader::AdvanceReadPosition(size_t Size)
 {
-	Detail->AdvanceReadPosition(ReaderIdx, Size);
+	Detail->AdvanceReadPosition(Size);
 }
 
 size_t FComputeBufferReader::GetMaxReadSize() const
 {
-	return Detail->GetMaxReadSize(ReaderIdx);
+	return Detail->GetMaxReadSize();
 }
 
 size_t FComputeBufferReader::Read(void* Buffer, size_t MaxSize, int TimeoutMs)
@@ -622,12 +732,12 @@ size_t FComputeBufferReader::Read(void* Buffer, size_t MaxSize, int TimeoutMs)
 
 const unsigned char* FComputeBufferReader::WaitToRead(size_t MinSize, int TimeoutMs)
 {
-	return Detail->WaitToRead(ReaderIdx, MinSize, TimeoutMs);
+	return Detail->WaitToRead(MinSize, TimeoutMs);
 }
 
 const wchar_t* FComputeBufferReader::GetName() const
 {
-	return Detail->Name;
+	return Detail->Buffer->Name;
 }
 
 
diff --git a/Engine/Source/Programs/Horde/Samples/RemoteWorkerCpp/ComputeBuffer.h b/Engine/Source/Programs/Horde/Samples/RemoteWorkerCpp/ComputeBuffer.h
--- a/Engine/Source/Programs/Horde/Samples/RemoteWorkerCpp/ComputeBuffer.h
+++ b/Engine/Source/Programs/Horde/Samples/RemoteWorkerCpp/ComputeBuffer.h
@@ -9,6 +9,7 @@ class FComputeBufferWriter;
 
 struct FHeader;
 struct FComputeBufferDetail;
+struct FComputeBufferReaderDetail;
 
 //
 // Implements a ring buffer using shared memory, with one writer and multiple readers. 
@@ -42,8 +43,6 @@ public:
 	};
 
 	FComputeBuffer();
-	FComputeBuffer(const FComputeBuffer& Buffer);
-	FComputeBuffer(FComputeBuffer&& Buffer) noexcept;
 	~FComputeBuffer();
 
 	// Creates a new buffer
@@ -58,9 +57,8 @@ public:
 	// Test if the buffer is currently open
 	bool IsOpen() const { return Detail != nullptr; }
 
-	// Gets a reference to the reader instance
-	FComputeBufferReader& GetReader();
-	const FComputeBufferReader& GetReader() const;
+	// Creates a new reader for this buffer
+	FComputeBufferReader CreateReader();
 
 	// Gets a reference to the writer instance
 	FComputeBufferWriter& GetWriter();
@@ -77,18 +75,21 @@ class FComputeBufferReader
 {
 public:
 	FComputeBufferReader();
-	FComputeBufferReader(std::shared_ptr<FComputeBufferDetail> Detail, int ReaderIdx);
+	FComputeBufferReader(std::shared_ptr<FComputeBufferReaderDetail> Detail);
 	~FComputeBufferReader();
 
+	// Closes the handle to the underlying reader instance, resetting this instance back to empty
+	void Close();
+
+	// Detaches this reader from the buffer, causing all pending and subsequent reads to return immediately.
+	void Detach();
+
 	// Test if the reader is valid
 	bool IsValid() const { return Detail.get() != nullptr; }
 
 	// Test whether the buffer has finished being written to (ie. MarkComplete() has been called by the writer) and all data has been read from it.
 	bool IsComplete() const;
 
-	// Mark this buffer as complete. This is technically a "write" operation, but is needed to release blocking waits on shutdown.
-	void ForceComplete();
-
 	// Move the read cursor forwards by the given number of bytes
 	void AdvanceReadPosition(size_t Size);
 
@@ -105,8 +106,7 @@ public:
 private:
 	friend class FWorkerComputeSocket;
 
-	std::shared_ptr<FComputeBufferDetail> Detail;
-	int ReaderIdx;
+	std::shared_ptr<FComputeBufferReaderDetail> Detail;
 
 	const wchar_t* GetName() const;
 };
diff --git a/Engine/Source/Programs/Horde/Samples/RemoteWorkerCpp/ComputePlatform.cpp b/Engine/Source/Programs/Horde/Samples/RemoteWorkerCpp/ComputePlatform.cpp
--- a/Engine/Source/Programs/Horde/Samples/RemoteWorkerCpp/ComputePlatform.cpp
+++ b/Engine/Source/Programs/Horde/Samples/RemoteWorkerCpp/ComputePlatform.cpp
@@ -214,6 +214,11 @@ long long FComputePlatform::AtomicAdd64(volatile long long* Ptr, long long Value
 	return InterlockedAdd64(Ptr, Value);
 }
 
+long FComputePlatform::AtomicAnd(volatile long* Ptr, long Value)
+{
+	return InterlockedAnd(Ptr, Value);
+}
+
 long long FComputePlatform::AtomicAnd64(volatile long long* Ptr, long long Value)
 {
 	return InterlockedAnd64(Ptr, Value);
@@ -234,6 +239,11 @@ long long FComputePlatform::AtomicExchange64(volatile long long* Ptr, long long
 	return InterlockedExchange64(Ptr, Exchange);
 }
 
+bool FComputePlatform::AtomicCompareExchange(volatile long* Ptr, long Exchange, long Comperand)
+{
+	return InterlockedCompareExchange(Ptr, Exchange, Comperand) == Comperand;
+}
+
 bool FComputePlatform::AtomicCompareExchange64(volatile long long* Ptr, long long Exchange, long long Comperand)
 {
 	return InterlockedCompareExchange64(Ptr, Exchange, Comperand) == Comperand;
diff --git a/Engine/Source/Programs/Horde/Samples/RemoteWorkerCpp/ComputePlatform.h b/Engine/Source/Programs/Horde/Samples/RemoteWorkerCpp/ComputePlatform.h
--- a/Engine/Source/Programs/Horde/Samples/RemoteWorkerCpp/ComputePlatform.h
+++ b/Engine/Source/Programs/Horde/Samples/RemoteWorkerCpp/ComputePlatform.h
@@ -116,6 +116,9 @@ struct FComputePlatform
 	// Adds a value to the given memory location, and returns the resulting value.
 	static long long AtomicAdd64(volatile long long* Ptr, long long Value);
 
+	// ANDs a value with a given memory location, and returns the contents of the memory location BEFORE the operation.
+	static long AtomicAnd(volatile long* Ptr, long Value);
+
 	// ANDs a value with a given memory location, and returns the contents of the memory location BEFORE the operation.
 	static long long AtomicAnd64(volatile long long* Ptr, long long Value);
 
@@ -125,6 +128,9 @@ struct FComputePlatform
 	// Sets the contents of a pointer to the given value, and returns the original value in that memory location.
 	static long long AtomicExchange64(volatile long long* Ptr, long long Exchange);
 
+	// Compares the contents of a memory location to a value, and exchanges it for another value if they are equal. Returns true if the operation succeeded.
+	static bool AtomicCompareExchange(volatile long* Ptr, long Exchange, long Comperand);
+
 	// Compares the contents of a memory location to a value, and exchanges it for another value if they are equal. Returns true if the operation succeeded.
 	static bool AtomicCompareExchange64(volatile long long* Ptr, long long Exchange, long long Comperand);
 };
diff --git a/Engine/Source/Programs/Horde/Samples/RemoteWorkerCpp/ComputeSocket.cpp b/Engine/Source/Programs/Horde/Samples/RemoteWorkerCpp/ComputeSocket.cpp
--- a/Engine/Source/Programs/Horde/Samples/RemoteWorkerCpp/ComputeSocket.cpp
+++ b/Engine/Source/Programs/Horde/Samples/RemoteWorkerCpp/ComputeSocket.cpp
@@ -39,9 +39,9 @@ FComputeChannel FComputeSocket::CreateChannel(int ChannelId)
 FComputeChannel FComputeSocket::CreateChannel(int ChannelId, FComputeBuffer RecvBuffer, FComputeBuffer SendBuffer)
 {
 	AttachRecvBuffer(ChannelId, RecvBuffer.GetWriter());
-	AttachSendBuffer(ChannelId, SendBuffer.GetReader());
+	AttachSendBuffer(ChannelId, SendBuffer.CreateReader());
 
-	return FComputeChannel(RecvBuffer.GetReader(), SendBuffer.GetWriter());
+	return FComputeChannel(RecvBuffer.CreateReader(), SendBuffer.GetWriter());
 }
 
 //////////////////////////////////////////////////////
@@ -131,7 +131,7 @@ void FWorkerComputeSocket::RunServer(FComputeBufferReader& CommandBufferReader,
 				FComputeBuffer Buffer;
 				if (Buffer.OpenExisting(Name))
 				{
-					Socket.AttachSendBuffer(ChannelId, Buffer.GetReader());
+					Socket.AttachSendBuffer(ChannelId, Buffer.CreateReader());
 				}
 				else
 				{
@@ -270,7 +270,7 @@ public:
 	{
 		for (FComputeBufferReader& Reader : Readers)
 		{
-			Reader.ForceComplete();
+			Reader.Detach();
 		}
 
 		for (std::pair<const int, std::thread>& Pair : SendThreads)
diff --git a/Engine/Source/Programs/Horde/Samples/RemoteWorkerCpp/ComputeTransport.cpp b/Engine/Source/Programs/Horde/Samples/RemoteWorkerCpp/ComputeTransport.cpp
--- a/Engine/Source/Programs/Horde/Samples/RemoteWorkerCpp/ComputeTransport.cpp
+++ b/Engine/Source/Programs/Horde/Samples/RemoteWorkerCpp/ComputeTransport.cpp
@@ -42,6 +42,7 @@ FBufferTransport::FBufferTransport(FComputeBuffer InSendBuffer, FComputeBuffer I
 	: SendBuffer(std::move(InSendBuffer))
 	, RecvBuffer(std::move(InRecvBuffer))
 {
+	RecvBufferReader = std::move(RecvBuffer.CreateReader());
 }
 
 size_t FBufferTransport::Send(const void* Data, size_t Size)
@@ -58,12 +59,11 @@ size_t FBufferTransport::Send(const void* Data, size_t Size)
 
 size_t FBufferTransport::Recv(void* Data, size_t Size)
 {
-	FComputeBufferReader& Reader = RecvBuffer.GetReader();
-	const unsigned char* Buffer = Reader.WaitToRead(1);
+	const unsigned char* Buffer = RecvBufferReader.WaitToRead(1);
 
-	size_t ReadSize = std::min(Size, Reader.GetMaxReadSize());
+	size_t ReadSize = std::min(Size, RecvBufferReader.GetMaxReadSize());
 	memcpy(Data, Buffer, ReadSize);
-	Reader.AdvanceReadPosition(ReadSize);
+	RecvBufferReader.AdvanceReadPosition(ReadSize);
 
 	return ReadSize;
 }
diff --git a/Engine/Source/Programs/Horde/Samples/RemoteWorkerCpp/ComputeTransport.h b/Engine/Source/Programs/Horde/Samples/RemoteWorkerCpp/ComputeTransport.h
--- a/Engine/Source/Programs/Horde/Samples/RemoteWorkerCpp/ComputeTransport.h
+++ b/Engine/Source/Programs/Horde/Samples/RemoteWorkerCpp/ComputeTransport.h
@@ -46,4 +46,5 @@ protected:
 private:
 	FComputeBuffer SendBuffer;
 	FComputeBuffer RecvBuffer;
+	FComputeBufferReader RecvBufferReader;
 };
diff --git a/Engine/Source/Programs/Horde/Samples/RemoteWorkerCpp/Tests.cpp b/Engine/Source/Programs/Horde/Samples/RemoteWorkerCpp/Tests.cpp
--- a/Engine/Source/Programs/Horde/Samples/RemoteWorkerCpp/Tests.cpp
+++ b/Engine/Source/Programs/Horde/Samples/RemoteWorkerCpp/Tests.cpp
@@ -28,7 +28,7 @@ void ComputeBufferTest()
 	memcpy(WriteBuffer, TestData, sizeof(TestData));
 	Writer.AdvanceWritePosition(sizeof(TestData));
 
-	FComputeBufferReader& Reader = Buffer.GetReader();
+	FComputeBufferReader Reader = Buffer.CreateReader();
 	const unsigned char* ReadBuffer = Reader.WaitToRead(sizeof(TestData));
 	verify(memcmp(ReadBuffer, TestData, sizeof(TestData)) == 0);
 }
