Horde: Merge CB unit tests from Jupiter.

diff --git a/Engine/Source/Programs/Horde/Horde.Server.Tests/Ddc/UnitTests/CompactBinaryTests.cs b/Engine/Source/Programs/Horde/Horde.Server.Tests/Ddc/UnitTests/CompactBinaryTests.cs
new file mode 100644
--- /dev/null
+++ b/Engine/Source/Programs/Horde/Horde.Server.Tests/Ddc/UnitTests/CompactBinaryTests.cs
@@ -0,0 +1,133 @@
+// Copyright Epic Games, Inc. All Rights Reserved.
+
+using System.Collections.Generic;
+using System.IO;
+using System.Linq;
+using EpicGames.Core;
+using EpicGames.Serialization;
+using Microsoft.VisualStudio.TestTools.UnitTesting;
+
+namespace Horde.Server.Tests.Ddc.UnitTests
+{
+	[TestClass]
+	public class CompactBinaryTests
+	{
+		[TestMethod]
+		public void BuildObject()
+		{
+			byte[] bytes = File.ReadAllBytes("Ddc/UnitTests/CompactBinaryObjects/build");
+
+			CbObject o = new CbObject(bytes);
+
+			Assert.AreEqual("BuildAction", o.AsField().Name);
+			List<CbField> buildActionFields = o.ToList();
+			Assert.AreEqual(3, buildActionFields.Count);
+			Assert.AreEqual("Function", buildActionFields[0].Name);
+			Assert.AreEqual("Constants", buildActionFields[1].Name);
+			Assert.AreEqual("Inputs", buildActionFields[2].Name);
+
+			List<CbField>  constantsFields = buildActionFields[1].ToList();
+			Assert.AreEqual(3, constantsFields.Count);
+			Assert.AreEqual("TextureBuildSettings", constantsFields[0].Name);
+			Assert.AreEqual("TextureOutputSettings", constantsFields[1].Name);
+			Assert.AreEqual("TextureSource", constantsFields[2].Name);
+
+			List<CbField>  inputsFields = buildActionFields[2].ToList();
+			Assert.AreEqual(1, inputsFields.Count);
+			Assert.AreEqual("7587B323422942733DDD048A91709FDE", inputsFields[0].Name);
+			Assert.IsTrue(inputsFields[0].IsBinaryAttachment());
+			Assert.IsTrue(inputsFields[0].IsAttachment());
+			Assert.IsFalse(inputsFields[0].IsObjectAttachment());
+			Assert.AreEqual(IoHash.Parse("f855382171a0b1e5a1c653aa6c5121a05cbf4ba0"), inputsFields[0].AsHash());
+		}
+
+		
+		[TestMethod]
+		public void ReferenceOutput()
+		{
+			byte[] bytes = File.ReadAllBytes("Ddc/UnitTests/CompactBinaryObjects/ReferenceOutput");
+
+			CbObject o = new CbObject(bytes);
+
+			Assert.AreEqual("BuildOutput", o.AsField().Name);
+			List<CbField> buildActionFields = o.ToList();
+			Assert.AreEqual(1, buildActionFields.Count);
+			CbField payloads = buildActionFields[0];
+			List<CbField> payloadFields = payloads.ToList();
+			Assert.AreEqual(3, payloadFields.Count);
+
+			Assert.AreEqual(IoHash.Parse("5d8a6dc277c968f0d027c98f879c955c1905c293"), payloadFields[0]["RawHash"]!.AsHash());
+			Assert.AreEqual(IoHash.Parse("313f0d0d334100d83aeb1ee2c42794fd087cb0ae"), payloadFields[1]["RawHash"]!.AsHash());
+			Assert.AreEqual(IoHash.Parse("c7a03f83c08cdca882110ecf2b5654ee3b09b11e"), payloadFields[2]["RawHash"]!.AsHash());
+		}
+
+		[TestMethod]
+		public void CompactBinary()
+		{
+			byte[] bytes = File.ReadAllBytes("Ddc/UnitTests/CompactBinaryObjects/compact_binary");
+
+			CbObject o = new CbObject(bytes);
+			Assert.AreEqual("", o.AsField().Name);
+			List<CbField> buildActionFields = o.ToList();
+			Assert.AreEqual(3, buildActionFields.Count);
+			CbField payloads = buildActionFields[0];
+			List<CbField> payloadFields = payloads.ToList();
+			Assert.AreEqual(2, payloadFields.Count);
+
+			Assert.AreEqual("{\"Key\":{\"Bucket\":\"EditorDomainPackage\",\"Hash\":\"37dbaa409ef30ba67f18c8fc2faaf606636cb915\"},\"Meta\":{\"FileSize\":24789},\"Attachments\":[{\"Id\":\"000000000000000000000001\",\"RawHash\":\"da6fc57e4b9f91377c9509ea0ad567bacb3796c5\",\"RawSize\":24789}]}", o.ToJson());
+		}
+
+		[TestMethod]
+		public void WriteArray()
+		{
+			IoHash hash1 = IoHash.Parse("5d8a6dc277c968f0d027c98f879c955c1905c293");
+			IoHash hash2 = IoHash.Parse("313f0d0d334100d83aeb1ee2c42794fd087cb0ae");
+
+			CbWriter writer = new CbWriter();
+			writer.BeginObject();
+			writer.BeginUniformArray("needs", CbFieldType.Hash);
+			writer.WriteHashValue(hash1);
+			writer.WriteHashValue(hash2);
+			writer.EndUniformArray();
+			writer.EndObject();
+
+			byte[] objectData = writer.ToByteArray();
+			CbObject o = new CbObject(objectData);
+
+			// the top object has no name
+			Assert.AreEqual("", o.AsField().Name);
+			List<CbField> fields = o.ToList();
+			Assert.AreEqual(1, fields.Count);
+			CbField? needs = o["needs"];
+			List<CbField> blobList = needs!.AsArray().ToList();
+			IoHash[] blobs = blobList.Select(field => field!.AsHash()).ToArray();
+			CollectionAssert.AreEqual(new IoHash[] {hash1, hash2}, blobs);
+		}
+
+		[TestMethod]
+		public void WriteObject()
+		{
+			IoHash hash1 = IoHash.Parse("5d8a6dc277c968f0d027c98f879c955c1905c293");
+
+			CbWriter writer = new CbWriter();
+			writer.BeginObject();
+			writer.WriteString("string", "test");
+			writer.WriteBinaryAttachment("hash", hash1);
+			writer.EndObject();
+
+			byte[] objectData = writer.ToByteArray();
+			CbObject o = new CbObject(objectData);
+
+			// the object has no name and 2 fields
+			Assert.AreEqual("", o.AsField().Name);
+			List<CbField> fields = o.ToList();
+			Assert.AreEqual(2, fields.Count);
+
+			CbField? stringField = o["string"];
+			Assert.AreEqual("test", stringField!.AsString());
+
+			CbField? hashField = o["hash"];
+			Assert.AreEqual(hash1, hashField!.AsAttachment());
+		}
+	}
+}
diff --git a/Engine/Source/Programs/Horde/Horde.Server.Tests/Horde.Server.Tests.csproj b/Engine/Source/Programs/Horde/Horde.Server.Tests/Horde.Server.Tests.csproj
--- a/Engine/Source/Programs/Horde/Horde.Server.Tests/Horde.Server.Tests.csproj
+++ b/Engine/Source/Programs/Horde/Horde.Server.Tests/Horde.Server.Tests.csproj
@@ -57,4 +57,16 @@
     <Folder Include="Resources\" />
   </ItemGroup>
 
+  <ItemGroup>
+    <None Update="Ddc\UnitTests\CompactBinaryObjects\build">
+      <CopyToOutputDirectory>Always</CopyToOutputDirectory>
+    </None>
+    <None Update="Ddc\UnitTests\CompactBinaryObjects\compact_binary">
+      <CopyToOutputDirectory>Always</CopyToOutputDirectory>
+    </None>
+    <None Update="Ddc\UnitTests\CompactBinaryObjects\ReferenceOutput">
+      <CopyToOutputDirectory>Always</CopyToOutputDirectory>
+    </None>
+  </ItemGroup>
+
 </Project>
