Horde: Add compressed buffer DDC tests.

diff --git a/Engine/Source/Programs/Horde/Horde.Server.Tests/Ddc/UnitTests/CompressedBufferTests.cs b/Engine/Source/Programs/Horde/Horde.Server.Tests/Ddc/UnitTests/CompressedBufferTests.cs
new file mode 100644
--- /dev/null
+++ b/Engine/Source/Programs/Horde/Horde.Server.Tests/Ddc/UnitTests/CompressedBufferTests.cs
@@ -0,0 +1,34 @@
+// Copyright Epic Games, Inc. All Rights Reserved.
+
+using System.Runtime.InteropServices;
+using System.Text;
+using Horde.Server.Ddc;
+using Microsoft.VisualStudio.TestTools.UnitTesting;
+using OpenTelemetry.Trace;
+
+namespace Horde.Server.Tests.Ddc.UnitTests
+{
+	[TestClass]
+	public class CompressedBufferTests
+	{
+
+		[TestMethod]
+		public void CompressAndDecompress()
+		{
+			if (RuntimeInformation.IsOSPlatform(OSPlatform.Windows) && RuntimeInformation.OSArchitecture == Architecture.Arm64)
+			{
+				Assert.Inconclusive("No oodle libs for Windows-Arm64");
+			}
+
+			byte[] bytes = Encoding.UTF8.GetBytes("this is a test string");
+
+			CompressedBufferUtils bufferUtils = new(TracerProvider.Default.GetTracer("TestTracer"));
+
+			byte[] compressedBytes = bufferUtils.CompressContent(OoodleCompressorMethod.Mermaid, OoodleCompressionLevel.VeryFast, bytes);
+
+			byte[] roundTrippedBytes = bufferUtils.DecompressContent(compressedBytes);
+
+			CollectionAssert.AreEqual(bytes, roundTrippedBytes);
+		}
+	}
+}
