Horde: Add next job id to bisection responses for running jobs, utility method to create bisection responses, fix  non-cached user access
#rnx

diff --git a/Engine/Source/Programs/Horde/Horde.Server/Jobs/Bisect/BisectController.cs b/Engine/Source/Programs/Horde/Horde.Server/Jobs/Bisect/BisectController.cs
--- a/Engine/Source/Programs/Horde/Horde.Server/Jobs/Bisect/BisectController.cs
+++ b/Engine/Source/Programs/Horde/Horde.Server/Jobs/Bisect/BisectController.cs
@@ -103,16 +103,22 @@ namespace Horde.Server.Jobs.Bisect
 		/// <inheritdoc cref="IBisectTask.CurrentChange"/>
 		public int CurrentChange => _bisectTask.CurrentChange;
 
+		/// <summary>
+		/// The next job id for a running bisection task
+		/// </summary>
+		public JobId? NextJobId { get; }
+
 		/// <summary>
 		/// The steps involved in the bisection
 		/// </summary>
 		public List<GetJobStepRefResponse> Steps { get; }
 
-		internal GetBisectTaskResponse(IBisectTask bisectTask, GetThinUserInfoResponse owner, List<IJobStepRef> steps)
+		internal GetBisectTaskResponse(IBisectTask bisectTask, GetThinUserInfoResponse owner, List<IJobStepRef> steps, JobId? nextJobId)
 		{
 			_bisectTask = bisectTask;
 			Owner = owner;
 			Steps = steps.Select(s => new GetJobStepRefResponse(s)).ToList();
+			NextJobId = nextJobId;
 
 		}
 	}
@@ -272,10 +278,7 @@ namespace Horde.Server.Jobs.Bisect
 				return Forbid(BisectTaskAclAction.ViewBisectTask, streamConfig.Id);
 			}
 
-			IUser? user = await _userCollection.GetUserAsync(bisectTask.OwnerId);
-			List<IJobStepRef> steps = await _jobStepRefs.GetStepsForNodeAsync(initialJob.StreamId, initialJob.TemplateId, bisectTask.NodeName, null, true, 1024, bisectTask.Id, cancellationToken);
-
-			return new GetBisectTaskResponse(bisectTask, new GetThinUserInfoResponse(user), steps);
+			return await CreateBisectTaskResponse(initialJob, bisectTask, cancellationToken);
 		}
 
 		/// <summary>
@@ -331,6 +334,14 @@ namespace Horde.Server.Jobs.Bisect
 			}
 		}
 
+		async Task<GetBisectTaskResponse> CreateBisectTaskResponse(IJob initialJob, IBisectTask task, CancellationToken cancellationToken = default)
+		{
+			IUser? user = await _userCollection.GetCachedUserAsync(task.OwnerId);
+			List<IJobStepRef> steps = await _jobStepRefs.GetStepsForNodeAsync(initialJob.StreamId, initialJob.TemplateId, task.NodeName, null, true, 1024, task.Id, cancellationToken);
+			IJob? nextJob = task.State == BisectTaskState.Running ? await _jobCollection.FindBisectTaskJobsAsync(task.Id, true, cancellationToken).FirstOrDefaultAsync(cancellationToken) : null;
+			return new GetBisectTaskResponse(task, new GetThinUserInfoResponse(user), steps, nextJob?.Id);
+		}
+
 		/// <summary>
 		/// Gets the bisections run on a specific job
 		/// </summary>
@@ -363,20 +374,14 @@ namespace Horde.Server.Jobs.Bisect
 					return Forbid(BisectTaskAclAction.ViewBisectTask, streamConfig.Id);
 				}
 
-
-
-				IUser? user = await _userCollection.GetUserAsync(bisectTask.OwnerId);
-				List<IJobStepRef> steps = await _jobStepRefs.GetStepsForNodeAsync(job.StreamId, job.TemplateId, bisectTask.NodeName, null, true, 1024, bisectTask.Id, cancellationToken);
-
-				response.Add(new GetBisectTaskResponse(bisectTask, new GetThinUserInfoResponse(user), steps));
-
+				response.Add(await CreateBisectTaskResponse(job, bisectTask, cancellationToken));
 			}
 
 			return response;
 		}
 
 		/// <summary>
-		/// 
+		/// Gets bisection tasks based on specified criteria
 		/// </summary>
 		/// <returns></returns>
 		[HttpGet]
@@ -407,13 +412,11 @@ namespace Horde.Server.Jobs.Bisect
 
 			for (int i = 0; i < tasks.Count; i++)
 			{
-				IBisectTask task = tasks[i];
-				IUser? user = await _userCollection.GetUserAsync(task.OwnerId);
+				IBisectTask task = tasks[i];								
 				IJob? job = jobs.FirstOrDefault(x => x.Id == task.InitialJobId);
-				if (job != null && user != null)
+				if (job != null)
 				{
-					List<IJobStepRef> steps = await _jobStepRefs.GetStepsForNodeAsync(job.StreamId, job.TemplateId, task.NodeName, null, true, 1024,task.Id, cancellationToken);
-					responses.Add(new GetBisectTaskResponse(task, new GetThinUserInfoResponse(user), steps));
+					responses.Add(await CreateBisectTaskResponse(job, task, cancellationToken));				
 				}
 			}
 
