Horde: Add next job change to bisect response so that we don't need to look it up on client from id
#rnx

diff --git a/Engine/Source/Programs/Horde/Horde.Server/Jobs/Bisect/BisectController.cs b/Engine/Source/Programs/Horde/Horde.Server/Jobs/Bisect/BisectController.cs
--- a/Engine/Source/Programs/Horde/Horde.Server/Jobs/Bisect/BisectController.cs
+++ b/Engine/Source/Programs/Horde/Horde.Server/Jobs/Bisect/BisectController.cs
@@ -108,18 +108,23 @@ namespace Horde.Server.Jobs.Bisect
 		/// </summary>
 		public JobId? NextJobId { get; }
 
+		/// <summary>
+		/// The next job change
+		/// </summary>
+		public int? NextJobChange { get; }
+
 		/// <summary>
 		/// The steps involved in the bisection
 		/// </summary>
 		public List<GetJobStepRefResponse> Steps { get; }
 
-		internal GetBisectTaskResponse(IBisectTask bisectTask, GetThinUserInfoResponse owner, List<IJobStepRef> steps, JobId? nextJobId)
+		internal GetBisectTaskResponse(IBisectTask bisectTask, GetThinUserInfoResponse owner, List<IJobStepRef> steps, IJob? nextJob)
 		{
 			_bisectTask = bisectTask;
 			Owner = owner;
 			Steps = steps.Select(s => new GetJobStepRefResponse(s)).ToList();
-			NextJobId = nextJobId;
-
+			NextJobId = nextJob?.Id;
+			NextJobChange = nextJob?.Change; 
 		}
 	}
 
@@ -339,7 +344,7 @@ namespace Horde.Server.Jobs.Bisect
 			IUser? user = await _userCollection.GetCachedUserAsync(task.OwnerId);
 			List<IJobStepRef> steps = await _jobStepRefs.GetStepsForNodeAsync(initialJob.StreamId, initialJob.TemplateId, task.NodeName, null, true, 1024, task.Id, cancellationToken);
 			IJob? nextJob = task.State == BisectTaskState.Running ? await _jobCollection.FindBisectTaskJobsAsync(task.Id, true, cancellationToken).FirstOrDefaultAsync(cancellationToken) : null;
-			return new GetBisectTaskResponse(task, new GetThinUserInfoResponse(user), steps, nextJob?.Id);
+			return new GetBisectTaskResponse(task, new GetThinUserInfoResponse(user), steps, nextJob);
 		}
 
 		/// <summary>
