UnrealAutomatedTests plugin - Add changelist range to history item

#rnx

diff --git a/Engine/Source/Programs/Horde/HordeDashboard/src/Plugins/UnrealAutomatedTests/components/UnrealAutomatedTests.tsx b/Engine/Source/Programs/Horde/HordeDashboard/src/Plugins/UnrealAutomatedTests/components/UnrealAutomatedTests.tsx
--- a/Engine/Source/Programs/Horde/HordeDashboard/src/Plugins/UnrealAutomatedTests/components/UnrealAutomatedTests.tsx
+++ b/Engine/Source/Programs/Horde/HordeDashboard/src/Plugins/UnrealAutomatedTests/components/UnrealAutomatedTests.tsx
@@ -6,6 +6,8 @@ import React, { useState, useEffect, useRef } from 'react';
 import { useNavigate } from 'react-router';
 import { Link } from 'react-router-dom';
 import backend from '../../../backend';
+import dashboard from "../../../backend/Dashboard";
+import { projectStore } from "../../../backend/ProjectStore";
 import { ArtifactData } from '../../../backend/Api';
 import { msecToElapsed } from '../../../base/utilities/timeUtils';
 import { testDataHandler } from '../../../components/TestReportView';
@@ -257,6 +259,10 @@ const HistoryItem: React.FC<{ item: TestStateHistoryItem, testName: string, sele
             <Link to={`/testreport/${item.TestdataId}?test=${testName}`}>
                <Text variant="smallPlus">{getStateLabel(item.State)} on {item.Change}</Text>
             </Link>
+            {item.RangeUrl &&
+               <a href={item.RangeUrl} target="blank">
+                  <Text variant="smallPlus" styles={{ root: {color: theme.palette.neutralDark,  paddingLeft: 4, paddingRight: 4 } }}>[ Swarm Range ]</Text>
+               </a>}
          </Stack.Item>
       </Stack>
    );
@@ -363,21 +369,28 @@ const TestResultPane: React.FC<{ test: TestResult, selected: boolean }> = (props
    function getTestHistory(testFullName: string, testdataItems: TestDataWrapper[]): TestStateHistoryItem[] {
       const testResults: TestStateHistoryItem[] = [];
 
-      testdataItems.forEach((item) => {
+      let previousChange: number = 0;
+      const stream = projectStore.streamById(testDataHandler.streamId)!;
+      const project = projectStore.byId(stream?.projectId)!;
+      const name = project?.name === "Engine" ? "UE4" : project.name;
+
+      testdataItems.reverse().forEach((item) => {
          const testdata = item.data as TestPassSummary;
          const foundTest = testdata.Tests.find(testItem => testItem.FullTestPath === testFullName);
 
          if (foundTest !== undefined) {
             const testHistoryItem: TestStateHistoryItem = {
                Change: item.change,
+               RangeUrl: previousChange > 0 ? `${dashboard.swarmUrl}/files/${name}/${stream.name}?range=@${previousChange},@${item.change}#commits` : undefined,
                TestdataId: item.id,
                State: foundTest.Warnings > 0 && foundTest.State === TestState.Success ? TestState.SuccessWithWarnings : foundTest.State,
             };
             testResults.push(testHistoryItem);
          }
+         previousChange = item.change;
       });
 
-      return testResults;
+      return testResults.reverse();
    }
 
    function onClickHistory() {
@@ -435,7 +448,7 @@ const TestResultPane: React.FC<{ test: TestResult, selected: boolean }> = (props
                   <Stack onClick={(ev) => { onClickHistory() }} tabIndex={0} onBlur={onHistoryBlur} className={styles.itemHover} style={{ position: "relative" }}>
                      <Text styles={{ root: { fontWeight: 'bold', paddingLeft: 4, paddingRight: 4 } }}>History</Text>
                      {historyVisible &&
-                        <div style={{ position: "absolute", width: 200, top: "100%", maxHeight: 300 }} className={styles.historyList} onClick={(ev) => { ev.stopPropagation() }}>
+                        <div style={{ position: "absolute", width: 250, top: "100%", maxHeight: 300 }} className={styles.historyList} onClick={(ev) => { ev.stopPropagation() }}>
                            {!historyLoaded && <Spinner styles={{ root: { padding: 10 } }} size={SpinnerSize.small}></Spinner>}
                            {historyLoaded && history.length > 0 &&
                               history.map((historyItem) => <HistoryItem key={historyItem.Change} item={historyItem} testName={test.FullTestPath} selected={historyItem.Change === testDataHandler.cursor?.change} />)
diff --git a/Engine/Source/Programs/Horde/HordeDashboard/src/Plugins/UnrealAutomatedTests/models/UnrealAutomatedTests.ts b/Engine/Source/Programs/Horde/HordeDashboard/src/Plugins/UnrealAutomatedTests/models/UnrealAutomatedTests.ts
--- a/Engine/Source/Programs/Horde/HordeDashboard/src/Plugins/UnrealAutomatedTests/models/UnrealAutomatedTests.ts
+++ b/Engine/Source/Programs/Horde/HordeDashboard/src/Plugins/UnrealAutomatedTests/models/UnrealAutomatedTests.ts
@@ -83,5 +83,6 @@ export type TestPassSummary = {
 export type TestStateHistoryItem = {
     TestdataId: string;
     Change: number;
+    RangeUrl?: string;
     State: string;
 }
