Horde: First pass at cross-platform support for worker remote execution lib.

diff --git a/Engine/Source/Programs/Horde/Samples/RemoteWorkerCpp/ComputeBuffer.cpp b/Engine/Source/Programs/Horde/Samples/RemoteWorkerCpp/ComputeBuffer.cpp
--- a/Engine/Source/Programs/Horde/Samples/RemoteWorkerCpp/ComputeBuffer.cpp
+++ b/Engine/Source/Programs/Horde/Samples/RemoteWorkerCpp/ComputeBuffer.cpp
@@ -5,6 +5,8 @@
 #include <assert.h>
 #include <iostream>
 #include <assert.h>
+#include <stdio.h>
+#include <string.h>
 
 struct FComputeBufferDetail
 {
@@ -179,24 +181,24 @@ struct FComputeBufferDetail
 		FChunkStatePtr Chunks[FComputeBuffer::MaxChunks];
 	};
 
-	wchar_t Name[260];
+	char Name[FComputeBuffer::MaxNameLength];
 
 	FHeader* Header;
 	unsigned char* ChunkPtrs[FComputeBuffer::MaxChunks];
 
-	FComputeManualResetEvent WriterEvent;
-	FComputeManualResetEvent ReaderEvents[FComputeBuffer::MaxReaders];
+	FComputeEvent WriterEvent;
+	FComputeEvent ReaderEvents[FComputeBuffer::MaxReaders];
 
 	FComputeBufferWriter Writer;
 
-	FComputeBufferDetail(const wchar_t* Name)
+	FComputeBufferDetail(const char* InName)
 		: Header(nullptr)
 		, ChunkPtrs{ nullptr, }
 	{
 		static_assert(sizeof(FChunkStatePtr) == sizeof(long long), "Incorrect size of FChunkStatePtr; check union is declared correctly.");
 		static_assert(sizeof(FWriterStatePtr) == sizeof(long long), "Incorrect size of FWriterStatePtr; check union is declared correctly.");
 
-		wcscpy_s(this->Name, Name);
+		FComputePlatform::Strcpy(Name, FComputeBuffer::MaxNameLength, InName);
 	}
 
 	~FComputeBufferDetail()
@@ -208,17 +210,17 @@ struct FComputeBufferDetail
 	{
 		long long Capacity = sizeof(FHeader) + (Params.NumChunks * sizeof(unsigned int)) + (Params.NumChunks * Params.ChunkLength);
 
-		const wchar_t* Name = Params.Name;
+		const char* Name = Params.Name;
 
-		wchar_t BaseNameBuffer[FComputeBuffer::MaxNameLength];
+		char BaseNameBuffer[FComputeBuffer::MaxNameLength];
 		if (Name == nullptr)
 		{
 			FComputePlatform::CreateUniqueName(BaseNameBuffer, FComputeBuffer::MaxNameLength);
 			Name = BaseNameBuffer;
 		}
 
-		wchar_t NameBuffer[FComputeBuffer::MaxNameLength];
-		swprintf(NameBuffer, FComputeBuffer::MaxNameLength, L"%s_M", Name);
+		char NameBuffer[FComputeBuffer::MaxNameLength];
+		snprintf(NameBuffer, FComputeBuffer::MaxNameLength, "%s_M", Name);
 
 		std::shared_ptr<FComputeBufferDetail> Detail = std::make_shared<FComputeBufferDetail>(Name);
 		if (!Detail->MemoryMappedFile.Create(NameBuffer, Capacity))
@@ -238,7 +240,7 @@ struct FComputeBufferDetail
 		Detail->Header->ChunkLength = Params.ChunkLength;
 		Detail->Header->Chunks[0].Set(FChunkState(EWriteState::Writing, 0, 0));
 
-		swprintf(NameBuffer, FComputeBuffer::MaxNameLength, L"%s_W", Name);
+		snprintf(NameBuffer, FComputeBuffer::MaxNameLength, "%s_W", Name);
 		if (!Detail->WriterEvent.Create(NameBuffer))
 		{
 			return nullptr;
@@ -246,7 +248,7 @@ struct FComputeBufferDetail
 
 		for (int ReaderIdx = 0; ReaderIdx < Detail->Header->NumReaders; ReaderIdx++)
 		{
-			swprintf(NameBuffer, FComputeBuffer::MaxNameLength, L"%s_R%d", Name, ReaderIdx);
+			snprintf(NameBuffer, FComputeBuffer::MaxNameLength, "%s_R%d", Name, ReaderIdx);
 			if (!Detail->ReaderEvents[ReaderIdx].Create(NameBuffer))
 			{
 				return nullptr;
@@ -257,10 +259,10 @@ struct FComputeBufferDetail
 		return Detail;
 	}
 
-	static std::shared_ptr<FComputeBufferDetail> OpenExisting(const wchar_t* Name)
+	static std::shared_ptr<FComputeBufferDetail> OpenExisting(const char* Name)
 	{
-		wchar_t NameBuffer[FComputeBuffer::MaxNameLength];
-		swprintf(NameBuffer, FComputeBuffer::MaxNameLength, L"%s_M", Name);
+		char NameBuffer[FComputeBuffer::MaxNameLength];
+		snprintf(NameBuffer, FComputeBuffer::MaxNameLength, "%s_M", Name);
 
 		std::shared_ptr<FComputeBufferDetail> Detail = std::make_shared<FComputeBufferDetail>(Name);
 		if (!Detail->MemoryMappedFile.OpenExisting(NameBuffer))
@@ -274,7 +276,7 @@ struct FComputeBufferDetail
 			return nullptr;
 		}
 
-		swprintf(NameBuffer, FComputeBuffer::MaxNameLength, L"%s_W", Name);
+		snprintf(NameBuffer, FComputeBuffer::MaxNameLength, "%s_W", Name);
 		if (!Detail->WriterEvent.OpenExisting(NameBuffer))
 		{
 			return nullptr;
@@ -282,7 +284,7 @@ struct FComputeBufferDetail
 
 		for (int ReaderIdx = 0; ReaderIdx < Detail->Header->NumReaders; ReaderIdx++)
 		{
-			swprintf(NameBuffer, FComputeBuffer::MaxNameLength, L"%s_R%d", Name, ReaderIdx);
+			snprintf(NameBuffer, FComputeBuffer::MaxNameLength, "%s_R%d", Name, ReaderIdx);
 			if (!Detail->ReaderEvents[ReaderIdx].OpenExisting(NameBuffer))
 			{
 				return nullptr;
@@ -308,7 +310,7 @@ struct FComputeBufferDetail
 	{
 		for (;;)
 		{
-			int AllocatedReaders = Header->AllocatedReaders;
+			int AllocatedReaders = (int)Header->AllocatedReaders;
 			int ReaderFlag = (AllocatedReaders + 1) ^ AllocatedReaders;
 
 			int ReaderIdx = FComputePlatform::FloorLog2(ReaderFlag);
@@ -379,7 +381,6 @@ struct FComputeBufferDetail
 			if (!ChunkState.HasReaderFlag(ReaderIdx))
 			{
 				// Wait until the current chunk is readable
-				ReaderEvents[ReaderIdx].Reset();
 				if (!ChunkState.HasReaderFlag(ReaderIdx) && !ReaderEvents[ReaderIdx].Wait(TimeoutMs))
 				{
 					return nullptr;
@@ -393,7 +394,6 @@ struct FComputeBufferDetail
 			else if (ChunkState.GetWriteState() == EWriteState::Writing)
 			{
 				// Wait until there is more data in the chunk
-				ReaderEvents[ReaderIdx].Reset();
 				if (Header->Chunks[ReaderState.ChunkIdx].Get().Value == ChunkState.Value && !ReaderEvents[ReaderIdx].Wait(TimeoutMs))
 				{
 					return nullptr;
@@ -408,7 +408,7 @@ struct FComputeBufferDetail
 			{
 				// Move to the next chunk
 				ChunkStatePtr.FinishReading(ReaderIdx);
-				WriterEvent.Set();
+				WriterEvent.Signal();
 
 				if (++ReaderState.ChunkIdx == Header->NumChunks)
 				{
@@ -519,7 +519,6 @@ struct FComputeBufferDetail
 				{
 					return nullptr;
 				}
-				WriterEvent.Reset();
 			}
 			else if (NextWriteChunkStatePtr.TryUpdate(NextWriteChunkState, FChunkState(EWriteState::Writing, WriterState.GetReaderFlags(), 0)))
 			{
@@ -557,7 +556,7 @@ private:
 	{
 		for (int Idx = 0; Idx < Header->NumReaders; Idx++)
 		{
-			ReaderEvents[Idx].Set();
+			ReaderEvents[Idx].Signal();
 		}
 	}
 };
@@ -591,7 +590,7 @@ struct FComputeBufferReaderDetail
 	void Detach()
 	{
 		ReaderState.Detached = true;
-		Buffer->ReaderEvents[ReaderState.ReaderIdx].Set();
+		Buffer->ReaderEvents[ReaderState.ReaderIdx].Signal();
 	}
 
 	void AdvanceReadPosition(size_t Size)
@@ -632,7 +631,7 @@ bool FComputeBuffer::CreateNew(const FParams& Params)
 	return Detail != nullptr;
 }
 
-bool FComputeBuffer::OpenExisting(const wchar_t* Name)
+bool FComputeBuffer::OpenExisting(const char* Name)
 {
 	Detail = FComputeBufferDetail::OpenExisting(Name);
 	return Detail != nullptr;
@@ -735,7 +734,7 @@ const unsigned char* FComputeBufferReader::WaitToRead(size_t MinSize, int Timeou
 	return Detail->WaitToRead(MinSize, TimeoutMs);
 }
 
-const wchar_t* FComputeBufferReader::GetName() const
+const char* FComputeBufferReader::GetName() const
 {
 	return Detail->Buffer->Name;
 }
@@ -799,7 +798,7 @@ unsigned char* FComputeBufferWriter::WaitToWrite(size_t MinSize, int TimeoutMs)
 	return Detail->WaitToWrite(MinSize, TimeoutMs);
 }
 
-const wchar_t* FComputeBufferWriter::GetName() const
+const char* FComputeBufferWriter::GetName() const
 {
 	return Detail->Name;
 }
diff --git a/Engine/Source/Programs/Horde/Samples/RemoteWorkerCpp/ComputeBuffer.h b/Engine/Source/Programs/Horde/Samples/RemoteWorkerCpp/ComputeBuffer.h
--- a/Engine/Source/Programs/Horde/Samples/RemoteWorkerCpp/ComputeBuffer.h
+++ b/Engine/Source/Programs/Horde/Samples/RemoteWorkerCpp/ComputeBuffer.h
@@ -36,7 +36,7 @@ public:
 	// Parameters for constructing a compute buffer
 	struct FParams
 	{
-		const wchar_t* Name = nullptr; // When creating a buffer, we should not normally set the name. We can only *OPEN* buffers with a name. Usually we create buffers and attach them to channels.
+		const char* Name = nullptr; // When creating a buffer, we should not normally set the name. We can only *OPEN* buffers with a name. Usually we create buffers and attach them to channels.
 		int NumChunks = 2;
 		int ChunkLength = 64 * 1024;
 		int NumReaders = 1;
@@ -49,7 +49,7 @@ public:
 	bool CreateNew(const FParams& Params);
 
 	// Opens an existing shared memory buffer (typically from handles created in another process)
-	bool OpenExisting(const wchar_t* Name);
+	bool OpenExisting(const char* Name);
 
 	// Close the current buffer and release all allocated resources
 	void Close();
@@ -108,7 +108,7 @@ private:
 
 	std::shared_ptr<FComputeBufferReaderDetail> Detail;
 
-	const wchar_t* GetName() const;
+	const char* GetName() const;
 };
 
 //
@@ -145,5 +145,5 @@ private:
 
 	std::shared_ptr<FComputeBufferDetail> Detail;
 
-	const wchar_t* GetName() const;
+	const char* GetName() const;
 };
diff --git a/Engine/Source/Programs/Horde/Samples/RemoteWorkerCpp/ComputePlatform.cpp b/Engine/Source/Programs/Horde/Samples/RemoteWorkerCpp/ComputePlatform.cpp
--- a/Engine/Source/Programs/Horde/Samples/RemoteWorkerCpp/ComputePlatform.cpp
+++ b/Engine/Source/Programs/Horde/Samples/RemoteWorkerCpp/ComputePlatform.cpp
@@ -3,63 +3,157 @@
 #include "ComputePlatform.h"
 #include <assert.h>
 #include <wchar.h>
-
-#ifdef _MSC_VER
-#include <Windows.h>
-#undef GetEnvironmentVariable
+#include <bit>
+#include <algorithm>
+#include <iostream>
+
+#if UE_COMPUTE_PLATFORM_WINDOWS
+	#include <Windows.h>
+	#undef min
+	#undef max
+	#undef GetEnvironmentVariable
+#else
+	#include <semaphore.h>
+	#include <unistd.h>
+	#include <atomic>
+	#include <sys/mman.h>
+	#include <sys/stat.h>
+	#include <time.h>
+	#include <fcntl.h>
 #endif
 
-/////////////////////////////////////////////////// 
+///////////////////////////////////////////////////
 
-FComputeManualResetEvent::FComputeManualResetEvent()
+FComputeEvent::FComputeEvent()
 	: Handle(nullptr)
 {
 }
 
-FComputeManualResetEvent::~FComputeManualResetEvent()
+FComputeEvent::~FComputeEvent()
 {
 	Close();
 }
 
-bool FComputeManualResetEvent::Create(const wchar_t* Name)
+bool FComputeEvent::Create(const char* Name)
 {
 	Close();
 
-	Handle = CreateEventW(NULL, TRUE, FALSE, Name);
+#if UE_COMPUTE_PLATFORM_WINDOWS
+	Handle = CreateEventA(NULL, FALSE, FALSE, Name);
 	return Handle != nullptr;
+#else
+	sem_t* Value = sem_open(Name, O_CREAT | O_EXCL, 0666, 1);
+	if(Value != SEM_FAILED)
+	{
+		Handle = Value;
+		return true;
+	}
+	return false;
+#endif
 }
 
-bool FComputeManualResetEvent::OpenExisting(const wchar_t* Name)
+bool FComputeEvent::OpenExisting(const char* Name)
 {
 	Close();
 
-	Handle = OpenEventW(SYNCHRONIZE | EVENT_MODIFY_STATE, TRUE, Name);
+#if UE_COMPUTE_PLATFORM_WINDOWS
+	Handle = OpenEventA(SYNCHRONIZE | EVENT_MODIFY_STATE, TRUE, Name);
 	return Handle != nullptr;
+#else
+	sem_t* Value = sem_open(Name, 0);
+	if(Value != SEM_FAILED)
+	{
+		Handle = Value;
+		return true;
+	}
+	return false;
+#endif
 }
 
-void FComputeManualResetEvent::Close()
+void FComputeEvent::Close()
 {
+#if UE_COMPUTE_PLATFORM_WINDOWS
 	if (Handle != nullptr)
 	{
 		CloseHandle(Handle);
 		Handle = nullptr;
 	}
+#else
+	if (Handle != nullptr)
+	{
+		sem_close((sem_t*)Handle);
+		Handle = nullptr;
+	}
+#endif
 }
 
-void FComputeManualResetEvent::Set()
+void FComputeEvent::Signal()
 {
+#if UE_COMPUTE_PLATFORM_WINDOWS
 	SetEvent(Handle);
+#else
+	sem_post((sem_t*)Handle);
+#endif
 }
 
-void FComputeManualResetEvent::Reset()
-{
-	ResetEvent(Handle);
-}
-
-bool FComputeManualResetEvent::Wait(int timeoutMs)
+bool FComputeEvent::Wait(int timeoutMs)
 {
+#if UE_COMPUTE_PLATFORM_WINDOWS
 	DWORD WaitParam = (timeoutMs < 0) ? INFINITE : (DWORD)timeoutMs;
 	return WaitForSingleObject(Handle, WaitParam) != WAIT_TIMEOUT;
+#else
+	if (timeoutMs == -1)
+	{
+		return sem_wait((sem_t*)Handle) == 0;
+	}
+
+	if (sem_trywait((sem_t*)Handle) == 0)
+	{
+		return true;
+	}
+
+	if(timeoutMs == 0)
+	{
+		return false;
+	}
+		
+	struct timespec ts;
+	if (clock_gettime(CLOCK_REALTIME, &ts) == -1)
+	{
+		assert(false);
+		return false;
+	}
+	
+	const long long NsPerSec = 1'000'000'000;
+	const long long NsPerMs = NsPerSec / 1000;
+	
+	long long newNs = (long long)ts.tv_nsec + (timeoutMs * NsPerMs);
+	ts.tv_nsec = newNs % NsPerSec;
+	ts.tv_sec += newNs / NsPerSec;
+	
+#if UE_COMPUTE_PLATFORM_MAC
+	for(;;)
+	{
+		if (sem_trywait((sem_t*)Handle) == 0)
+		{
+			return true;
+		}
+		
+		struct timespec currentTs;
+		if (clock_gettime(CLOCK_REALTIME, &currentTs) != 0 || currentTs.tv_sec > ts.tv_sec || (currentTs.tv_sec == ts.tv_sec && currentTs.tv_nsec > ts.tv_nsec))
+		{
+			return false;
+		}
+		
+		struct timespec sleepTs = { 0, };
+		sleepTs.tv_nsec = 100 * NsPerMs;
+		
+		nanosleep(&sleepTs, nullptr);
+	}
+#else
+	return sem_timedwait((sem_t*)Handle, &ts) == 0;
+#endif
+#endif
 }
 
 /////////////////////////////////////////////////// 
@@ -67,6 +161,8 @@ bool FComputeManualResetEvent::Wait(int timeoutMs)
 FComputeMemoryMappedFile::FComputeMemoryMappedFile()
 	: Handle(nullptr)
 	, Pointer(nullptr)
+	, MappedSize(0)
+	, OwnerName(nullptr)
 {
 }
 
@@ -75,14 +171,15 @@ FComputeMemoryMappedFile::~FComputeMemoryMappedFile()
 	Close();
 }
 
-bool FComputeMemoryMappedFile::Create(const wchar_t* Name, long long Capacity)
+bool FComputeMemoryMappedFile::Create(const char* Name, long long Capacity)
 {
 	Close();
 
+#if UE_COMPUTE_PLATFORM_WINDOWS
 	LARGE_INTEGER LargeInteger;
 	LargeInteger.QuadPart = Capacity;
 
-	Handle = CreateFileMappingW(INVALID_HANDLE_VALUE, NULL, PAGE_READWRITE, LargeInteger.HighPart, LargeInteger.LowPart, Name);
+	Handle = CreateFileMappingA(INVALID_HANDLE_VALUE, NULL, PAGE_READWRITE, LargeInteger.HighPart, LargeInteger.LowPart, Name);
 	if (Handle == nullptr)
 	{
 		return false;
@@ -95,13 +192,41 @@ bool FComputeMemoryMappedFile::Create(const wchar_t* Name, long long Capacity)
 	}
 
 	return true;
+#else
+	int Fd = shm_open(Name, O_CREAT | O_EXCL | O_RDWR, 0666);
+	if(Fd < 0)
+	{
+		std::cerr << "Unable to create shared memory object '" << Name << "' (errno=" << errno << ")" << std::endl;
+		return false;
+	}
+
+	Handle = (void*)(size_t)Fd;
+	MappedSize = Capacity;
+	OwnerName = strdup(Name);
+
+	if(ftruncate(Fd, MappedSize) < 0)
+	{
+		std::cerr << "Unable to update size of shared memory object '" << Name << "' to " << MappedSize << " (errno=" << errno << ")" << std::endl;
+		return false;
+	}
+
+	Pointer = mmap(nullptr, MappedSize, PROT_READ | PROT_WRITE, MAP_SHARED, Fd, 0);
+	if(Pointer == MAP_FAILED)
+	{
+		std::cerr << "Unable to map shared memory object '" << Name << " (errno=" << errno << ")" << std::endl;
+		return false;
+	}
+	
+	return true;
+#endif
 }
 
-bool FComputeMemoryMappedFile::OpenExisting(const wchar_t* Name)
+bool FComputeMemoryMappedFile::OpenExisting(const char* Name)
 {
 	Close();
 
-	Handle = OpenFileMappingW(FILE_MAP_ALL_ACCESS, TRUE, Name);
+#if UE_COMPUTE_PLATFORM_WINDOWS
+	Handle = OpenFileMappingA(FILE_MAP_ALL_ACCESS, TRUE, Name);
 	if (Handle == nullptr)
 	{
 		return false;
@@ -114,10 +239,34 @@ bool FComputeMemoryMappedFile::OpenExisting(const wchar_t* Name)
 	}
 
 	return true;
+#else
+	int Fd = shm_open(Name, O_RDWR, 0666);
+	if(Fd < 0)
+	{
+		std::cerr << "Unable to open shared memory object '" << Name << "' (errno=" << errno << ")" << std::endl;
+		return false;
+	}
+
+	Handle = (void*)(size_t)Fd;
+
+	struct stat st;
+	fstat(Fd, &st);
+	MappedSize = st.st_size;
+
+	Pointer = mmap(nullptr, MappedSize, PROT_READ | PROT_WRITE, MAP_SHARED, Fd, 0);
+	if(Pointer == MAP_FAILED)
+	{
+		std::cerr << "Unable to map shared memory object '" << Name << " (errno=" << errno << ")" << std::endl;
+		return false;
+	}
+	
+	return true;
+#endif
 }
 
 void FComputeMemoryMappedFile::Close()
 {
+#if UE_COMPUTE_PLATFORM_WINDOWS
 	if (Pointer != nullptr)
 	{
 		UnmapViewOfFile(Pointer);
@@ -129,123 +278,201 @@ void FComputeMemoryMappedFile::Close()
 		CloseHandle(Handle);
 		Handle = nullptr;
 	}
+#else
+	if(OwnerName != nullptr)
+	{
+		shm_unlink(OwnerName);
+		OwnerName = nullptr;
+	}
+	
+	if(Pointer != nullptr)
+	{
+		munmap(Pointer, MappedSize);
+		Pointer = nullptr;
+	}
+
+	int Fd = (int)(size_t)Handle;
+	if(Fd >= 0)
+	{
+		close(Fd);
+		Handle = nullptr;
+	}
+#endif
 }
 
 void* FComputeMemoryMappedFile::GetPointer() const
 {
+#if UE_COMPUTE_PLATFORM_WINDOWS
 	return Pointer;
+#else
+	return (unsigned char*)Pointer + 16;
+#endif
 }
 
 /////////////////////////////////////////////////// 
 
-bool FComputePlatform::GetEnvironmentVariable(const wchar_t* Name, wchar_t* Buffer, size_t BufferLen)
+bool FComputePlatform::GetEnvironmentVariable(const char* Name, char* Buffer, size_t BufferLen)
 {
-	int Length = GetEnvironmentVariableW(Name, Buffer, (DWORD)BufferLen);
+#if UE_COMPUTE_PLATFORM_WINDOWS
+	int Length = GetEnvironmentVariableA(Name, Buffer, (DWORD)BufferLen);
 	return Length > 0 && Length < BufferLen;
+#else
+	char* Value = getenv(Name);
+	if(Value != nullptr)
+	{
+		FComputePlatform::Strcpy(Buffer, BufferLen, Value);
+		return true;
+	}
+	return false;
+#endif
 }
 
-void FComputePlatform::CreateUniqueName(wchar_t* NameBuffer, size_t NameBufferLen)
+void FComputePlatform::CreateUniqueName(char* NameBuffer, size_t NameBufferLen)
 {
 	static long Counter = 0;
 
+#if UE_COMPUTE_PLATFORM_WINDOWS
 	DWORD Pid = GetCurrentProcessId();
 	ULONGLONG TickCount = GetTickCount64();
-	swprintf(NameBuffer, NameBufferLen, L"Local\\COMPUTE_%u_%llu_%lu", Pid, TickCount, AtomicIncrement(&Counter));
+	snprintf(NameBuffer, NameBufferLen, "Local\\COMPUTE_%u_%llu_%lu", Pid, TickCount, AtomicIncrement(&Counter));
+#else
+	struct timespec ts;
+	clock_gettime(CLOCK_MONOTONIC, &ts);
+	snprintf(NameBuffer, NameBufferLen, "/COMPUTE_%u_%zu_%zu_%lu", getpid(), (size_t)ts.tv_sec, (size_t)ts.tv_nsec, AtomicIncrement(&Counter));
+#endif
 }
 
 unsigned int FComputePlatform::FloorLog2(unsigned int Value)
 {
-	// Use BSR to return the log2 of the integer
-	// return 0 if value is 0
-	unsigned long BitIndex;
-	return _BitScanReverse(&BitIndex, Value) ? BitIndex : 0;
+	return std::max<unsigned int>(0, 31 - std::countl_zero(Value));
 }
 
 unsigned int FComputePlatform::CountLeadingZeros(unsigned int Value)
 {
-	// return 32 if value is zero
-	unsigned long BitIndex;
-	_BitScanReverse64(&BitIndex, (unsigned long long)(Value) * 2 + 1);
-	return 32 - BitIndex;
+	return std::countl_zero(Value);
 }
 
-size_t FComputePlatform::Utf8ToWchar(const char* Source, size_t SourceLen, wchar_t* Dest, size_t DestMaxLen)
+void FComputePlatform::Strcpy(char* Dest, size_t DestLen, const char* Source)
 {
-	size_t DecodedLen = MultiByteToWideChar(CP_UTF8, 0, Source, (int)SourceLen, Dest, (int)DestMaxLen);
-	Dest[DecodedLen] = 0;
-	return DecodedLen;
+	size_t Length = std::min(strlen(Source), DestLen - 1);
+	memcpy(Dest, Source, Length);
+	Dest[Length] = 0;
 }
 
-size_t FComputePlatform::WcharToUtf8(const wchar_t* Source, size_t SourceLen, char* Dest, size_t DestMaxLen)
+int FComputePlatform::Stricmp(const char* A, const char* B)
 {
-	if (SourceLen == 0)
-	{
-		return 0;
-	}
-
-	int Result = WideCharToMultiByte(CP_UTF8, 0, Source, (int)SourceLen, Dest, (int)DestMaxLen, nullptr, nullptr);
-	assert(Result > 0);
-
-	return (size_t)Result;
+#if UE_COMPUTE_PLATFORM_WINDOWS
+	return _stricmp(A, B);
+#else
+	return strcasecmp(A, B);
+#endif
 }
 
 long long FComputePlatform::AtomicRead64(const volatile long long* Ptr)
 {
+#if UE_COMPUTE_PLATFORM_WINDOWS
 	return InterlockedCompareExchange64(const_cast<volatile long long*>(Ptr), 0, 0);
+#else
+	return std::atomic_load((std::atomic<long long>*)Ptr);
+#endif
 }
 
 void FComputePlatform::AtomicWrite64(volatile long long* Ptr, long long Value)
 {
+#if UE_COMPUTE_PLATFORM_WINDOWS
 	InterlockedExchange64(Ptr, Value);
+#else
+	return std::atomic_store((std::atomic<long long>*)Ptr, Value);
+#endif
 }
 
 long FComputePlatform::AtomicIncrement(volatile long* Ptr)
 {
+#if UE_COMPUTE_PLATFORM_WINDOWS
 	return InterlockedIncrement(Ptr);
+#else
+	return std::atomic_fetch_add((std::atomic<long>*)Ptr, 1) + 1;
+#endif
 }
 
 long FComputePlatform::AtomicDecrement(volatile long* Ptr)
 {
+#if UE_COMPUTE_PLATFORM_WINDOWS
 	return InterlockedDecrement(Ptr);
+#else
+	return std::atomic_fetch_add((std::atomic<long>*)Ptr, -1) - 1;
+#endif
 }
 
 long long FComputePlatform::AtomicAdd64(volatile long long* Ptr, long long Value)
 {
+#if UE_COMPUTE_PLATFORM_WINDOWS
 	return InterlockedAdd64(Ptr, Value);
+#else
+	return std::atomic_fetch_add((std::atomic<long long>*)Ptr, Value) + Value;
+#endif
 }
 
 long FComputePlatform::AtomicAnd(volatile long* Ptr, long Value)
 {
+#if UE_COMPUTE_PLATFORM_WINDOWS
 	return InterlockedAnd(Ptr, Value);
+#else
+	return std::atomic_fetch_and((std::atomic<long>*)Ptr, Value);
+#endif
 }
 
 long long FComputePlatform::AtomicAnd64(volatile long long* Ptr, long long Value)
 {
+#if UE_COMPUTE_PLATFORM_WINDOWS
 	return InterlockedAnd64(Ptr, Value);
+#else
+	return std::atomic_fetch_and((std::atomic<long long>*)Ptr, Value);
+#endif
 }
 
 long long FComputePlatform::AtomicOr64(volatile long long* Ptr, long long Value)
 {
+#if UE_COMPUTE_PLATFORM_WINDOWS
 	return InterlockedOr64(Ptr, Value);
+#else
+	return std::atomic_fetch_or((std::atomic<long long>*)Ptr, Value);
+#endif
 }
 
 long long FComputePlatform::AtomicIncrement64(volatile long long* Ptr)
 {
+#if UE_COMPUTE_PLATFORM_WINDOWS
 	return InterlockedIncrement64(Ptr);
+#else
+	return std::atomic_fetch_add((std::atomic<long long>*)Ptr, 1) + 1;
+#endif
 }
 
 long long FComputePlatform::AtomicExchange64(volatile long long* Ptr, long long Exchange)
 {
+#if UE_COMPUTE_PLATFORM_WINDOWS
 	return InterlockedExchange64(Ptr, Exchange);
+#else
+	return std::atomic_exchange((std::atomic<long long>*)Ptr, Exchange);
+#endif
 }
 
 bool FComputePlatform::AtomicCompareExchange(volatile long* Ptr, long Exchange, long Comperand)
 {
+#if UE_COMPUTE_PLATFORM_WINDOWS
 	return InterlockedCompareExchange(Ptr, Exchange, Comperand) == Comperand;
+#else
+	return std::atomic_compare_exchange_strong((std::atomic<long>*)Ptr, &Comperand, Exchange);
+#endif
 }
 
 bool FComputePlatform::AtomicCompareExchange64(volatile long long* Ptr, long long Exchange, long long Comperand)
 {
+#if UE_COMPUTE_PLATFORM_WINDOWS
 	return InterlockedCompareExchange64(Ptr, Exchange, Comperand) == Comperand;
+#else
+	return std::atomic_compare_exchange_strong((std::atomic<long long>*)Ptr, &Comperand, Exchange);
+#endif
 }
 
diff --git a/Engine/Source/Programs/Horde/Samples/RemoteWorkerCpp/ComputePlatform.h b/Engine/Source/Programs/Horde/Samples/RemoteWorkerCpp/ComputePlatform.h
--- a/Engine/Source/Programs/Horde/Samples/RemoteWorkerCpp/ComputePlatform.h
+++ b/Engine/Source/Programs/Horde/Samples/RemoteWorkerCpp/ComputePlatform.h
@@ -2,29 +2,45 @@
 
 #pragma once
 
+#include <stdio.h>
+#include <stddef.h>
+#include <stdint.h>
+
+// Defines for the current platform
+#ifdef _MSC_VER
+	#define UE_COMPUTE_PLATFORM_WINDOWS 1
+	#define UE_COMPUTE_PLATFORM_MAC 0
+	#define UE_COMPUTE_PLATFORM_LINUX 0
+#elif defined(__APPLE__)
+	#define UE_COMPUTE_PLATFORM_WINDOWS 0
+	#define UE_COMPUTE_PLATFORM_MAC 1
+	#define UE_COMPUTE_PLATFORM_LINUX 0
+#else
+	#define UE_COMPUTE_PLATFORM_WINDOWS 0
+	#define UE_COMPUTE_PLATFORM_MAC 0
+	#define UE_COMPUTE_PLATFORM_LINUX 1
+#endif
+
 //
 // IPC manually-reset event object
 //
-class FComputeManualResetEvent
+class FComputeEvent
 {
 public:
-	FComputeManualResetEvent();
-	~FComputeManualResetEvent();
+	FComputeEvent();
+	~FComputeEvent();
 
 	// Creates a new event with the given name
-	bool Create(const wchar_t* Name);
+	bool Create(const char* Name);
 
 	// Opens an existing event created elsewhere
-	bool OpenExisting(const wchar_t* Name);
+	bool OpenExisting(const char* Name);
 
 	// Close the event and release its resources
 	void Close();
 
 	// Signal the event, releasing any waiters
-	void Set();
-
-	// Reset the event after a call to Set()
-	void Reset();
+	void Signal();
 
 	// Wait for the event to be signalled or timeout. Pass -1 for timeoutMs to wait infinitely.
 	bool Wait(int timeoutMs);
@@ -43,10 +59,10 @@ public:
 	~FComputeMemoryMappedFile();
 
 	// Creates a new memory mapped file with the given capacity
-	bool Create(const wchar_t* Name, long long Capacity);
+	bool Create(const char* Name, long long Capacity);
 
 	// Opens an existing memory mapped file
-	bool OpenExisting(const wchar_t* Name);
+	bool OpenExisting(const char* Name);
 
 	// Close the memory mapped file handle
 	void Close();
@@ -57,6 +73,8 @@ public:
 private:
 	void* Handle;
 	void* Pointer;
+	long long MappedSize;
+	char* OwnerName;
 };
 
 //
@@ -67,12 +85,12 @@ struct FComputePlatform
 	//
 	// General
 	//
-	// 
+
 	// Reads an environment variable
-	static bool GetEnvironmentVariable(const wchar_t* Name, wchar_t* Buffer, size_t BufferLen);
+	static bool GetEnvironmentVariable(const char* Name, char* Buffer, size_t BufferLen);
 
 	// Creates a unique object name
-	static void CreateUniqueName(wchar_t* NameBuffer, size_t NameBufferLen);
+	static void CreateUniqueName(char* NameBuffer, size_t NameBufferLen);
 
 	//
 	// Math
@@ -88,11 +106,11 @@ struct FComputePlatform
 	// Strings
 	//
 
-	// Translate a UTF8 string to a wchar_t string. Returns the length of the converted string in characters, even if the supplied buffer is not large enough to hold it.
-	static size_t Utf8ToWchar(const char* Source, size_t SourceLen, wchar_t* Dest, size_t DestMaxLen);
+	// Copy a string from one buffer to another, not exceeding the destination buffer size
+	static void Strcpy(char* Dest, size_t DestLen, const char* Source);
 
-	// Translate a wchar_t string to a UTF8 string. Returns the length of the converted string in characters, even if the supplied buffer is not large enough to hold it.
-	static size_t WcharToUtf8(const wchar_t* Source, size_t SourceLen, char* Dest, size_t DestMaxLen);
+	// Perform a case-insensitive comparison of two strings
+	static int Stricmp(const char* A, const char* B);
 
 	//
 	// Atomics
diff --git a/Engine/Source/Programs/Horde/Samples/RemoteWorkerCpp/ComputeSocket.cpp b/Engine/Source/Programs/Horde/Samples/RemoteWorkerCpp/ComputeSocket.cpp
--- a/Engine/Source/Programs/Horde/Samples/RemoteWorkerCpp/ComputeSocket.cpp
+++ b/Engine/Source/Programs/Horde/Samples/RemoteWorkerCpp/ComputeSocket.cpp
@@ -46,7 +46,7 @@ FComputeChannel FComputeSocket::CreateChannel(int ChannelId, FComputeBuffer Recv
 
 //////////////////////////////////////////////////////
 
-const wchar_t* const FWorkerComputeSocket::IpcEnvVar = L"UE_HORDE_COMPUTE_IPC";
+const char* const FWorkerComputeSocket::IpcEnvVar = "UE_HORDE_COMPUTE_IPC";
 
 enum class FWorkerComputeSocket::EMessageType
 {
@@ -65,7 +65,7 @@ FWorkerComputeSocket::~FWorkerComputeSocket()
 
 bool FWorkerComputeSocket::Open()
 {
-	wchar_t EnvVar[FComputeBuffer::MaxNameLength];
+	char EnvVar[FComputeBuffer::MaxNameLength];
 	if (!FComputePlatform::GetEnvironmentVariable(IpcEnvVar, EnvVar, sizeof(EnvVar) / sizeof(EnvVar[0])))
 	{
 		return false;
@@ -74,7 +74,7 @@ bool FWorkerComputeSocket::Open()
 	return Open(EnvVar);
 }
 
-bool FWorkerComputeSocket::Open(const wchar_t* CommandBufferName)
+bool FWorkerComputeSocket::Open(const char* CommandBufferName)
 {
 	return CommandBuffer.OpenExisting(CommandBufferName);
 }
@@ -94,7 +94,7 @@ void FWorkerComputeSocket::AttachSendBuffer(int ChannelId, FComputeBufferReader
 	AttachBuffer(ChannelId, EMessageType::AttachSendBuffer, Reader.GetName());
 }
 
-void FWorkerComputeSocket::AttachBuffer(int ChannelId, EMessageType Type, const wchar_t* Name)
+void FWorkerComputeSocket::AttachBuffer(int ChannelId, EMessageType Type, const char* Name)
 {
 	FComputeBufferWriter& Writer = CommandBuffer.GetWriter();
 	unsigned char* Data = Writer.WaitToWrite(1024);
@@ -125,7 +125,7 @@ void FWorkerComputeSocket::RunServer(FComputeBufferReader& CommandBufferReader,
 				unsigned int ChannelId;
 				Len += ReadVarUInt(Message + Len, &ChannelId);
 
-				wchar_t Name[FComputeBuffer::MaxNameLength];
+				char Name[FComputeBuffer::MaxNameLength];
 				Len += ReadString(Message + Len, Name, FComputeBuffer::MaxNameLength);
 
 				FComputeBuffer Buffer;
@@ -144,7 +144,7 @@ void FWorkerComputeSocket::RunServer(FComputeBufferReader& CommandBufferReader,
 				unsigned int ChannelId;
 				Len += ReadVarUInt(Message + Len, &ChannelId);
 
-				wchar_t Name[FComputeBuffer::MaxNameLength];
+				char Name[FComputeBuffer::MaxNameLength];
 				Len += ReadString(Message + Len, Name, FComputeBuffer::MaxNameLength);
 
 				FComputeBuffer Buffer;
@@ -190,12 +190,12 @@ size_t FWorkerComputeSocket::ReadVarUInt(const unsigned char* Pos, unsigned int*
 	return ByteCount;
 }
 
-size_t FWorkerComputeSocket::ReadString(const unsigned char* Pos, wchar_t* OutText, size_t OutTextMaxLen)
+size_t FWorkerComputeSocket::ReadString(const unsigned char* Pos, char* OutText, size_t OutTextMaxLen)
 {
 	unsigned int TextLen;
-	size_t Len = ReadVarUInt(Pos, &TextLen);
 
-	FComputePlatform::Utf8ToWchar((const char*)Pos + Len, TextLen, OutText, OutTextMaxLen);
+	size_t Len = ReadVarUInt(Pos, &TextLen);
+	FComputePlatform::Strcpy(OutText, OutTextMaxLen, (const char*)Pos + Len);
 
 	return Len + TextLen;
 }
@@ -209,27 +209,26 @@ size_t FWorkerComputeSocket::WriteVarUInt(unsigned char* Pos, unsigned int Value
 	unsigned char* OutBytes = Pos + ByteCount - 1;
 	switch (ByteCount - 1)
 	{
-	case 4: *OutBytes-- = (unsigned char)(Value); Value >>= 8;
-	case 3: *OutBytes-- = (unsigned char)(Value); Value >>= 8;
-	case 2: *OutBytes-- = (unsigned char)(Value); Value >>= 8;
-	case 1: *OutBytes-- = (unsigned char)(Value); Value >>= 8;
-	default: break;
+	case 4: *OutBytes-- = (unsigned char)(Value); Value >>= 8; [[fallthrough]];
+	case 3: *OutBytes-- = (unsigned char)(Value); Value >>= 8; [[fallthrough]];
+	case 2: *OutBytes-- = (unsigned char)(Value); Value >>= 8; [[fallthrough]];
+	case 1: *OutBytes-- = (unsigned char)(Value); Value >>= 8; [[fallthrough]];
+	default: 
+		break;
 	}
 	*OutBytes = (unsigned char)(0xff << (9 - ByteCount)) | (unsigned char)(Value);
 
 	return ByteCount;
 }
 
-size_t FWorkerComputeSocket::WriteString(unsigned char* Pos, const wchar_t* Text)
+size_t FWorkerComputeSocket::WriteString(unsigned char* Pos, const char* Text)
 {
-	size_t TextLen = wcslen(Text);
+	size_t TextLen = strlen(Text);
 
-	size_t EncodedLen = FComputePlatform::WcharToUtf8(Text, TextLen, nullptr, 0);
+	size_t Len = WriteVarUInt(Pos, (int)TextLen);
+	memcpy((char*)Pos + Len, Text, TextLen);
 
-	size_t Len = WriteVarUInt(Pos, (int)EncodedLen);
-	FComputePlatform::WcharToUtf8(Text, TextLen, (char*)Pos + Len, EncodedLen);
-
-	return Len + EncodedLen;
+	return Len + TextLen;
 }
 
 //////////////////////////////////////////////////////
diff --git a/Engine/Source/Programs/Horde/Samples/RemoteWorkerCpp/ComputeSocket.h b/Engine/Source/Programs/Horde/Samples/RemoteWorkerCpp/ComputeSocket.h
--- a/Engine/Source/Programs/Horde/Samples/RemoteWorkerCpp/ComputeSocket.h
+++ b/Engine/Source/Programs/Horde/Samples/RemoteWorkerCpp/ComputeSocket.h
@@ -39,7 +39,7 @@ public:
 class FWorkerComputeSocket final : public FComputeSocket
 {
 public:
-	static const wchar_t* const IpcEnvVar;
+	static const char* const IpcEnvVar;
 
 	FWorkerComputeSocket();
 	~FWorkerComputeSocket();
@@ -48,7 +48,7 @@ public:
 	bool Open();
 
 	// Opens a connection to the agent process using a specific command buffer name
-	bool Open(const wchar_t* CommandBufferName);
+	bool Open(const char* CommandBufferName);
 
 	// Close the current connection
 	void Close();
@@ -67,13 +67,13 @@ private:
 
 	FComputeBuffer CommandBuffer;
 
-	void AttachBuffer(int ChannelId, EMessageType Type, const wchar_t* Name);
+	void AttachBuffer(int ChannelId, EMessageType Type, const char* Name);
 
 	static size_t ReadVarUInt(const unsigned char* Pos, unsigned int* OutValue);
-	static size_t ReadString(const unsigned char* Pos, wchar_t* OutText, size_t OutTextMaxLen);
+	static size_t ReadString(const unsigned char* Pos, char* OutText, size_t OutTextMaxLen);
 
 	static size_t WriteVarUInt(unsigned char* Pos, unsigned int Value);
-	static size_t WriteString(unsigned char* Pos, const wchar_t* Text);
+	static size_t WriteString(unsigned char* Pos, const char* Text);
 };
 
 //
diff --git a/Engine/Source/Programs/Horde/Samples/RemoteWorkerCpp/ComputeTransport.cpp b/Engine/Source/Programs/Horde/Samples/RemoteWorkerCpp/ComputeTransport.cpp
--- a/Engine/Source/Programs/Horde/Samples/RemoteWorkerCpp/ComputeTransport.cpp
+++ b/Engine/Source/Programs/Horde/Samples/RemoteWorkerCpp/ComputeTransport.cpp
@@ -1,6 +1,11 @@
 // Copyright Epic Games, Inc. All Rights Reserved.
 
 #include "ComputeTransport.h"
+#include <algorithm>
+
+FComputeTransport::~FComputeTransport()
+{
+}
 
 bool FComputeTransport::SendMessage(const void* Data, size_t Size)
 {
diff --git a/Engine/Source/Programs/Horde/Samples/RemoteWorkerCpp/ComputeTransport.h b/Engine/Source/Programs/Horde/Samples/RemoteWorkerCpp/ComputeTransport.h
--- a/Engine/Source/Programs/Horde/Samples/RemoteWorkerCpp/ComputeTransport.h
+++ b/Engine/Source/Programs/Horde/Samples/RemoteWorkerCpp/ComputeTransport.h
@@ -10,6 +10,8 @@
 class FComputeTransport
 {
 public:
+	virtual ~FComputeTransport();
+	
 	// Sends data to the remote
 	virtual size_t Send(const void* Data, size_t Size) = 0;
 
diff --git a/Engine/Source/Programs/Horde/Samples/RemoteWorkerCpp/Main.cpp b/Engine/Source/Programs/Horde/Samples/RemoteWorkerCpp/Main.cpp
--- a/Engine/Source/Programs/Horde/Samples/RemoteWorkerCpp/Main.cpp
+++ b/Engine/Source/Programs/Horde/Samples/RemoteWorkerCpp/Main.cpp
@@ -3,13 +3,14 @@
 #include "ComputeChannel.h"
 #include "ComputeBuffer.h"
 #include "ComputeSocket.h"
+#include "ComputePlatform.h"
 #include <iostream>
 #include <thread>
 #include <assert.h>
 
 int main(int argc, const char* argv[])
 {
-	if (argc >= 2 && !_stricmp(argv[1], "-Test"))
+	if (argc >= 2 && !FComputePlatform::Stricmp(argv[1], "-Test"))
 	{
 		void RunTests();
 		RunTests();
@@ -56,7 +57,4 @@ int main(int argc, const char* argv[])
 			Length = 0;
 		}
 	}
-
-	Channel.MarkComplete();
-	return 0;
 }
diff --git a/Engine/Source/Programs/Horde/Samples/RemoteWorkerCpp/Tests.cpp b/Engine/Source/Programs/Horde/Samples/RemoteWorkerCpp/Tests.cpp
--- a/Engine/Source/Programs/Horde/Samples/RemoteWorkerCpp/Tests.cpp
+++ b/Engine/Source/Programs/Horde/Samples/RemoteWorkerCpp/Tests.cpp
@@ -11,9 +11,11 @@ void RunTests()
 {
 	void ComputeBufferTest();
 	ComputeBufferTest();
+	std::cout << "ComputeBufferTest ok" << std::endl;
 
 	void ComputeSocketTest();
 	ComputeSocketTest();
+	std::cout << "ComputeSocketTest ok" << std::endl;
 }
 
 void ComputeBufferTest()
