Horde: Unifying C#/C++ implementations of remote execution code.

- Handle ref counting in the same way for each implementation.
- Writers now need to be created explicitly, and are marked complete when their last writer is released.
- Added -InProc command line argument to sample app, which runs C# assemblies in the same process when possible.

diff --git a/Engine/Source/Programs/Horde/Horde.Agent/Commands/Compute/ComputeCommand.cs b/Engine/Source/Programs/Horde/Horde.Agent/Commands/Compute/ComputeCommand.cs
--- a/Engine/Source/Programs/Horde/Horde.Agent/Commands/Compute/ComputeCommand.cs
+++ b/Engine/Source/Programs/Horde/Horde.Agent/Commands/Compute/ComputeCommand.cs
@@ -34,6 +34,9 @@ namespace Horde.Agent.Commands.Compute
 		[CommandLine("-Loopback")]
 		public bool Loopback { get; set; }
 
+		[CommandLine("-InProc")]
+		public bool InProc { get; set; }
+
 		[CommandLine("-Sandbox=")]
 		public DirectoryReference SandboxDir { get; set; } = DirectoryReference.Combine(Program.DataDir, "Sandbox");
 
@@ -81,7 +84,7 @@ namespace Horde.Agent.Commands.Compute
 		{
 			if (Local)
 			{
-				return new LocalComputeClient(2000, SandboxDir, logger);
+				return new LocalComputeClient(2000, SandboxDir, InProc, logger);
 			}
 			else if (Loopback)
 			{
@@ -143,7 +146,7 @@ namespace Horde.Agent.Commands.Compute
 		{
 			DirectoryReference sandboxDir = DirectoryReference.Combine(Program.DataDir, "Sandbox");
 
-			AgentMessageHandler worker = new AgentMessageHandler(sandboxDir, memoryCache, logger);
+			AgentMessageHandler worker = new AgentMessageHandler(sandboxDir, memoryCache, false, logger);
 			await worker.RunAsync(socket, cancellationToken);
 		}
 	}
diff --git a/Engine/Source/Programs/Horde/Horde.Agent/Leases/Handlers/ComputeHandler.cs b/Engine/Source/Programs/Horde/Horde.Agent/Leases/Handlers/ComputeHandler.cs
--- a/Engine/Source/Programs/Horde/Horde.Agent/Leases/Handlers/ComputeHandler.cs
+++ b/Engine/Source/Programs/Horde/Horde.Agent/Leases/Handlers/ComputeHandler.cs
@@ -103,7 +103,7 @@ namespace Horde.Agent.Leases.Handlers
 							{
 								DirectoryReference.CreateDirectory(sandboxDir);
 
-								AgentMessageHandler worker = new AgentMessageHandler(sandboxDir, _memoryCache, _logger);
+								AgentMessageHandler worker = new AgentMessageHandler(sandboxDir, _memoryCache, false, _logger);
 								await worker.RunAsync(socket, cts.Token);
 								await socket.CloseAsync(cts.Token);
 								return LeaseResult.Success;
diff --git a/Engine/Source/Programs/Horde/Samples/RemoteClient/RemoteClientApp.cs b/Engine/Source/Programs/Horde/Samples/RemoteClient/RemoteClientApp.cs
--- a/Engine/Source/Programs/Horde/Samples/RemoteClient/RemoteClientApp.cs
+++ b/Engine/Source/Programs/Horde/Samples/RemoteClient/RemoteClientApp.cs
@@ -1,8 +1,6 @@
 // Copyright Epic Games, Inc. All Rights Reserved.
 
-using System.Buffers;
 using System.Buffers.Binary;
-using System.IO.Pipes;
 using System.Net.Http.Headers;
 using System.Reflection;
 using EpicGames.Core;
@@ -10,8 +8,6 @@ using EpicGames.Horde.Common;
 using EpicGames.Horde.Compute;
 using EpicGames.Horde.Compute.Buffers;
 using EpicGames.Horde.Compute.Clients;
-using EpicGames.Horde.Compute.Transports;
-using EpicGames.Horde.Storage;
 using EpicGames.Horde.Storage.Backends;
 using EpicGames.Horde.Storage.Bundles;
 using EpicGames.Horde.Storage.Nodes;
@@ -32,6 +28,9 @@ namespace RemoteClient
 		[CommandLine("-Condition=")]
 		public string? Condition { get; set; }
 
+		[CommandLine]
+		public bool InProc { get; set; }
+
 		[CommandLine("-Cpp")]
 		public bool UseCppWorker { get; set; }
 	}
@@ -108,7 +107,9 @@ namespace RemoteClient
 				}
 
 				// Run the task remotely in the background and echo the output to the console
-				await using BackgroundTask tickTask = BackgroundTask.StartNew(ctx => WriteNumbersAsync(lease.Socket, logger, ctx));
+				using ComputeChannel childProcessChannel = lease.Socket.CreateChannel(ChildProcessChannelId);
+				await using BackgroundTask tickTask = BackgroundTask.StartNew(ctx => WriteNumbersAsync(childProcessChannel, logger, ctx));
+
 				await using (AgentManagedProcess process = await channel.ExecuteAsync(executable, arguments, null, null))
 				{
 					string? line;
@@ -121,15 +122,12 @@ namespace RemoteClient
 			await lease.CloseAsync();
 		}
 		
-		static async Task WriteNumbersAsync(RemoteComputeSocket socket, ILogger logger, CancellationToken cancellationToken)
+		static async Task WriteNumbersAsync(ComputeChannel channel, ILogger logger, CancellationToken cancellationToken)
 		{
 			// Wait until the remote sends a message indicating that it's ready
-			using (PooledBuffer recvBuffer = new PooledBuffer(1, 20))
+			if (!await channel.Reader.WaitToReadAsync(1, cancellationToken))
 			{
-				socket.AttachRecvBuffer(ChildProcessChannelId, recvBuffer);
-			
-				using ComputeBufferReader reader = recvBuffer.CreateReader();
-				await reader.WaitToReadAsync(1, cancellationToken);
+				throw new NotImplementedException();
 			}
 
 			// Write data to the child process channel. The remote server will echo them back to us as it receives them, then exit when the channel is complete/closed.
@@ -139,11 +137,11 @@ namespace RemoteClient
 				cancellationToken.ThrowIfCancellationRequested();
 				logger.LogInformation("Writing value: {Value}", idx);
 				BinaryPrimitives.WriteInt32LittleEndian(buffer, idx);
-				await socket.SendAsync(ChildProcessChannelId, buffer, cancellationToken);
+				await channel.Writer.WriteAsync(buffer, cancellationToken);
 				await Task.Delay(1000, cancellationToken);
 			}
 
-			await socket.MarkCompleteAsync(ChildProcessChannelId, cancellationToken);
+			channel.MarkComplete();
 		}
 
 		static async Task RunBackgroundXorAsync(AgentMessageChannel channel)
@@ -170,7 +168,7 @@ namespace RemoteClient
 			if (options.Server == null)
 			{
 				DirectoryReference sandboxDir = DirectoryReference.Combine(new FileReference(Assembly.GetExecutingAssembly().Location).Directory, "Sandbox");
-				return new LocalComputeClient(2000, sandboxDir, logger);
+				return new LocalComputeClient(2000, sandboxDir, options.InProc, logger);
 			}
 			else
 			{
diff --git a/Engine/Source/Programs/Horde/Samples/RemoteWorkerCpp/ComputeBuffer.cpp b/Engine/Source/Programs/Horde/Samples/RemoteWorkerCpp/ComputeBuffer.cpp
--- a/Engine/Source/Programs/Horde/Samples/RemoteWorkerCpp/ComputeBuffer.cpp
+++ b/Engine/Source/Programs/Horde/Samples/RemoteWorkerCpp/ComputeBuffer.cpp
@@ -10,32 +10,37 @@
 
 struct FComputeBufferDetail
 {
+	static const int HeaderSize = (2 + FComputeBuffer::MaxReaders + FComputeBuffer::MaxChunks) * sizeof(long long);
+
 	enum class EWriteState : unsigned long long
 	{
 		// Writer has moved to the next chunk
-		MovedToNext = 0b00,
+		MovedToNext = 0,
 
 		// Chunk is still being appended to
-		Writing = 0b10,
+		Writing = 2,
 
 		// This chunk marks the end of the stream
-		Complete = 0b11,
+		Complete = 3,
 	};
 
 	union FChunkState
 	{
+	public:
 		long long Value;
 
+		FChunkState()
+			: Value(0)
+		{ }
+
 		FChunkState(long long InValue)
 			: Value(InValue)
-		{
-		}
+		{ }
 
 		// Constructor
 		FChunkState(EWriteState WriteState, int ReaderFlags, int Length)
 			: Value((unsigned long long)Length | ((unsigned long long)ReaderFlags << 31) | ((unsigned long long)WriteState << 62))
-		{
-		}
+		{ }
 
 		// Written length of this chunk
 		int GetLength() const { return (int)(Value & 0x7fffffff); }
@@ -61,100 +66,178 @@ struct FComputeBufferDetail
 	struct FChunkStatePtr
 	{
 	public:
+		FChunkStatePtr(volatile FChunkState* InChunkState)
+			: State(InChunkState)
+		{ }
+
 		// Read the state value from memory
 		FChunkState Get() const
 		{
-			return FChunkState(FComputePlatform::AtomicRead64(&State.Value));
+			return FChunkState(FComputePlatform::AtomicRead64(&State->Value));
 		}
 
 		// Read the state value from memory
 		void Set(FChunkState NewState)
 		{
-			FComputePlatform::AtomicWrite64(&State.Value, NewState.Value);
+			FComputePlatform::AtomicWrite64(&State->Value, NewState.Value);
 		}
 
 		// Attempt to update the chunk state
 		bool TryUpdate(FChunkState PrevState, FChunkState NextState)
 		{
-			return FComputePlatform::AtomicCompareExchange64(&State.Value, NextState.Value, PrevState.Value);
+			return FComputePlatform::AtomicCompareExchange64(&State->Value, NextState.Value, PrevState.Value);
 		}
 
 		// Append data to the chunk
 		void Append(long long Length)
 		{
-			FComputePlatform::AtomicAdd64(&State.Value, Length);
+			FComputePlatform::AtomicAdd64(&State->Value, Length);
 		}
 
 		// Mark this chunk as the end of the stream
 		void MarkComplete()
 		{
-			FComputePlatform::AtomicOr64(&State.Value, FChunkState(EWriteState::Complete, 0, 0).Value);
+			FComputePlatform::AtomicOr64(&State->Value, FChunkState(EWriteState::Complete, 0, 0).Value);
 		}
 
 		// Start reading the chunk with the given reader
 		void StartReading(int ReaderIdx)
 		{
-			FComputePlatform::AtomicOr64(&State.Value, FChunkState((EWriteState)0, 1 << ReaderIdx, 0).Value);
+			FComputePlatform::AtomicOr64(&State->Value, FChunkState((EWriteState)0, 1 << ReaderIdx, 0).Value);
 		}
 
 		// Clear the reader flag
 		void FinishReading(int ReaderIdx)
 		{
-			FComputePlatform::AtomicAnd64(&State.Value, ~FChunkState((EWriteState)0, 1 << ReaderIdx, 0).Value);
+			FComputePlatform::AtomicAnd64(&State->Value, ~FChunkState((EWriteState)0, 1 << ReaderIdx, 0).Value);
 		}
 
 		// Move to the next chunk
 		void FinishWriting()
 		{
-			FComputePlatform::AtomicAnd64(&State.Value, ~FChunkState(EWriteState::Writing, 0, 0).Value);
+			FComputePlatform::AtomicAnd64(&State->Value, ~FChunkState(EWriteState::Writing, 0, 0).Value);
 		}
 
 	private:
-		volatile FChunkState State;
+		volatile FChunkState* const State;
 	};
 
-	struct FReaderState
+	union FReaderState
 	{
-		const int ReaderIdx;
+		long long Value;
 
-		int Offset;
-		int ChunkIdx;
-		bool Detached;
+		FReaderState()
+			: Value(0)
+		{ }
 
-		FReaderState(int InReaderIdx)
-			: ReaderIdx(InReaderIdx)
-			, Offset(0)
-			, ChunkIdx(0)
-			, Detached(false)
+		FReaderState(long long InValue)
+			: Value(InValue)
 		{ }
+
+		FReaderState(int ChunkIdx, int Offset, int RefCount, bool Detached)
+			: Value((long long)(unsigned int)Offset | ((long long)(unsigned int)ChunkIdx << 32) | ((long long)(unsigned int)RefCount << 40) | (long long)(Detached ? (1ULL << 63) : 0))
+		{ }
+
+		int GetOffset() const
+		{
+			return (int)(Value & 0xffffffff);
+		}
+
+		int GetChunkIdx() const
+		{
+			return (int)((Value >> 32) & 0xff);
+		}
+
+		int GetRefCount() const
+		{
+			return (int)((Value >> 40) & 0x7fff);
+		}
+
+		bool IsDetached() const
+		{
+			return (Value & (1ULL << 63)) != 0;
+		}
+
+	private:
+		struct // For debugging purposes only; non-portable assumption of bitfield layout
+		{
+			unsigned long long Offset : 32;
+			unsigned long long ChunkIdx : 8;
+			unsigned long long RefCount : 23;
+			unsigned long long Detached : 1;
+		};
+	};
+
+	struct FReaderStatePtr
+	{
+	public:
+		FReaderStatePtr(volatile FReaderState* InReaderState)
+			: State(InReaderState)
+		{ }
+
+		// Read the state value from memory
+		FReaderState Get() const
+		{
+			return FReaderState(FComputePlatform::AtomicRead64(&State->Value));
+		}
+
+		// Read the state value from memory
+		void Set(FReaderState NewState)
+		{
+			FComputePlatform::AtomicWrite64(&State->Value, NewState.Value);
+		}
+
+		// Attempt to update the chunk state
+		bool TryUpdate(FReaderState PrevState, FReaderState NextState)
+		{
+			return FComputePlatform::AtomicCompareExchange64(&State->Value, NextState.Value, PrevState.Value);
+		}
+
+	private:
+		volatile FReaderState* const State;
 	};
 
 	union FWriterState
 	{
 		long long Value;
 
-		FWriterState(long long InValue) : Value(InValue) { }
-		FWriterState(int ChunkIdx, int ReaderFlags, bool HasWrapped) : Value(ChunkIdx | ((long long)ReaderFlags << 32) | (HasWrapped? (1ULL << 63) : 0)) { }
+		FWriterState()
+			: Value(0)
+		{ }
+
+		FWriterState(long long InValue) 
+			: Value(InValue) 
+		{ }
+		
+		FWriterState(int ChunkIdx, int ReaderFlags, int RefCount, bool HasWrapped) 
+			: Value((long long)(unsigned int)ChunkIdx | ((long long)(unsigned int)ReaderFlags << 32) | ((long long)(unsigned int)RefCount << 48) | (HasWrapped ? (1ULL << 63) : 0))
+		{ }
 
-		int GetReaderFlags() const { return ((unsigned long long)Value >> 32) & 0xffff; }
 		int GetChunkIdx() const { return Value & 0x7fffffff; }
+		int GetReaderFlags() const { return ((unsigned long long)Value >> 32) & 0xffff; }
+		int GetRefCount() const { return ((unsigned long long)Value >> 48) & 0x7fff; }
 		bool HasWrapped() const { return (Value & (1ULL << 63)) != 0; }
 
 	private:
 		struct // For debugging purposes only; non-portable assumption of bitfield layout
 		{
 			int ChunkIdx : 32;
-			int ReaderFlags : 31;
-			int Wrapped : 1;
+			int ReaderFlags : 8;
+			int RefCount : 23;
+			unsigned int Wrapped : 1;
 		};
 	};
 
 	struct FWriterStatePtr
 	{
 	public:
+		FWriterStatePtr(volatile FWriterState* InState)
+			: State(InState)
+		{ }
+
 		FWriterState Get() const
 		{
-			return FWriterState(FComputePlatform::AtomicRead64(&State.Value));
+			return FWriterState(FComputePlatform::AtomicRead64(&State->Value));
 		}
 
 		void Set(FWriterState State)
@@ -164,21 +247,52 @@ struct FComputeBufferDetail
 
 		bool TryUpdate(FWriterState PrevValue, FWriterState NextValue)
 		{
-			return FComputePlatform::AtomicCompareExchange64(&State.Value, NextValue.Value, PrevValue.Value);
+			return FComputePlatform::AtomicCompareExchange64(&State->Value, NextValue.Value, PrevValue.Value);
 		}
 
 	private:
-		volatile FWriterState State;
+		volatile FWriterState* State;
 	};
 
 	struct FHeader
 	{
-		int NumReaders;
-		int NumChunks;
-		int ChunkLength;
-		long AllocatedReaders;
-		FWriterStatePtr Writer;
-		FChunkStatePtr Chunks[FComputeBuffer::MaxChunks];
+	public:
+		const unsigned short int NumReaders;
+		const unsigned short int NumChunks;
+		const int ChunkLength;
+
+		FHeader(int InNumReaders, int InNumChunks, int InChunkLength)
+			: NumReaders(InNumReaders)
+			, NumChunks(InNumChunks)
+			, ChunkLength(InChunkLength)
+		{
+			static_assert(sizeof(FWriterState) == sizeof(long long), "Incorrect size of FWriterState");
+			static_assert(sizeof(FReaderState) == sizeof(long long), "Incorrect size of FReaderState");
+			static_assert(sizeof(FChunkState) == sizeof(long long), "Incorrect size of FChunkState");
+			static_assert(sizeof(FHeader) == HeaderSize, "Incorrect size of FHeader");
+
+			GetChunkStatePtr(0).Set(FChunkState(EWriteState::Writing, 0, 0));
+		}
+
+		FWriterStatePtr GetWriterStatePtr()
+		{
+			return FWriterStatePtr(&Writer);
+		}
+
+		FReaderStatePtr GetReaderStatePtr(int ReaderIdx)
+		{
+			return FReaderStatePtr(&Readers[ReaderIdx]);
+		}
+
+		FChunkStatePtr GetChunkStatePtr(int ChunkIdx)
+		{
+			return FChunkStatePtr(&Chunks[ChunkIdx]);
+		}
+
+	private:
+		FWriterState Writer;
+		FReaderState Readers[FComputeBuffer::MaxReaders];
+		FChunkState Chunks[FComputeBuffer::MaxChunks];
 	};
 
 	char Name[FComputeBuffer::MaxNameLength];
@@ -189,11 +303,12 @@ struct FComputeBufferDetail
 	FComputeEvent WriterEvent;
 	FComputeEvent ReaderEvents[FComputeBuffer::MaxReaders];
 
-	FComputeBufferWriter Writer;
+	long RefCount;
 
 	FComputeBufferDetail(const char* InName)
 		: Header(nullptr)
 		, ChunkPtrs{ nullptr, }
+		, RefCount(1)
 	{
 		static_assert(sizeof(FChunkStatePtr) == sizeof(long long), "Incorrect size of FChunkStatePtr; check union is declared correctly.");
 		static_assert(sizeof(FWriterStatePtr) == sizeof(long long), "Incorrect size of FWriterStatePtr; check union is declared correctly.");
@@ -206,7 +321,20 @@ struct FComputeBufferDetail
 		Header = nullptr;
 	}
 
-	static std::shared_ptr<FComputeBufferDetail> CreateNew(const FComputeBuffer::FParams& Params)
+	void AddRef()
+	{
+		FComputePlatform::AtomicIncrement(&RefCount);
+	}
+
+	void Release()
+	{
+		if (FComputePlatform::AtomicDecrement(&RefCount) == 0)
+		{
+			delete this;
+		}
+	}
+
+	static std::unique_ptr<FComputeBufferDetail> CreateNew(const FComputeBuffer::FParams& Params)
 	{
 		long long Capacity = sizeof(FHeader) + (Params.NumChunks * sizeof(unsigned int)) + (Params.NumChunks * Params.ChunkLength);
 
@@ -222,23 +350,19 @@ struct FComputeBufferDetail
 		char NameBuffer[FComputeBuffer::MaxNameLength];
 		snprintf(NameBuffer, FComputeBuffer::MaxNameLength, "%s_M", Name);
 
-		std::shared_ptr<FComputeBufferDetail> Detail = std::make_shared<FComputeBufferDetail>(Name);
+		std::unique_ptr<FComputeBufferDetail> Detail = std::make_unique<FComputeBufferDetail>(Name);
 		if (!Detail->MemoryMappedFile.Create(NameBuffer, Capacity))
 		{
 			return nullptr;
 		}
 
-		Detail->Header = (FHeader*)Detail->MemoryMappedFile.GetPointer();
-		if (Detail->Header == nullptr)
+		void* Pointer = Detail->MemoryMappedFile.GetPointer();
+		if (Pointer == nullptr)
 		{
 			return nullptr;
 		}
 
-		memset(Detail->Header, 0, sizeof(*Header));
-		Detail->Header->NumReaders = Params.NumReaders;
-		Detail->Header->NumChunks = Params.NumChunks;
-		Detail->Header->ChunkLength = Params.ChunkLength;
-		Detail->Header->Chunks[0].Set(FChunkState(EWriteState::Writing, 0, 0));
+		Detail->Header = new(Pointer) FHeader(Params.NumReaders, Params.NumChunks, Params.ChunkLength);
 
 		snprintf(NameBuffer, FComputeBuffer::MaxNameLength, "%s_W", Name);
 		if (!Detail->WriterEvent.Create(NameBuffer))
@@ -255,16 +379,16 @@ struct FComputeBufferDetail
 			}
 		}
 
-		InitShared(Detail);
+		InitShared(Detail.get());
 		return Detail;
 	}
 
-	static std::shared_ptr<FComputeBufferDetail> OpenExisting(const char* Name)
+	static std::unique_ptr<FComputeBufferDetail> OpenExisting(const char* Name)
 	{
 		char NameBuffer[FComputeBuffer::MaxNameLength];
 		snprintf(NameBuffer, FComputeBuffer::MaxNameLength, "%s_M", Name);
 
-		std::shared_ptr<FComputeBufferDetail> Detail = std::make_shared<FComputeBufferDetail>(Name);
+		std::unique_ptr<FComputeBufferDetail> Detail = std::make_unique<FComputeBufferDetail>(Name);
 		if (!Detail->MemoryMappedFile.OpenExisting(NameBuffer))
 		{
 			return nullptr;
@@ -291,115 +415,239 @@ struct FComputeBufferDetail
 			}
 		}
 
-		InitShared(Detail);
+		InitShared(Detail.get());
 		return Detail;
 	}
 
-	bool IsComplete(const FReaderState& ReaderState) const
+	int CreateReader()
 	{
-		if (ReaderState.Detached)
+		for(int ReaderIdx = 0; ReaderIdx < Header->NumReaders; ReaderIdx++)
 		{
-			return true;
+			FReaderStatePtr ReaderStatePtr = Header->GetReaderStatePtr(ReaderIdx);
+			for (;;)
+			{
+				FReaderState ReaderState = ReaderStatePtr.Get();
+				if (ReaderState.GetRefCount() > 0)
+				{
+					break;
+				}
+				if (ReaderStatePtr.TryUpdate(ReaderState, FReaderState(0, 0, 1, false)))
+				{
+					FWriterStatePtr WriterStatePtr = Header->GetWriterStatePtr();
+					for (;;)
+					{
+						FWriterState WriterState = WriterStatePtr.Get();
+						assert(!WriterState.HasWrapped());
+
+						if (WriterStatePtr.TryUpdate(WriterState, FWriterState(WriterState.GetChunkIdx(), WriterState.GetReaderFlags() | (1 << ReaderIdx), WriterState.GetRefCount(), WriterState.HasWrapped())))
+						{
+							for (int WriteChunkIdx = 0; WriteChunkIdx <= WriterState.GetChunkIdx(); WriteChunkIdx++)
+							{
+								Header->GetChunkStatePtr(WriteChunkIdx).StartReading(ReaderIdx);
+							}
+							return ReaderIdx;
+						}
+					}
+				}
+			}
 		}
+		return -1;
+	}
+
+	void AddReaderRef(int ReaderIdx)
+	{
+		FReaderStatePtr ReaderStatePtr = Header->GetReaderStatePtr(ReaderIdx);
+		for (; ; )
+		{
+			FReaderState ReaderState = ReaderStatePtr.Get();
+			assert(ReaderState.GetRefCount() > 0);
 
-		FChunkState ChunkState = Header->Chunks[ReaderState.ChunkIdx].Get();
-		return ChunkState.GetWriteState() == EWriteState::Complete && ReaderState.Offset == ChunkState.GetLength();
+			if (ReaderStatePtr.TryUpdate(ReaderState, FReaderState(ReaderState.GetChunkIdx(), ReaderState.GetOffset(), ReaderState.GetRefCount() + 1, ReaderState.IsDetached())))
+			{
+				break;
+			}
+		}
 	}
 
-	int AllocateReader(int ChunkIdx, int Offset)
+	void ReleaseReaderRef(int ReaderIdx)
 	{
-		for (;;)
+		FReaderStatePtr ReaderStatePtr = Header->GetReaderStatePtr(ReaderIdx);
+		for (; ; )
 		{
-			int AllocatedReaders = (int)Header->AllocatedReaders;
-			int ReaderFlag = (AllocatedReaders + 1) ^ AllocatedReaders;
+			FReaderState ReaderState = ReaderStatePtr.Get();
+			assert(ReaderState.GetRefCount() > 0);
 
-			int ReaderIdx = FComputePlatform::FloorLog2(ReaderFlag);
-			if (ReaderIdx >= Header->NumReaders)
+			if (ReaderState.GetRefCount() == 1)
 			{
-				return -1;
+				for (int Idx = 0; Idx < Header->NumChunks; Idx++)
+				{
+					Header->GetChunkStatePtr(Idx).FinishReading(ReaderIdx);
+				}
 			}
 
-			if (FComputePlatform::AtomicCompareExchange(&Header->AllocatedReaders, AllocatedReaders | ReaderFlag, AllocatedReaders))
+			if (ReaderStatePtr.TryUpdate(ReaderState, FReaderState(ReaderState.GetChunkIdx(), ReaderState.GetOffset(), ReaderState.GetRefCount() - 1, ReaderState.IsDetached())))
+			{
+				break;
+			}
+		}
+	}
+
+	void CreateWriter()
+	{
+		FWriterStatePtr WriterStatePtr = Header->GetWriterStatePtr();
+		for (; ; )
+		{
+			FWriterState WriterState = WriterStatePtr.Get();
+			assert(WriterState.GetRefCount() == 0);
+
+			if (WriterStatePtr.TryUpdate(WriterState, FWriterState(WriterState.GetChunkIdx(), WriterState.GetReaderFlags(), 1, WriterState.HasWrapped())))
 			{
-				for (;;)
+				FChunkStatePtr ChunkStatePtr = Header->GetChunkStatePtr(WriterState.GetChunkIdx());
+				for (; ; )
 				{
-					FWriterState State = Header->Writer.Get();
-					if (Header->Writer.TryUpdate(State, FWriterState(State.GetChunkIdx(), State.GetReaderFlags() | ReaderFlag, State.HasWrapped())))
+					FChunkState ChunkState = ChunkStatePtr.Get();
+					if (ChunkStatePtr.TryUpdate(ChunkState, FChunkState(EWriteState::Writing, ChunkState.GetReaderFlags(), ChunkState.GetLength())))
 					{
-						for (int WriteChunkIdx = 0; WriteChunkIdx <= State.GetChunkIdx(); WriteChunkIdx++)
-						{
-							Header->Chunks[WriteChunkIdx].StartReading(ReaderIdx);
-						}
-						return ReaderIdx;
+						break;
 					}
 				}
+				break;
+			}
+		}
+	}
+
+	void AddWriterRef()
+	{
+		FWriterStatePtr WriterStatePtr = Header->GetWriterStatePtr();
+		for (; ; )
+		{
+			FWriterState WriterState = WriterStatePtr.Get();
+			assert(WriterState.GetRefCount() > 0);
+
+			if (WriterStatePtr.TryUpdate(WriterState, FWriterState(WriterState.GetChunkIdx(), WriterState.GetReaderFlags(), WriterState.GetRefCount() + 1, WriterState.HasWrapped())))
+			{
+				break;
+			}
+		}
+	}
+
+	void ReleaseWriterRef()
+	{
+		FWriterStatePtr writerStatePtr = Header->GetWriterStatePtr();
+		for (; ; )
+		{
+			FWriterState writerState = writerStatePtr.Get();
+			assert(writerState.GetRefCount() > 0);
+
+			if (writerState.GetRefCount() == 1)
+			{
+				MarkComplete();
+			}
+
+			if (writerStatePtr.TryUpdate(writerState, FWriterState(writerState.GetChunkIdx(), writerState.GetReaderFlags(), writerState.GetRefCount() - 1, writerState.HasWrapped())))
+			{
+				break;
 			}
 		}
 	}
 
-	void ReleaseReader(int ReaderIdx)
+	bool IsComplete(int ReaderIdx) const
 	{
-		int ReaderFlag = 1 << ReaderIdx;
-		for (int Idx = 0; Idx < Header->NumChunks; Idx++)
+		FReaderState ReaderState = Header->GetReaderStatePtr(ReaderIdx).Get();
+		if (ReaderState.IsDetached())
 		{
-			Header->Chunks[Idx].FinishReading(ReaderFlag);
+			return true;
 		}
-		FComputePlatform::AtomicAnd(&Header->AllocatedReaders, (long)~ReaderFlag);
+
+		FChunkState ChunkState = Header->GetChunkStatePtr(ReaderState.GetChunkIdx()).Get();
+		return ChunkState.GetWriteState() == EWriteState::Complete && ReaderState.GetOffset() == ChunkState.GetLength();
 	}
 
-	size_t GetMaxReadSize(const FReaderState& ReaderState) const
+	void DetachReader(int ReaderIdx)
 	{
-		if (ReaderState.Detached)
+		FReaderStatePtr ReaderStatePtr = Header->GetReaderStatePtr(ReaderIdx);
+		for (; ; )
+		{
+			FReaderState ReaderState = ReaderStatePtr.Get();
+			if (ReaderStatePtr.TryUpdate(ReaderState, FReaderState(ReaderState.GetChunkIdx(), ReaderState.GetOffset(), ReaderState.GetRefCount(), true)))
+			{
+				ReaderEvents[ReaderIdx].Signal();
+				break;
+			}
+		}
+	}
+
+	void AdvanceReadPosition(int ReaderIdx, size_t Offset)
+	{
+		FReaderStatePtr ReaderStatePtr = Header->GetReaderStatePtr(ReaderIdx);
+		for (; ; )
+		{
+			FReaderState ReaderState = ReaderStatePtr.Get();
+			if (ReaderStatePtr.TryUpdate(ReaderState, FReaderState(ReaderState.GetChunkIdx(), ReaderState.GetOffset() + (int)Offset, ReaderState.GetRefCount(), ReaderState.IsDetached())))
+			{
+				ReaderEvents[ReaderIdx].Signal();
+				break;
+			}
+		}
+	}
+
+	size_t GetMaxReadSize(int ReaderIdx) const
+	{
+		FReaderState ReaderState = Header->GetReaderStatePtr(ReaderIdx).Get();
+		if (ReaderState.IsDetached())
 		{
 			return 0;
 		}
 
-		const FChunkStatePtr& ChunkStatePtr = Header->Chunks[ReaderState.ChunkIdx];
+		FChunkStatePtr ChunkStatePtr = Header->GetChunkStatePtr(ReaderState.GetChunkIdx());
 		FChunkState ChunkState = ChunkStatePtr.Get();
 
-		if (!ChunkState.HasReaderFlag(ReaderState.ReaderIdx))
+		if (ChunkState.HasReaderFlag(ReaderIdx))
+		{
+			return ChunkState.GetLength() - ReaderState.GetOffset();
+		}
+		else
 		{
 			return 0;
 		}
-
-		return ChunkState.GetLength() - ReaderState.Offset;
 	}
 
-	const unsigned char* WaitToRead(FReaderState& ReaderState, size_t MinSize, int TimeoutMs)
+	const unsigned char* WaitToRead(int ReaderIdx, size_t MinSize, int TimeoutMs)
 	{
-		int ReaderIdx = ReaderState.ReaderIdx;
 		for (; ; )
 		{
-			if (ReaderState.Detached)
+			FReaderStatePtr ReaderStatePtr = Header->GetReaderStatePtr(ReaderIdx);
+			FReaderState ReaderState = ReaderStatePtr.Get();
+			if (ReaderState.IsDetached())
 			{
 				return nullptr;
 			}
 
-			FChunkStatePtr& ChunkStatePtr = Header->Chunks[ReaderState.ChunkIdx];
+			FChunkStatePtr ChunkStatePtr = Header->GetChunkStatePtr(ReaderState.GetChunkIdx());
 			FChunkState ChunkState = ChunkStatePtr.Get();
 
 			if (!ChunkState.HasReaderFlag(ReaderIdx))
 			{
 				// Wait until the current chunk is readable
-				if (!ChunkState.HasReaderFlag(ReaderIdx) && !ReaderEvents[ReaderIdx].Wait(TimeoutMs))
+				if (!ReaderEvents[ReaderIdx].Wait(TimeoutMs))
 				{
 					return nullptr;
 				}
 			}
-			else if (ReaderState.Offset + MinSize <= ChunkState.GetLength())
+			else if (ReaderState.GetOffset() + MinSize <= ChunkState.GetLength())
 			{
 				// We have enough data in the chunk to be able to read a message
-				return ChunkPtrs[ReaderState.ChunkIdx] + ReaderState.Offset;
+				return ChunkPtrs[ReaderState.GetChunkIdx()] + ReaderState.GetOffset();
 			}
 			else if (ChunkState.GetWriteState() == EWriteState::Writing)
 			{
 				// Wait until there is more data in the chunk
-				if (Header->Chunks[ReaderState.ChunkIdx].Get().Value == ChunkState.Value && !ReaderEvents[ReaderIdx].Wait(TimeoutMs))
+				if (!ReaderEvents[ReaderIdx].Wait(TimeoutMs))
 				{
 					return nullptr;
 				}
 			}
-			else if (ReaderState.Offset < ChunkState.GetLength() || ChunkState.GetWriteState() == EWriteState::Complete)
+			else if (ReaderState.GetOffset() < ChunkState.GetLength() || ChunkState.GetWriteState() == EWriteState::Complete)
 			{
 				// Cannot read the requested amount of data from this chunk.
 				return nullptr;
@@ -410,12 +658,13 @@ struct FComputeBufferDetail
 				ChunkStatePtr.FinishReading(ReaderIdx);
 				WriterEvent.Signal();
 
-				if (++ReaderState.ChunkIdx == Header->NumChunks)
+				int NextChunkIdx = ReaderState.GetChunkIdx() + 1;
+				if (NextChunkIdx == Header->NumChunks)
 				{
-					ReaderState.ChunkIdx = 0;
+					NextChunkIdx = 0;
 				}
 
-				ReaderState.Offset = 0;
+				ReaderStatePtr.TryUpdate(ReaderState, FReaderState(NextChunkIdx, 0, ReaderState.GetRefCount(), ReaderState.IsDetached()));
 			}
 			else
 			{
@@ -424,32 +673,49 @@ struct FComputeBufferDetail
 		}
 	}
 
-	void MarkComplete()
+	bool MarkComplete()
 	{
-		FWriterState State = Header->Writer.Get();
-		Header->Chunks[State.GetChunkIdx()].MarkComplete();
-		SetAllReaderEvents();
+		FWriterState WriterState = Header->GetWriterStatePtr().Get();
+
+		FChunkStatePtr ChunkStatePtr = Header->GetChunkStatePtr(WriterState.GetChunkIdx());
+		if (ChunkStatePtr.Get().GetWriteState() != EWriteState::Complete)
+		{
+			ChunkStatePtr.MarkComplete();
+			SetAllReadEvents();
+			return true;
+		}
+
+		return false;
 	}
 
 	void AdvanceWritePosition(size_t Size)
 	{
-		FWriterState State = Header->Writer.Get();
-		Header->Chunks[State.GetChunkIdx()].Append(Size);
-		SetAllReaderEvents();
+		if (Size > 0)
+		{
+			FWriterState WriterState = Header->GetWriterStatePtr().Get();
+
+			FChunkStatePtr ChunkStatePtr = Header->GetChunkStatePtr(WriterState.GetChunkIdx());
+			FChunkState ChunkState = ChunkStatePtr.Get();
+
+			assert(ChunkState.GetWriteState() == EWriteState::Writing);
+			ChunkStatePtr.Append(Size);
+
+			SetAllReadEvents();
+		}
 	}
 
 	size_t GetMaxWriteSize() const
 	{
-		FWriterState State = Header->Writer.Get();
+		FWriterState WriterState = Header->GetWriterStatePtr().Get();
 
-		FChunkState Value = Header->Chunks[State.GetChunkIdx()].Get();
-		if (Value.GetWriteState() == EWriteState::Complete)
+		FChunkState ChunkState = Header->GetChunkStatePtr(WriterState.GetChunkIdx()).Get();
+		if (ChunkState.GetWriteState() == EWriteState::Writing)
 		{
-			return 0;
+			return Header->ChunkLength - ChunkState.GetLength();
 		}
 		else
 		{
-			return Header->ChunkLength - Value.GetLength();
+			return 0;
 		}
 	}
 
@@ -477,10 +743,10 @@ struct FComputeBufferDetail
 		assert(MinSize <= Header->ChunkLength);
 
 		// Get the current chunk we're writing to
-		FWriterState WriterState = Header->Writer.Get();
+		FWriterState WriterState = Header->GetWriterStatePtr().Get();
 		int WriteChunkIdx = WriterState.GetChunkIdx();
 
-		FChunkStatePtr& WriteChunkStatePtr = Header->Chunks[WriteChunkIdx];
+		FChunkStatePtr WriteChunkStatePtr = Header->GetChunkStatePtr(WriteChunkIdx);
 
 		// Check if we can append to this chunk
 		FChunkState ChunkState = WriteChunkStatePtr.Get();
@@ -493,7 +759,7 @@ struct FComputeBufferDetail
 			}
 
 			WriteChunkStatePtr.FinishWriting(); // STATE CHANGE
-			SetAllReaderEvents();
+			SetAllReadEvents();
 		}
 
 		if (ChunkState.GetWriteState() == EWriteState::Complete)
@@ -509,7 +775,7 @@ struct FComputeBufferDetail
 		}
 
 		// Wait until all readers have finished with the chunk, and we can update the writer to match
-		FChunkStatePtr& NextWriteChunkStatePtr = Header->Chunks[NextWriteChunkIdx];
+		FChunkStatePtr NextWriteChunkStatePtr = Header->GetChunkStatePtr(NextWriteChunkIdx);
 		for (;;)
 		{
 			FChunkState NextWriteChunkState = NextWriteChunkStatePtr.Get();
@@ -522,13 +788,14 @@ struct FComputeBufferDetail
 			}
 			else if (NextWriteChunkStatePtr.TryUpdate(NextWriteChunkState, FChunkState(EWriteState::Writing, WriterState.GetReaderFlags(), 0)))
 			{
-				if (Header->Writer.TryUpdate(WriterState, FWriterState(NextWriteChunkIdx, WriterState.GetReaderFlags(), NextWriteChunkIdx == 0)))
+				FWriterStatePtr WriterStatePtr = Header->GetWriterStatePtr();
+				if (WriterStatePtr.TryUpdate(WriterState, FWriterState(NextWriteChunkIdx, WriterState.GetReaderFlags(), WriterState.GetRefCount(), NextWriteChunkIdx == 0)))
 				{
 					break;
 				}
 				else
 				{
-					WriterState = Header->Writer.Get();
+					WriterState = WriterStatePtr.Get();
 				}
 			}
 		}
@@ -538,10 +805,8 @@ struct FComputeBufferDetail
 private:
 	FComputeMemoryMappedFile MemoryMappedFile;
 
-	static void InitShared(std::shared_ptr<FComputeBufferDetail> Detail)
+	static void InitShared(FComputeBufferDetail* Detail)
 	{
-		Detail->Writer = FComputeBufferWriter(Detail);
-
 		FHeader* Header = Detail->Header;
 
 		unsigned char* NextPtr = (unsigned char*)(Header + 1);
@@ -552,7 +817,7 @@ private:
 		}
 	}
 
-	void SetAllReaderEvents()
+	void SetAllReadEvents()
 	{
 		for (int Idx = 0; Idx < Header->NumReaders; Idx++)
 		{
@@ -564,103 +829,112 @@ private:
 
 
 
-//// FComputeBufferReaderDetail /////
 
-struct FComputeBufferReaderDetail
-{
-	std::shared_ptr<FComputeBufferDetail> Buffer;
-	FComputeBufferDetail::FReaderState ReaderState;
 
-	FComputeBufferReaderDetail(std::shared_ptr<FComputeBufferDetail> InBuffer, int InReaderIdx)
-		: Buffer(InBuffer)
-		, ReaderState(InReaderIdx)
-	{
-	}
 
-	~FComputeBufferReaderDetail()
-	{
-		Buffer->ReleaseReader(ReaderState.ReaderIdx);
-	}
-
-	bool IsComplete() const
-	{
-		return Buffer->IsComplete(ReaderState);
-	}
 
-	void Detach()
-	{
-		ReaderState.Detached = true;
-		Buffer->ReaderEvents[ReaderState.ReaderIdx].Signal();
-	}
-
-	void AdvanceReadPosition(size_t Size)
-	{
-		ReaderState.Offset += (long)Size;
-	}
+//// FComputeBuffer /////
 
-	size_t GetMaxReadSize() const
-	{
-		return Buffer->GetMaxReadSize(ReaderState);
-	}
+FComputeBuffer::FComputeBuffer()
+	: Detail(nullptr)
+{
+}
 
-	const unsigned char* WaitToRead(size_t MinSize, int TimeoutMs)
+FComputeBuffer::FComputeBuffer(const FComputeBuffer& Other)
+	: Detail(Other.Detail)
+{
+	if (Detail)
 	{
-		return Buffer->WaitToRead(ReaderState, MinSize, TimeoutMs);
+		Detail->AddRef();
 	}
-};
-
+}
 
+FComputeBuffer::FComputeBuffer(FComputeBuffer&& Other) noexcept
+	: Detail(Other.Detail)
+{
+	Other.Detail = nullptr;
+}
 
+FComputeBuffer::~FComputeBuffer()
+{
+	Close();
+}
 
+FComputeBuffer& FComputeBuffer::operator=(const FComputeBuffer& Other)
+{
+	Close();
 
-//// FComputeBuffer /////
+	if (Other.Detail)
+	{
+		Detail = Other.Detail;
+		Detail->AddRef();
+	}
 
-FComputeBuffer::FComputeBuffer()
-	: Detail(nullptr)
-{
+	return *this;
 }
 
-FComputeBuffer::~FComputeBuffer()
+FComputeBuffer& FComputeBuffer::operator=(FComputeBuffer&& Other) noexcept
 {
 	Close();
+
+	Detail = Other.Detail;
+	Other.Detail = nullptr;
+
+	return *this;
 }
 
 bool FComputeBuffer::CreateNew(const FParams& Params)
 {
-	Detail = FComputeBufferDetail::CreateNew(Params);
-	return Detail != nullptr;
+	Close();
+
+	std::unique_ptr<FComputeBufferDetail> Buffer = FComputeBufferDetail::CreateNew(Params);
+	if (Buffer)
+	{
+		Detail = Buffer.release();
+		return true;
+	}
+	return false;
 }
 
 bool FComputeBuffer::OpenExisting(const char* Name)
 {
-	Detail = FComputeBufferDetail::OpenExisting(Name);
-	return Detail != nullptr;
+	Close();
+
+	std::unique_ptr<FComputeBufferDetail> Buffer = FComputeBufferDetail::OpenExisting(Name);
+	if (Buffer)
+	{
+		Detail = Buffer.release();
+		return true;
+	}
+	return false;
 }
 
 void FComputeBuffer::Close()
 {
-	Detail.reset();
+	if (Detail != nullptr)
+	{
+		Detail->Release();
+		Detail = nullptr;
+	}
 }
 
 FComputeBufferReader FComputeBuffer::CreateReader()
 {
-	int ReaderIdx = Detail->AllocateReader(0, 0);
-	if (ReaderIdx == -1)
-	{
-		assert(false);
-		return FComputeBufferReader();
-	}
-	return FComputeBufferReader(std::make_shared<FComputeBufferReaderDetail>(Detail, ReaderIdx));
+	int ReaderIdx = Detail->CreateReader();
+	Detail->AddRef();
+	return FComputeBufferReader(Detail, ReaderIdx);
 }
 
-FComputeBufferWriter& FComputeBuffer::GetWriter()
+FComputeBufferWriter FComputeBuffer::CreateWriter()
 {
-	return Detail->Writer;
+	Detail->CreateWriter();
+	Detail->AddRef();
+	return FComputeBufferWriter(Detail);
 }
 
-const FComputeBufferWriter& FComputeBuffer::GetWriter() const
+const char* FComputeBuffer::GetName() const
 {
-	return Detail->Writer;
+	return Detail->Name;
 }
 
 
@@ -668,16 +942,30 @@ const FComputeBufferWriter& FComputeBuffer::GetWriter() const
 
 
 
-
 //// FComputeBufferReader /////
 
 FComputeBufferReader::FComputeBufferReader()
+	: Detail(nullptr)
+	, ReaderIdx(0)
+{
+}
+
+FComputeBufferReader::FComputeBufferReader(const FComputeBufferReader& Other)
+	: Detail(Other.Detail)
+	, ReaderIdx(Other.ReaderIdx)
 {
+	if (Detail)
+	{
+		Detail->AddRef();
+		Detail->AddReaderRef(ReaderIdx);
+	}
 }
 
-FComputeBufferReader::FComputeBufferReader(std::shared_ptr<FComputeBufferReaderDetail> Detail)
-	: Detail(std::move(Detail))
+FComputeBufferReader::FComputeBufferReader(FComputeBufferReader&& Other) noexcept
+	: Detail(Other.Detail)
+	, ReaderIdx(Other.ReaderIdx)
 {
+	Other.Detail = nullptr;
 }
 
 FComputeBufferReader::~FComputeBufferReader()
@@ -685,29 +973,65 @@ FComputeBufferReader::~FComputeBufferReader()
 	Close();
 }
 
+FComputeBufferReader& FComputeBufferReader::operator=(const FComputeBufferReader& Other)
+{
+	Close();
+
+	if (Other.Detail)
+	{
+		Detail = Other.Detail;
+		ReaderIdx = Other.ReaderIdx;
+
+		Detail->AddRef();
+		Detail->AddReaderRef(ReaderIdx);
+	}
+
+	return *this;
+}
+
+FComputeBufferReader& FComputeBufferReader::operator=(FComputeBufferReader&& Other) noexcept
+{
+	Close();
+
+	Detail = Other.Detail;
+	ReaderIdx = Other.ReaderIdx;
+
+	Other.Detail = nullptr;
+	return *this;
+}
+
 void FComputeBufferReader::Close()
 {
-	Detail.reset();
+	if (Detail != nullptr)
+	{
+		Detail->ReleaseReaderRef(ReaderIdx);
+		Detail->Release();
+		Detail = nullptr;
+	}
 }
 
 void FComputeBufferReader::Detach()
 {
-	Detail->Detach();
+	assert(Detail);
+	Detail->DetachReader(ReaderIdx);
 }
 
 bool FComputeBufferReader::IsComplete() const
 {
-	return Detail->IsComplete();
+	assert(Detail);
+	return Detail->IsComplete(ReaderIdx);
 }
 
 void FComputeBufferReader::AdvanceReadPosition(size_t Size)
 {
-	Detail->AdvanceReadPosition(Size);
+	assert(Detail);
+	Detail->AdvanceReadPosition(ReaderIdx, Size);
 }
 
 size_t FComputeBufferReader::GetMaxReadSize() const
 {
-	return Detail->GetMaxReadSize();
+	assert(Detail);
+	return Detail->GetMaxReadSize(ReaderIdx);
 }
 
 size_t FComputeBufferReader::Read(void* Buffer, size_t MaxSize, int TimeoutMs)
@@ -731,12 +1055,20 @@ size_t FComputeBufferReader::Read(void* Buffer, size_t MaxSize, int TimeoutMs)
 
 const unsigned char* FComputeBufferReader::WaitToRead(size_t MinSize, int TimeoutMs)
 {
-	return Detail->WaitToRead(MinSize, TimeoutMs);
+	assert(Detail);
+	return Detail->WaitToRead(ReaderIdx, MinSize, TimeoutMs);
+}
+
+FComputeBufferReader::FComputeBufferReader(FComputeBufferDetail* InDetail, int InReaderIdx)
+	: Detail(InDetail)
+	, ReaderIdx(InReaderIdx)
+{
 }
 
 const char* FComputeBufferReader::GetName() const
 {
-	return Detail->Buffer->Name;
+	assert(Detail);
+	return Detail->Name;
 }
 
 
@@ -750,13 +1082,60 @@ FComputeBufferWriter::FComputeBufferWriter()
 {
 }
 
-FComputeBufferWriter::FComputeBufferWriter(std::shared_ptr<FComputeBufferDetail> Detail)
-	: Detail(std::move(Detail))
+FComputeBufferWriter::FComputeBufferWriter(const FComputeBufferWriter& Other)
+	: Detail(Other.Detail)
+{
+	if (Detail)
+	{
+		Detail->AddRef();
+		Detail->AddWriterRef();
+	}
+}
+
+FComputeBufferWriter::FComputeBufferWriter(FComputeBufferWriter&& Other) noexcept
+	: Detail(Other.Detail)
 {
+	Other.Detail = nullptr;
 }
 
 FComputeBufferWriter::~FComputeBufferWriter()
 {
+	Close();
+}
+
+FComputeBufferWriter& FComputeBufferWriter::operator=(const FComputeBufferWriter& Other)
+{
+	Close();
+
+	if (Other.Detail)
+	{
+		Detail = Other.Detail;
+		Detail->AddWriterRef();
+		Detail->AddRef();
+	}
+
+	return *this;
+}
+
+FComputeBufferWriter& FComputeBufferWriter::operator=(FComputeBufferWriter&& Other) noexcept
+{
+	Close();
+
+	Detail = Other.Detail;
+	Other.Detail = nullptr;
+
+	return *this;
+}
+
+
+void FComputeBufferWriter::Close()
+{
+	if (Detail != nullptr)
+	{
+		Detail->ReleaseWriterRef();
+		Detail->Release();
+		Detail = nullptr;
+	}
 }
 
 void FComputeBufferWriter::MarkComplete()
@@ -798,6 +1177,11 @@ unsigned char* FComputeBufferWriter::WaitToWrite(size_t MinSize, int TimeoutMs)
 	return Detail->WaitToWrite(MinSize, TimeoutMs);
 }
 
+FComputeBufferWriter::FComputeBufferWriter(FComputeBufferDetail* InDetail)
+	: Detail(InDetail)
+{
+}
+
 const char* FComputeBufferWriter::GetName() const
 {
 	return Detail->Name;
diff --git a/Engine/Source/Programs/Horde/Samples/RemoteWorkerCpp/ComputeBuffer.h b/Engine/Source/Programs/Horde/Samples/RemoteWorkerCpp/ComputeBuffer.h
--- a/Engine/Source/Programs/Horde/Samples/RemoteWorkerCpp/ComputeBuffer.h
+++ b/Engine/Source/Programs/Horde/Samples/RemoteWorkerCpp/ComputeBuffer.h
@@ -4,12 +4,11 @@
 
 #include <memory>
 
-class FComputeBufferReader;
-class FComputeBufferWriter;
-
 struct FHeader;
 struct FComputeBufferDetail;
-struct FComputeBufferReaderDetail;
+
+class FComputeBufferReader;
+class FComputeBufferWriter;
 
 //
 // Implements a ring buffer using shared memory, with one writer and multiple readers. 
@@ -43,8 +42,13 @@ public:
 	};
 
 	FComputeBuffer();
+	FComputeBuffer(const FComputeBuffer& Other);
+	FComputeBuffer(FComputeBuffer&& Other) noexcept;
 	~FComputeBuffer();
 
+	FComputeBuffer& operator=(const FComputeBuffer& Other);
+	FComputeBuffer& operator=(FComputeBuffer&& Other) noexcept;
+
 	// Creates a new buffer
 	bool CreateNew(const FParams& Params);
 
@@ -55,17 +59,20 @@ public:
 	void Close();
 
 	// Test if the buffer is currently open
-	bool IsOpen() const { return Detail != nullptr; }
+	bool IsValid() const { return Detail != nullptr; }
 
 	// Creates a new reader for this buffer
 	FComputeBufferReader CreateReader();
 
-	// Gets a reference to the writer instance
-	FComputeBufferWriter& GetWriter();
-	const FComputeBufferWriter& GetWriter() const;
+	// Creates a new writer for this buffer
+	FComputeBufferWriter CreateWriter();
 
 private:
-	std::shared_ptr<FComputeBufferDetail> Detail;
+	friend class FWorkerComputeSocket;
+
+	FComputeBufferDetail* Detail;
+
+	const char* GetName() const;
 };
 
 //
@@ -75,9 +82,13 @@ class FComputeBufferReader
 {
 public:
 	FComputeBufferReader();
-	FComputeBufferReader(std::shared_ptr<FComputeBufferReaderDetail> Detail);
+	FComputeBufferReader(const FComputeBufferReader& Other);
+	FComputeBufferReader(FComputeBufferReader&& Other) noexcept;
 	~FComputeBufferReader();
 
+	FComputeBufferReader& operator=(const FComputeBufferReader& Other);
+	FComputeBufferReader& operator=(FComputeBufferReader&& Other) noexcept;
+
 	// Closes the handle to the underlying reader instance, resetting this instance back to empty
 	void Close();
 
@@ -85,7 +96,7 @@ public:
 	void Detach();
 
 	// Test if the reader is valid
-	bool IsValid() const { return Detail.get() != nullptr; }
+	bool IsValid() const { return Detail != nullptr; }
 
 	// Test whether the buffer has finished being written to (ie. MarkComplete() has been called by the writer) and all data has been read from it.
 	bool IsComplete() const;
@@ -104,10 +115,15 @@ public:
 	const unsigned char* WaitToRead(size_t MinSize, int TimeoutMs = -1);
 
 private:
+	struct FReaderRef;
+
+	friend class FComputeBuffer;
 	friend class FWorkerComputeSocket;
 
-	std::shared_ptr<FComputeBufferReaderDetail> Detail;
+	FComputeBufferDetail* Detail;
+	int ReaderIdx;
 
+	FComputeBufferReader(FComputeBufferDetail* Detail, int ReaderIdx);
 	const char* GetName() const;
 };
 
@@ -118,11 +134,18 @@ class FComputeBufferWriter
 {
 public:
 	FComputeBufferWriter();
-	FComputeBufferWriter(std::shared_ptr<FComputeBufferDetail> Detail);
+	FComputeBufferWriter(const FComputeBufferWriter& Other);
+	FComputeBufferWriter(FComputeBufferWriter&& Other) noexcept;
 	~FComputeBufferWriter();
 
+	FComputeBufferWriter& operator=(const FComputeBufferWriter& Other);
+	FComputeBufferWriter& operator=(FComputeBufferWriter&& Other) noexcept;
+
+	// Closes the handle to the underlying writer instance, resetting this instance back to empty
+	void Close();
+
 	// Test if the writer is valid
-	bool IsValid() const { return Detail.get() != nullptr; }
+	bool IsValid() const { return Detail != nullptr; }
 
 	// Signal that we've finished writing to this buffer
 	void MarkComplete();
@@ -141,9 +164,11 @@ public:
 	unsigned char* WaitToWrite(size_t MinSize, int TimeoutMs = -1);
 
 private:
+	friend class FComputeBuffer;
 	friend class FWorkerComputeSocket;
 
-	std::shared_ptr<FComputeBufferDetail> Detail;
+	FComputeBufferDetail* Detail;
 
+	FComputeBufferWriter(FComputeBufferDetail* Detail);
 	const char* GetName() const;
 };
diff --git a/Engine/Source/Programs/Horde/Samples/RemoteWorkerCpp/ComputeSocket.cpp b/Engine/Source/Programs/Horde/Samples/RemoteWorkerCpp/ComputeSocket.cpp
--- a/Engine/Source/Programs/Horde/Samples/RemoteWorkerCpp/ComputeSocket.cpp
+++ b/Engine/Source/Programs/Horde/Samples/RemoteWorkerCpp/ComputeSocket.cpp
@@ -33,15 +33,17 @@ FComputeChannel FComputeSocket::CreateChannel(int ChannelId)
 		return FComputeChannel();
 	}
 
-	return CreateChannel(ChannelId, RecvBuffer, SendBuffer);
+	return CreateChannel(ChannelId, std::move(RecvBuffer), std::move(SendBuffer));
 }
 
 FComputeChannel FComputeSocket::CreateChannel(int ChannelId, FComputeBuffer RecvBuffer, FComputeBuffer SendBuffer)
 {
-	AttachRecvBuffer(ChannelId, RecvBuffer.GetWriter());
-	AttachSendBuffer(ChannelId, SendBuffer.CreateReader());
+	FComputeChannel Channel(RecvBuffer.CreateReader(), SendBuffer.CreateWriter());
 
-	return FComputeChannel(RecvBuffer.CreateReader(), SendBuffer.GetWriter());
+	AttachRecvBuffer(ChannelId, std::move(RecvBuffer));
+	AttachSendBuffer(ChannelId, std::move(SendBuffer));
+
+	return Channel;
 }
 
 //////////////////////////////////////////////////////
@@ -76,35 +78,42 @@ bool FWorkerComputeSocket::Open()
 
 bool FWorkerComputeSocket::Open(const char* CommandBufferName)
 {
-	return CommandBuffer.OpenExisting(CommandBufferName);
+	FComputeBuffer CommandBuffer;
+	if (CommandBuffer.OpenExisting(CommandBufferName))
+	{
+		CommandBufferWriter = CommandBuffer.CreateWriter();
+		return true;
+	}
+	return false;
 }
 
 void FWorkerComputeSocket::Close()
 {
-	CommandBuffer.Close();
+	CommandBufferWriter.Close();
 }
 
-void FWorkerComputeSocket::AttachRecvBuffer(int ChannelId, FComputeBufferWriter Writer)
+void FWorkerComputeSocket::AttachRecvBuffer(int ChannelId, FComputeBuffer RecvBuffer)
 {
-	AttachBuffer(ChannelId, EMessageType::AttachRecvBuffer, Writer.GetName());
+	AttachBuffer(ChannelId, EMessageType::AttachRecvBuffer, RecvBuffer.GetName());
+	Buffers.push_back(std::move(RecvBuffer));
 }
 
-void FWorkerComputeSocket::AttachSendBuffer(int ChannelId, FComputeBufferReader Reader)
+void FWorkerComputeSocket::AttachSendBuffer(int ChannelId, FComputeBuffer SendBuffer)
 {
-	AttachBuffer(ChannelId, EMessageType::AttachSendBuffer, Reader.GetName());
+	AttachBuffer(ChannelId, EMessageType::AttachSendBuffer, SendBuffer.GetName());
+	Buffers.push_back(std::move(SendBuffer));
 }
 
 void FWorkerComputeSocket::AttachBuffer(int ChannelId, EMessageType Type, const char* Name)
 {
-	FComputeBufferWriter& Writer = CommandBuffer.GetWriter();
-	unsigned char* Data = Writer.WaitToWrite(1024);
+	unsigned char* Data = CommandBufferWriter.WaitToWrite(1024);
 
 	size_t Len = 0;
 	Len += WriteVarUInt(Data + Len, (unsigned char)Type);
 	Len += WriteVarUInt(Data + Len, (unsigned int)ChannelId);
 	Len += WriteString(Data + Len, Name);
 
-	Writer.AdvanceWritePosition(Len);
+	CommandBufferWriter.AdvanceWritePosition(Len);
 }
 
 void FWorkerComputeSocket::RunServer(FComputeBufferReader& CommandBufferReader, FComputeSocket& Socket)
@@ -131,7 +140,7 @@ void FWorkerComputeSocket::RunServer(FComputeBufferReader& CommandBufferReader,
 				FComputeBuffer Buffer;
 				if (Buffer.OpenExisting(Name))
 				{
-					Socket.AttachSendBuffer(ChannelId, Buffer.CreateReader());
+					Socket.AttachSendBuffer(ChannelId, Buffer);
 				}
 				else
 				{
@@ -150,7 +159,7 @@ void FWorkerComputeSocket::RunServer(FComputeBufferReader& CommandBufferReader,
 				FComputeBuffer Buffer;
 				if (Buffer.OpenExisting(Name))
 				{
-					Socket.AttachRecvBuffer(ChannelId, Buffer.GetWriter());
+					Socket.AttachRecvBuffer(ChannelId, Buffer);
 				}
 				else
 				{
@@ -355,15 +364,17 @@ public:
 		return true;
 	}
 
-	void AttachRecvBuffer(int ChannelId, FComputeBufferWriter Writer)
+	void AttachRecvBuffer(int ChannelId, FComputeBuffer RecvBuffer) override
 	{
 		std::lock_guard<std::mutex> Lock(CriticalSection);
+		FComputeBufferWriter Writer = RecvBuffer.CreateWriter();
 		Writers.insert(std::pair<int, FComputeBufferWriter>(ChannelId, std::move(Writer)));
 	}
 
-	void AttachSendBuffer(int ChannelId, FComputeBufferReader Reader)
+	void AttachSendBuffer(int ChannelId, FComputeBuffer SendBuffer) override
 	{
 		std::lock_guard<std::mutex> Lock(CriticalSection);
+		FComputeBufferReader Reader = SendBuffer.CreateReader();
 		Readers.push_back(Reader);
 		SendThreads.insert(std::make_pair(ChannelId, std::thread(&FRemoteComputeSocket::SendThreadProc, this, ChannelId, std::move(Reader))));
 	}
diff --git a/Engine/Source/Programs/Horde/Samples/RemoteWorkerCpp/ComputeSocket.h b/Engine/Source/Programs/Horde/Samples/RemoteWorkerCpp/ComputeSocket.h
--- a/Engine/Source/Programs/Horde/Samples/RemoteWorkerCpp/ComputeSocket.h
+++ b/Engine/Source/Programs/Horde/Samples/RemoteWorkerCpp/ComputeSocket.h
@@ -5,6 +5,7 @@
 #include "ComputeBuffer.h"
 #include "ComputeChannel.h"
 #include "ComputeTransport.h"
+#include <vector>
 
 //
 // Connection to a remote machine that multiplexes data into and out-of multiple buffers
@@ -20,10 +21,10 @@ public:
 	FComputeSocket& operator=(const FComputeSocket&) = delete;
 
 	// Attaches a new buffer for receiving data
-	virtual void AttachRecvBuffer(int ChannelId, FComputeBufferWriter Writer) = 0;
+	virtual void AttachRecvBuffer(int ChannelId, FComputeBuffer RecvBuffer) = 0;
 
 	// Attaches a new buffer for sending data */
-	virtual void AttachSendBuffer(int ChannelId, FComputeBufferReader Reader) = 0;
+	virtual void AttachSendBuffer(int ChannelId, FComputeBuffer SendBuffer) = 0;
 
 	// Attaches a channel to this socket
 	FComputeChannel CreateChannel(int ChannelId);
@@ -54,10 +55,10 @@ public:
 	void Close();
 
 	// Attaches a new buffer for receiving data
-	virtual void AttachRecvBuffer(int ChannelId, FComputeBufferWriter Writer) override;
+	virtual void AttachRecvBuffer(int ChannelId, FComputeBuffer RecvBuffer) override;
 
 	// Attaches a new buffer for sending data
-	virtual void AttachSendBuffer(int ChannelId, FComputeBufferReader Reader) override;
+	virtual void AttachSendBuffer(int ChannelId, FComputeBuffer SendBuffer) override;
 
 	// Reads and handles a command from the command buffer
 	static void RunServer(FComputeBufferReader& CommandBufferReader, FComputeSocket& Socket);
@@ -65,7 +66,8 @@ public:
 private:
 	enum class EMessageType;
 
-	FComputeBuffer CommandBuffer;
+	FComputeBufferWriter CommandBufferWriter;
+	std::vector<FComputeBuffer> Buffers;
 
 	void AttachBuffer(int ChannelId, EMessageType Type, const char* Name);
 
diff --git a/Engine/Source/Programs/Horde/Samples/RemoteWorkerCpp/ComputeTransport.cpp b/Engine/Source/Programs/Horde/Samples/RemoteWorkerCpp/ComputeTransport.cpp
--- a/Engine/Source/Programs/Horde/Samples/RemoteWorkerCpp/ComputeTransport.cpp
+++ b/Engine/Source/Programs/Horde/Samples/RemoteWorkerCpp/ComputeTransport.cpp
@@ -43,21 +43,19 @@ bool FComputeTransport::RecvMessage(void* Data, size_t Size)
 
 ///////////////////////////////////
 
-FBufferTransport::FBufferTransport(FComputeBuffer InSendBuffer, FComputeBuffer InRecvBuffer)
-	: SendBuffer(std::move(InSendBuffer))
-	, RecvBuffer(std::move(InRecvBuffer))
+FBufferTransport::FBufferTransport(FComputeBufferWriter InSendBufferWriter, FComputeBufferReader InRecvBufferReader)
+	: SendBufferWriter(std::move(InSendBufferWriter))
+	, RecvBufferReader(std::move(InRecvBufferReader))
 {
-	RecvBufferReader = std::move(RecvBuffer.CreateReader());
 }
 
 size_t FBufferTransport::Send(const void* Data, size_t Size)
 {
-	FComputeBufferWriter& Writer = SendBuffer.GetWriter();
-	unsigned char* Buffer = Writer.WaitToWrite(1);
+	unsigned char* Buffer = SendBufferWriter.WaitToWrite(1);
 
-	size_t WriteSize = std::min(Size, SendBuffer.GetWriter().GetMaxWriteSize());
+	size_t WriteSize = std::min(Size, SendBufferWriter.GetMaxWriteSize());
 	memcpy(Buffer, Data, WriteSize);
-	Writer.AdvanceWritePosition(WriteSize);
+	SendBufferWriter.AdvanceWritePosition(WriteSize);
 
 	return WriteSize;
 }
@@ -75,11 +73,10 @@ size_t FBufferTransport::Recv(void* Data, size_t Size)
 
 void FBufferTransport::MarkComplete()
 {
-	SendBuffer.GetWriter().MarkComplete();
+	SendBufferWriter.MarkComplete();
 }
 
 void FBufferTransport::Close()
 {
-	SendBuffer.GetWriter().MarkComplete();
-	RecvBuffer.GetWriter().MarkComplete();
-}
+	RecvBufferReader.Detach();
+}
\ No newline at end of file
diff --git a/Engine/Source/Programs/Horde/Samples/RemoteWorkerCpp/ComputeTransport.h b/Engine/Source/Programs/Horde/Samples/RemoteWorkerCpp/ComputeTransport.h
--- a/Engine/Source/Programs/Horde/Samples/RemoteWorkerCpp/ComputeTransport.h
+++ b/Engine/Source/Programs/Horde/Samples/RemoteWorkerCpp/ComputeTransport.h
@@ -37,7 +37,7 @@ public:
 class FBufferTransport final : public FComputeTransport
 {
 public:
-	FBufferTransport(FComputeBuffer InSendBuffer, FComputeBuffer InRecvBuffer);
+	FBufferTransport(FComputeBufferWriter InSendBufferWriter, FComputeBufferReader InRecvBufferReader);
 
 protected:
 	virtual size_t Send(const void* Data, size_t Size) override final;
@@ -46,7 +46,6 @@ protected:
 	virtual void Close() override final;
 
 private:
-	FComputeBuffer SendBuffer;
-	FComputeBuffer RecvBuffer;
 	FComputeBufferReader RecvBufferReader;
+	FComputeBufferWriter SendBufferWriter;
 };
diff --git a/Engine/Source/Programs/Horde/Samples/RemoteWorkerCpp/Tests.cpp b/Engine/Source/Programs/Horde/Samples/RemoteWorkerCpp/Tests.cpp
--- a/Engine/Source/Programs/Horde/Samples/RemoteWorkerCpp/Tests.cpp
+++ b/Engine/Source/Programs/Horde/Samples/RemoteWorkerCpp/Tests.cpp
@@ -25,7 +25,7 @@ void ComputeBufferTest()
 	FComputeBuffer Buffer;
 	verify(Buffer.CreateNew(FComputeBuffer::FParams()));
 
-	FComputeBufferWriter& Writer = Buffer.GetWriter();
+	FComputeBufferWriter Writer = Buffer.CreateWriter();
 	unsigned char* WriteBuffer = Writer.WaitToWrite(sizeof(TestData));
 	memcpy(WriteBuffer, TestData, sizeof(TestData));
 	Writer.AdvanceWritePosition(sizeof(TestData));
@@ -33,6 +33,13 @@ void ComputeBufferTest()
 	FComputeBufferReader Reader = Buffer.CreateReader();
 	const unsigned char* ReadBuffer = Reader.WaitToRead(sizeof(TestData));
 	verify(memcmp(ReadBuffer, TestData, sizeof(TestData)) == 0);
+	assert(!Reader.IsComplete());
+
+	Reader.AdvanceReadPosition(sizeof(TestData));
+	assert(!Reader.IsComplete());
+
+	Writer.MarkComplete();
+	assert(Reader.IsComplete());
 }
 
 template<size_t TestDataSize> void CheckChannelSendRecv(FComputeChannel& SendChannel, FComputeChannel& RecvChannel, const char(&TestData)[TestDataSize])
@@ -55,15 +62,19 @@ void ComputeSocketTest()
 	verify(ServerToClientBuffer.CreateNew(FComputeBuffer::FParams()));
 
 	// Client transport
-	std::unique_ptr<FComputeSocket> ClientSocket = CreateComputeSocket(std::make_unique<FBufferTransport>(ClientToServerBuffer, ServerToClientBuffer), EComputeSocketEndpoint::Local);
+	std::unique_ptr<FComputeSocket> ClientSocket = CreateComputeSocket(std::make_unique<FBufferTransport>(ClientToServerBuffer.CreateWriter(), ServerToClientBuffer.CreateReader()), EComputeSocketEndpoint::Local);
 	FComputeChannel ClientChannel1 = ClientSocket->CreateChannel(1);
 	FComputeChannel ClientChannel2 = ClientSocket->CreateChannel(2);
 
 	// Server socket
-	std::unique_ptr<FComputeSocket> ServerSocket = CreateComputeSocket(std::make_unique<FBufferTransport>(ServerToClientBuffer, ClientToServerBuffer), EComputeSocketEndpoint::Remote);
+	std::unique_ptr<FComputeSocket> ServerSocket = CreateComputeSocket(std::make_unique<FBufferTransport>(ServerToClientBuffer.CreateWriter(), ClientToServerBuffer.CreateReader()), EComputeSocketEndpoint::Remote);
 	FComputeChannel ServerChannel1 = ServerSocket->CreateChannel(1);
 	FComputeChannel ServerChannel2 = ServerSocket->CreateChannel(2);
 
+	// Close the original buffers now that the sockets are set up
+	ClientToServerBuffer.Close();
+	ServerToClientBuffer.Close();
+
 	// Send data over channel 1
 	CheckChannelSendRecv(ClientChannel1, ServerChannel1, "Channel 1: Client -> Server");
 
diff --git a/Engine/Source/Programs/Shared/EpicGames.Horde.Tests/BufferTests.cs b/Engine/Source/Programs/Shared/EpicGames.Horde.Tests/BufferTests.cs
--- a/Engine/Source/Programs/Shared/EpicGames.Horde.Tests/BufferTests.cs
+++ b/Engine/Source/Programs/Shared/EpicGames.Horde.Tests/BufferTests.cs
@@ -24,7 +24,8 @@ namespace EpicGames.Horde.Tests
 		public async Task TestSimpleBuffer()
 		{
 			using PooledBuffer buffer = new PooledBuffer(2, 1024);
-			buffer.Writer.AdvanceWritePosition(10);
+			using ComputeBufferWriter bufferWriter = buffer.CreateWriter();
+			bufferWriter.AdvanceWritePosition(10);
 
 			using ComputeBufferReader bufferReader = buffer.CreateReader();
 			await bufferReader.WaitToReadAsync(9);
@@ -37,25 +38,26 @@ namespace EpicGames.Horde.Tests
 		{
 			using PooledBuffer buffer = new PooledBuffer(2, 20);
 			using ComputeBufferReader bufferReader = buffer.CreateReader();
+			using ComputeBufferWriter bufferWriter = buffer.CreateWriter();
 
 			// Fill up the first chunk
-			Assert.AreEqual(20, buffer.Writer.GetWriteBuffer().Length);
-			buffer.Writer.AdvanceWritePosition(10);
-			Assert.AreEqual(10, buffer.Writer.GetWriteBuffer().Length);
-			buffer.Writer.AdvanceWritePosition(10);
-			Assert.AreEqual(0, buffer.Writer.GetWriteBuffer().Length);
+			Assert.AreEqual(20, bufferWriter.GetWriteBuffer().Length);
+			bufferWriter.AdvanceWritePosition(10);
+			Assert.AreEqual(10, bufferWriter.GetWriteBuffer().Length);
+			bufferWriter.AdvanceWritePosition(10);
+			Assert.AreEqual(0, bufferWriter.GetWriteBuffer().Length);
 
-			Task waitToWriteTask = buffer.Writer.WaitToWriteAsync(1).AsTask();
+			Task waitToWriteTask = bufferWriter.WaitToWriteAsync(1).AsTask();
 			Assert.IsTrue(waitToWriteTask.IsCompleted);
 
 			// Fill up the second chunk
-			Assert.AreEqual(20, buffer.Writer.GetWriteBuffer().Length);
-			buffer.Writer.AdvanceWritePosition(10);
-			Assert.AreEqual(10, buffer.Writer.GetWriteBuffer().Length);
-			buffer.Writer.AdvanceWritePosition(10);
-			Assert.AreEqual(0, buffer.Writer.GetWriteBuffer().Length);
+			Assert.AreEqual(20, bufferWriter.GetWriteBuffer().Length);
+			bufferWriter.AdvanceWritePosition(10);
+			Assert.AreEqual(10, bufferWriter.GetWriteBuffer().Length);
+			bufferWriter.AdvanceWritePosition(10);
+			Assert.AreEqual(0, bufferWriter.GetWriteBuffer().Length);
 
-			waitToWriteTask = buffer.Writer.WaitToWriteAsync(1).AsTask();
+			waitToWriteTask = bufferWriter.WaitToWriteAsync(1).AsTask();
 			Assert.IsFalse(waitToWriteTask.IsCompleted);
 
 			// Wait for data to be read
@@ -75,7 +77,7 @@ namespace EpicGames.Horde.Tests
 
 			// Make sure both reader and writer have something to work with
 			Assert.AreEqual(20, bufferReader.GetReadBuffer().Length);
-			Assert.AreEqual(20, buffer.Writer.GetWriteBuffer().Length);
+			Assert.AreEqual(20, bufferWriter.GetWriteBuffer().Length);
 		}
 
 		[TestMethod]
@@ -168,7 +170,9 @@ namespace EpicGames.Horde.Tests
 				using (PooledBuffer localBuffer = new PooledBuffer(1024))
 				{
 					localSocket.AttachSendBuffer(1, localBuffer);
-					await localBuffer.Writer.WriteAsync(new byte[] { 1, 2, 3 });
+
+					using ComputeBufferWriter localBufferWriter = localBuffer.CreateWriter();
+					await localBufferWriter.WriteAsync(new byte[] { 1, 2, 3 });
 				}
 
 				Assert.IsTrue(await reader.WaitToReadAsync(3));
diff --git a/Engine/Source/Programs/Shared/EpicGames.Horde/Compute/AgentMessageChannel.cs b/Engine/Source/Programs/Shared/EpicGames.Horde/Compute/AgentMessageChannel.cs
--- a/Engine/Source/Programs/Shared/EpicGames.Horde/Compute/AgentMessageChannel.cs
+++ b/Engine/Source/Programs/Shared/EpicGames.Horde/Compute/AgentMessageChannel.cs
@@ -113,7 +113,7 @@ namespace EpicGames.Horde.Compute
 			_socket.AttachRecvBuffer(channelId, recvBuffer);
 			_socket.AttachSendBuffer(channelId, sendBuffer);
 			_recvBufferReader = recvBuffer.CreateReader();
-			_sendBufferWriter = sendBuffer.Writer.AddRef();
+			_sendBufferWriter = sendBuffer.CreateWriter();
 			_logger = logger;
 		}
 
diff --git a/Engine/Source/Programs/Shared/EpicGames.Horde/Compute/AgentMessageHandler.cs b/Engine/Source/Programs/Shared/EpicGames.Horde/Compute/AgentMessageHandler.cs
--- a/Engine/Source/Programs/Shared/EpicGames.Horde/Compute/AgentMessageHandler.cs
+++ b/Engine/Source/Programs/Shared/EpicGames.Horde/Compute/AgentMessageHandler.cs
@@ -8,16 +8,12 @@ using System.Threading;
 using System.Threading.Tasks;
 using EpicGames.Core;
 using EpicGames.Horde.Compute.Buffers;
-using EpicGames.Horde.Storage;
 using EpicGames.Horde.Storage.Nodes;
 using Microsoft.Extensions.Caching.Memory;
 using Microsoft.Extensions.Logging;
 using EpicGames.Horde.Storage.Bundles;
-using System.Text;
 using System.IO;
-using System.Reflection;
 using System.Linq;
-using System.Buffers.Binary;
 
 namespace EpicGames.Horde.Compute
 {
@@ -28,20 +24,21 @@ namespace EpicGames.Horde.Compute
 	{
 		readonly DirectoryReference _sandboxDir;
 		readonly IMemoryCache _memoryCache;
+		readonly bool _executeInProcess;
 		readonly ILogger _logger;
 
-		readonly bool _executeLocally = false;
-
 		/// <summary>
 		/// Constructor
 		/// </summary>
 		/// <param name="sandboxDir">Directory to use for reading/writing files</param>
 		/// <param name="memoryCache">Cache for nodes read from storage</param>
+		/// <param name="executeInProcess">Whether to execute any external assemblies in the current process</param>
 		/// <param name="logger">Logger for diagnostics</param>
-		public AgentMessageHandler(DirectoryReference sandboxDir, IMemoryCache memoryCache, ILogger logger)
+		public AgentMessageHandler(DirectoryReference sandboxDir, IMemoryCache memoryCache, bool executeInProcess, ILogger logger)
 		{
 			_sandboxDir = sandboxDir;
 			_memoryCache = memoryCache;
+			_executeInProcess = executeInProcess;
 			_logger = logger;
 		}
 
@@ -199,7 +196,6 @@ namespace EpicGames.Horde.Compute
 				_logger.LogInformation("Launching {Executable} {Arguments}", CommandLineArguments.Quote(executable), CommandLineArguments.Join(arguments));
 				await ExecuteProcessInternalAsync(channel, executable, arguments, workingDir, newEnvVars, cancellationToken);
 				_logger.LogInformation("Finished executing process");
-				ipcBuffer.Writer.MarkComplete();
 			}
 
 			_logger.LogInformation("Child process has shut down");
@@ -268,8 +264,9 @@ namespace EpicGames.Horde.Compute
 
 		async Task ExecuteProcessInternalAsync(AgentMessageChannel channel, string executable, IReadOnlyList<string> arguments, string? workingDir, IReadOnlyDictionary<string, string?>? envVars, CancellationToken cancellationToken)
 		{
+			string resolvedExecutable = FileReference.Combine(_sandboxDir, executable).FullName;
 			string resolvedWorkingDir = DirectoryReference.Combine(_sandboxDir, workingDir ?? String.Empty).FullName;
-			if (_executeLocally)
+			if (_executeInProcess && Path.GetFileNameWithoutExtension(resolvedExecutable).Equals("dotnet", StringComparison.OrdinalIgnoreCase))
 			{
 				List<(string, string?)> prevEnvVars = new List<(string, string?)>();
 				if (envVars != null)
@@ -286,12 +283,14 @@ namespace EpicGames.Horde.Compute
 
 				try
 				{
-					string resolvedExecutable = FileReference.Combine(_sandboxDir, arguments[0]).FullName;
+					string assemblyPath = FileReference.Combine(_sandboxDir, arguments[0]).FullName;
 					string[] mainArgs = arguments.Skip(1).ToArray();
 
+					_logger.LogWarning("Note: Loading and running {Assembly} in process", assemblyPath);
+
 					TaskCompletionSource<int> resultTcs = new TaskCompletionSource<int>();
 
-					Thread thread = new Thread(() => resultTcs.SetResult(AppDomain.CurrentDomain.ExecuteAssembly(resolvedExecutable, mainArgs)));
+					Thread thread = new Thread(() => resultTcs.SetResult(AppDomain.CurrentDomain.ExecuteAssembly(assemblyPath, mainArgs)));
 					thread.Start();
 
 					int result = await resultTcs.Task;
@@ -308,7 +307,6 @@ namespace EpicGames.Horde.Compute
 			}
 			else
 			{
-				string resolvedExecutable = FileReference.Combine(_sandboxDir, executable).FullName;
 				string resolvedCommandLine = CommandLineArguments.Join(arguments);
 
 				Dictionary<string, string> resolvedEnvVars = ManagedProcess.GetCurrentEnvVars();
diff --git a/Engine/Source/Programs/Shared/EpicGames.Horde/Compute/Buffers/PooledBuffer.cs b/Engine/Source/Programs/Shared/EpicGames.Horde/Compute/Buffers/PooledBuffer.cs
--- a/Engine/Source/Programs/Shared/EpicGames.Horde/Compute/Buffers/PooledBuffer.cs
+++ b/Engine/Source/Programs/Shared/EpicGames.Horde/Compute/Buffers/PooledBuffer.cs
@@ -40,9 +40,9 @@ namespace EpicGames.Horde.Compute.Buffers
 		}
 
 		/// <inheritdoc/>
-		public override ComputeBuffer AddRef()
+		public override PooledBuffer AddRef()
 		{
-			_detail.AddWriterRef();
+			_detail.AddRef();
 			return new PooledBuffer(_detail);
 		}
 	}
@@ -99,18 +99,12 @@ namespace EpicGames.Horde.Compute.Buffers
 		/// <inheritdoc/>
 		public override void SetReadEvent(int readerIdx) => _readerEvent.Set();
 
-		/// <inheritdoc/>
-		public override void ResetReadEvent(int readerIdx) => _readerEvent.Reset();
-
 		/// <inheritdoc/>
 		public override Task WaitForReadEvent(int readerIdx, CancellationToken cancellationToken) => _readerEvent.Task.WaitAsync(cancellationToken);
 
 		/// <inheritdoc/>
 		public override void SetWriteEvent() => _writerEvent.Set();
 
-		/// <inheritdoc/>
-		public override void ResetWriteEvent() => _writerEvent.Reset();
-
 		/// <inheritdoc/>
 		public override Task WaitForWriteEvent(CancellationToken cancellationToken) => _writerEvent.Task.WaitAsync(cancellationToken);
 	}
diff --git a/Engine/Source/Programs/Shared/EpicGames.Horde/Compute/Buffers/SharedMemoryBuffer.cs b/Engine/Source/Programs/Shared/EpicGames.Horde/Compute/Buffers/SharedMemoryBuffer.cs
--- a/Engine/Source/Programs/Shared/EpicGames.Horde/Compute/Buffers/SharedMemoryBuffer.cs
+++ b/Engine/Source/Programs/Shared/EpicGames.Horde/Compute/Buffers/SharedMemoryBuffer.cs
@@ -27,9 +27,9 @@ namespace EpicGames.Horde.Compute.Buffers
 		}
 
 		/// <inheritdoc/>
-		public override ComputeBuffer AddRef()
+		public override SharedMemoryBuffer AddRef()
 		{
-			_detail.AddWriterRef();
+			_detail.AddRef();
 			return new SharedMemoryBuffer(_detail);
 		}
 
@@ -102,8 +102,8 @@ namespace EpicGames.Horde.Compute.Buffers
 			MemoryMappedViewAccessor memoryMappedViewAccessor = memoryMappedFile.CreateViewAccessor();
 			MemoryMappedView memoryMappedView = new MemoryMappedView(memoryMappedViewAccessor);
 
-			Native.EventHandle writerEvent = Native.EventHandle.CreateNew($"{name}_W", EventResetMode.ManualReset, true, HandleInheritability.Inheritable);
-			Native.EventHandle readerEvent = Native.EventHandle.CreateNew($"{name}_R0", EventResetMode.ManualReset, true, HandleInheritability.Inheritable);
+			Native.EventHandle writerEvent = Native.EventHandle.CreateNew($"{name}_W", EventResetMode.AutoReset, true, HandleInheritability.Inheritable);
+			Native.EventHandle readerEvent = Native.EventHandle.CreateNew($"{name}_R0", EventResetMode.AutoReset, true, HandleInheritability.Inheritable);
 
 			HeaderPtr headerPtr = new HeaderPtr((ulong*)memoryMappedView.GetPointer(), 1, numChunks, chunkLength);
 			Memory<byte>[] chunks = CreateChunks(headerPtr.NumChunks, headerPtr.ChunkLength, memoryMappedView);
@@ -163,18 +163,12 @@ namespace EpicGames.Horde.Compute.Buffers
 		/// <inheritdoc/>
 		public override void SetReadEvent(int readerIdx) => _readerEvent.Set();
 
-		/// <inheritdoc/>
-		public override void ResetReadEvent(int readerIdx) => _readerEvent.Reset();
-
 		/// <inheritdoc/>
 		public override Task WaitForReadEvent(int readerIdx, CancellationToken cancellationToken) => _readerEvent.WaitOneAsync(cancellationToken);
 
 		/// <inheritdoc/>
 		public override void SetWriteEvent() => _writerEvent.Set();
 
-		/// <inheritdoc/>
-		public override void ResetWriteEvent() => _writerEvent.Reset();
-
 		/// <inheritdoc/>
 		public override Task WaitForWriteEvent(CancellationToken cancellationToken) => _writerEvent.WaitOneAsync(cancellationToken);
 	}
diff --git a/Engine/Source/Programs/Shared/EpicGames.Horde/Compute/Clients/LocalComputeClient.cs b/Engine/Source/Programs/Shared/EpicGames.Horde/Compute/Clients/LocalComputeClient.cs
--- a/Engine/Source/Programs/Shared/EpicGames.Horde/Compute/Clients/LocalComputeClient.cs
+++ b/Engine/Source/Programs/Shared/EpicGames.Horde/Compute/Clients/LocalComputeClient.cs
@@ -38,6 +38,7 @@ namespace EpicGames.Horde.Compute.Clients
 		readonly BackgroundTask _listenerTask;
 		readonly Socket _listener;
 		readonly Socket _socket;
+		readonly bool _executeInProcess;
 		readonly ILogger _logger;
 
 		/// <summary>
@@ -45,16 +46,18 @@ namespace EpicGames.Horde.Compute.Clients
 		/// </summary>
 		/// <param name="port">Port to connect on</param>
 		/// <param name="sandboxDir">Sandbox directory for the worker</param>
+		/// <param name="executeInProcess">Whether to run external assemblies in-process. Useful for debugging.</param>
 		/// <param name="logger">Logger for diagnostic output</param>
-		public LocalComputeClient(int port, DirectoryReference sandboxDir, ILogger logger)
+		public LocalComputeClient(int port, DirectoryReference sandboxDir, bool executeInProcess, ILogger logger)
 		{
 			_logger = logger;
+			_executeInProcess = executeInProcess;
 
 			_listener = new Socket(SocketType.Stream, ProtocolType.IP);
 			_listener.Bind(new IPEndPoint(IPAddress.Loopback, port));
 			_listener.Listen();
 
-			_listenerTask = BackgroundTask.StartNew(ctx => RunListenerAsync(_listener, sandboxDir, logger, ctx));
+			_listenerTask = BackgroundTask.StartNew(ctx => RunListenerAsync(_listener, sandboxDir, _executeInProcess, logger, ctx));
 
 			_socket = new Socket(SocketType.Stream, ProtocolType.IP);
 			_socket.Connect(IPAddress.Loopback, port);
@@ -71,7 +74,7 @@ namespace EpicGames.Horde.Compute.Clients
 		/// <summary>
 		/// Sets up the loopback listener and calls the server method
 		/// </summary>
-		static async Task RunListenerAsync(Socket listener, DirectoryReference sandboxDir, ILogger logger, CancellationToken cancellationToken)
+		static async Task RunListenerAsync(Socket listener, DirectoryReference sandboxDir, bool executeInProcess, ILogger logger, CancellationToken cancellationToken)
 		{
 			using Socket tcpSocket = await listener.AcceptAsync(cancellationToken);
 
@@ -79,7 +82,7 @@ namespace EpicGames.Horde.Compute.Clients
 
 			await using (RemoteComputeSocket socket = new RemoteComputeSocket(new TcpTransport(tcpSocket), ComputeSocketEndpoint.Remote, logger))
 			{
-				AgentMessageHandler worker = new AgentMessageHandler(sandboxDir, memoryCache, logger);
+				AgentMessageHandler worker = new AgentMessageHandler(sandboxDir, memoryCache, executeInProcess, logger);
 				await worker.RunAsync(socket, cancellationToken);
 				await socket.CloseAsync(cancellationToken);
 			}
diff --git a/Engine/Source/Programs/Shared/EpicGames.Horde/Compute/ComputeBuffer.cs b/Engine/Source/Programs/Shared/EpicGames.Horde/Compute/ComputeBuffer.cs
--- a/Engine/Source/Programs/Shared/EpicGames.Horde/Compute/ComputeBuffer.cs
+++ b/Engine/Source/Programs/Shared/EpicGames.Horde/Compute/ComputeBuffer.cs
@@ -24,12 +24,6 @@ namespace EpicGames.Horde.Compute
 		public const int MaxReaders = 16;
 
 		internal ComputeBufferDetail _detail;
-		private ComputeBufferWriter _writer;
-
-		/// <summary>
-		/// Writer for this buffer
-		/// </summary>
-		public ComputeBufferWriter Writer => _writer;
 
 		/// <summary>
 		/// Constructor
@@ -38,7 +32,6 @@ namespace EpicGames.Horde.Compute
 		internal ComputeBuffer(ComputeBufferDetail detail)
 		{
 			_detail = detail;
-			_writer = new ComputeBufferWriter(detail);
 		}
 
 		/// <inheritdoc/>
@@ -55,7 +48,7 @@ namespace EpicGames.Horde.Compute
 		{
 			if (_detail != null)
 			{
-				_detail.ReleaseWriter();
+				_detail.Release();
 				_detail = null!;
 			}
 		}
@@ -63,12 +56,21 @@ namespace EpicGames.Horde.Compute
 		/// <summary>
 		/// Creates a new reader for this buffer
 		/// </summary>
-		/// <returns></returns>
 		public ComputeBufferReader CreateReader()
 		{
-			int readerIdx = _detail.AllocateReader();
+			int readerIdx = _detail.CreateReader();
 			_detail.AddRef();
-			return new ComputeBufferReader(_detail, new ComputeBufferDetail.ReaderState(readerIdx));
+			return new ComputeBufferReader(_detail, readerIdx);
+		}
+
+		/// <summary>
+		/// Writer for this buffer
+		/// </summary>
+		public ComputeBufferWriter CreateWriter()
+		{
+			_detail.CreateWriter();
+			_detail.AddRef();
+			return new ComputeBufferWriter(_detail);
 		}
 
 		/// <summary>
@@ -83,14 +85,14 @@ namespace EpicGames.Horde.Compute
 	public sealed class ComputeBufferReader : IDisposable
 	{
 		ComputeBufferDetail _buffer;
-		readonly ComputeBufferDetail.ReaderState _readerState;
+		readonly int _readerIdx;
 
 		internal ComputeBufferDetail Detail => _buffer;
 
-		internal ComputeBufferReader(ComputeBufferDetail buffer, ComputeBufferDetail.ReaderState readerState)
+		internal ComputeBufferReader(ComputeBufferDetail buffer, int readerIdx)
 		{
 			_buffer = buffer;
-			_readerState = readerState;
+			_readerIdx = readerIdx;
 		}
 
 		/// <summary>
@@ -99,7 +101,8 @@ namespace EpicGames.Horde.Compute
 		public ComputeBufferReader AddRef()
 		{
 			_buffer.AddRef();
-			return new ComputeBufferReader(_buffer, _readerState);
+			_buffer.AddReaderRef(_readerIdx);
+			return new ComputeBufferReader(_buffer, _readerIdx);
 		}
 
 		/// <inheritdoc/>
@@ -107,6 +110,7 @@ namespace EpicGames.Horde.Compute
 		{
 			if (_buffer != null)
 			{
+				_buffer.ReleaseReaderRef(_readerIdx);
 				_buffer.Release();
 				_buffer = null!;
 			}
@@ -115,24 +119,24 @@ namespace EpicGames.Horde.Compute
 		/// <summary>
 		/// Detaches this reader from the underlying buffer
 		/// </summary>
-		public void Detach() => _buffer.DetachReader(_readerState);
+		public void Detach() => _buffer.DetachReader(_readerIdx);
 
 		/// <summary>
 		/// Whether this buffer is complete (no more data will be added)
 		/// </summary>
-		public bool IsComplete => _buffer.IsComplete(_readerState);
+		public bool IsComplete => _buffer.IsComplete(_readerIdx);
 
 		/// <summary>
 		/// Updates the read position
 		/// </summary>
 		/// <param name="length">Size of data that was read</param>
-		public void AdvanceReadPosition(int length) => _readerState.Offset += length;
+		public void AdvanceReadPosition(int length) => _buffer.AdvanceReadPosition(_readerIdx, length);
 
 		/// <summary>
 		/// Gets the next data to read
 		/// </summary>
 		/// <returns>Memory to read from</returns>
-		public ReadOnlyMemory<byte> GetReadBuffer() => _buffer.GetReadBuffer(_readerState);
+		public ReadOnlyMemory<byte> GetReadBuffer() => _buffer.GetReadBuffer(_readerIdx);
 
 		/// <summary>
 		/// Read from a buffer into another buffer
@@ -162,7 +166,7 @@ namespace EpicGames.Horde.Compute
 		/// <param name="minLength">Minimum amount of data to read</param>
 		/// <param name="cancellationToken">Cancellation token for the operation</param>
 		/// <returns>True if new data is available, false if the buffer is complete</returns>
-		public ValueTask<bool> WaitToReadAsync(int minLength, CancellationToken cancellationToken = default) => _buffer.WaitToReadAsync(_readerState, minLength, cancellationToken);
+		public ValueTask<bool> WaitToReadAsync(int minLength, CancellationToken cancellationToken = default) => _buffer.WaitToReadAsync(_readerIdx, minLength, cancellationToken);
 	}
 
 	/// <summary>
@@ -181,6 +185,7 @@ namespace EpicGames.Horde.Compute
 		/// </summary>
 		public ComputeBufferWriter AddRef()
 		{
+			_detail.AddRef();
 			_detail.AddWriterRef();
 			return new ComputeBufferWriter(_detail);
 		}
@@ -190,7 +195,8 @@ namespace EpicGames.Horde.Compute
 		{
 			if (_detail != null)
 			{
-				_detail.ReleaseWriter();
+				_detail.ReleaseWriterRef();
+				_detail.Release();
 				_detail = null!;
 			}
 		}
@@ -242,9 +248,10 @@ namespace EpicGames.Horde.Compute
 	/// <summary>
 	/// State shared between buffer instances
 	/// </summary>
+	[DebuggerTypeProxy(typeof(BufferDebugProxy))]
 	internal abstract class ComputeBufferDetail : IDisposable
 	{
-		internal const int HeaderSize = (3 + ComputeBuffer.MaxChunks) * sizeof(ulong);
+		internal const int HeaderSize = (2 + ComputeBuffer.MaxReaders + ComputeBuffer.MaxChunks) * sizeof(ulong);
 
 		/// <summary>
 		/// Write state for a chunk
@@ -285,7 +292,7 @@ namespace EpicGames.Horde.Compute
 			public ChunkState(WriteState writerState, int readerFlags, int length) : this(((ulong)writerState << 62) | ((ulong)readerFlags << 31) | (uint)length) { }
 
 			// Test whether a particular reader is still referencing the chunk
-			public bool HasReaderFlag(int readerIdx) => (Value & (1UL << (readerIdx + 31))) != 0;
+			public bool HasReaderFlag(int readerIdx) => (Value & (1UL << (31 + readerIdx))) != 0;
 
 			/// <inheritdoc/>
 			public override string ToString() => $"{WriteState}, Length: {Length}, Readers: {ReaderFlags}";
@@ -301,11 +308,10 @@ namespace EpicGames.Horde.Compute
 			public ChunkStatePtr(ulong* data) => _data = data;
 
 			// Current value of the chunk state
-			public ChunkState Value
-			{
-				get => new ChunkState(Interlocked.CompareExchange(ref *_data, 0, 0));
-				set => Interlocked.Exchange(ref *_data, value.Value);
-			}
+			public ChunkState Get() => new ChunkState(Interlocked.CompareExchange(ref * _data, 0, 0));
+
+			// Set the current state
+			public void Set(ChunkState value) => Interlocked.Exchange(ref *_data, value.Value);
 
 			// Attempt to update the chunk state
 			public bool TryUpdate(ChunkState prevState, ChunkState nextState) => Interlocked.CompareExchange(ref *_data, nextState.Value, prevState.Value) == prevState.Value;
@@ -326,35 +332,62 @@ namespace EpicGames.Horde.Compute
 			public void FinishWriting() => Interlocked.And(ref *_data, ~new ChunkState(WriteState.Writing, 0, 0).Value);
 
 			/// <inheritdoc/>
-			public override string ToString() => Value.ToString();
+			public override string ToString() => Get().ToString();
 		}
 
 		/// <summary>
 		/// State of a reader
 		/// </summary>
-		internal class ReaderState
+		protected internal record struct ReaderState(ulong Value)
+		{
+			public ReaderState(int chunkIdx, int offset, int refCount, bool detached)
+				: this((ulong)(uint)offset | ((ulong)(uint)chunkIdx << 32) | ((ulong)(uint)refCount << 40) | ((ulong)((detached? (1UL << 63) : 0))))
+			{ }
+
+			public int Offset => (int)(Value & 0xffffffff);
+			public int ChunkIdx => (int)((Value >> 32) & 0xff);
+			public int RefCount => (int)((Value >> 40) & 0x7fff);
+			public bool Detached => (Value & (1UL << 63)) != 0;
+
+			/// <inheritdoc/>
+			public override string ToString() => $"Chunk: {ChunkIdx}, Offset: {Offset}, RefCount: {RefCount}, Detached: {Detached}";
+		}
+
+		/// <summary>
+		/// Wraps a pointer to the state of a writer
+		/// </summary>
+		protected internal readonly unsafe struct ReaderStatePtr
 		{
-			public int ReaderIdx { get; }
+			readonly ulong* _data;
 
-			public int Offset { get; set; }
-			public int ChunkIdx { get; set; }
-			public bool Detached { get; set; }
+			public ReaderStatePtr(ulong* data) => _data = data;
+
+			// Current value of the chunk state
+			public ReaderState Get() => new ReaderState(Interlocked.CompareExchange(ref *_data, 0, 0));
 
-			public ReaderState(int readerIdx) => ReaderIdx = readerIdx;
+			// Update current state
+			public void Set(ReaderState value) => Interlocked.Exchange(ref *_data, value.Value);
+
+			// Compare and swap
+			public bool TryUpdate(ReaderState prevState, ReaderState nextState) => Interlocked.CompareExchange(ref *_data, nextState.Value, prevState.Value) == prevState.Value;
 		}
 
 		/// <summary>
 		/// State of the writer
 		/// </summary>
-		public record struct WriterState(ulong Value)
+		protected internal record struct WriterState(ulong Value)
 		{
-			public WriterState(int chunkIdx, int readerFlags, bool hasWrapped)
-				: this((ulong)(uint)chunkIdx | ((ulong)(uint)readerFlags << 32) | (hasWrapped ? (1UL << 63) : 0))
+			public WriterState(int chunkIdx, int readerFlags, int refCount, bool hasWrapped)
+				: this((ulong)(uint)chunkIdx | ((ulong)(uint)readerFlags << 32) | ((ulong)(uint)refCount << 48) | (hasWrapped ? (1UL << 63) : 0))
 			{ }
 
-			public int ReaderFlags => (int)(Value >> 32) & 0xffff;
 			public int ChunkIdx => (int)(Value & 0x7fffffff);
+			public int ReaderFlags => (int)(Value >> 32) & 0xffff;
+			public int RefCount => (int)(Value >> 48) & 0x7fff;
 			public bool HasWrapped => (Value & (1UL << 63)) != 0;
+
+			/// <inheritdoc/>
+			public override string ToString() => $"Chunk: {ChunkIdx}, ReaderFlags: {ReaderFlags}, RefCount: {RefCount}, HasWrapped: {HasWrapped}";
 		}
 
 		/// <summary>
@@ -366,12 +399,11 @@ namespace EpicGames.Horde.Compute
 
 			public WriterStatePtr(ulong* data) => _data = data;
 
-			// Current value of the chunk state
-			public WriterState Value
-			{
-				get => new WriterState(Interlocked.CompareExchange(ref *_data, 0, 0));
-				set => Interlocked.Exchange(ref *_data, value.Value);
-			}
+			// Get the current value
+			public WriterState Get() => new WriterState(Interlocked.CompareExchange(ref *_data, 0, 0));
+
+			// Set the current value
+			public void Set(WriterState State) => Interlocked.Exchange(ref *_data, State.Value);
 
 			// Compare and swap
 			public bool TryUpdate(WriterState prevState, WriterState nextState) => Interlocked.CompareExchange(ref *_data, nextState.Value, prevState.Value) == prevState.Value;
@@ -389,67 +421,56 @@ namespace EpicGames.Horde.Compute
 			public HeaderPtr(ulong* data, int numReaders, int numChunks, int chunkLength)
 			{
 				_data = data;
+				data[0] = ((ulong)(uint)chunkLength << 32) | ((ulong)(uint)numChunks << 16) | (uint)numReaders;
 
-				data[0] = ((ulong)(uint)numChunks << 32) | (uint)numReaders;
-				data[1] = (uint)chunkLength;
-
-				GetChunkStatePtr(0).Value = new ChunkState(WriteState.Writing, 0, 0);
+				GetChunkStatePtr(0).Set(new ChunkState(WriteState.Writing, 0, 0));
 			}
 
-			public int NumReaders => (int)_data[0];
-			public int NumChunks => (int)(_data[0] >> 32);
-			public int ChunkLength => (int)_data[1];
-			public uint AllocatedReaders => (uint)(_data[1] >> 32);
+			public int NumReaders => (int)(_data[0] & 0xffff);
+			public int NumChunks => (int)((_data[0] >> 16) & 0xffff);
+			public int ChunkLength => (int)(_data[0] >> 32);
+
+			public WriterStatePtr GetWriterStatePtr() => new WriterStatePtr(_data + 1);
 
-			public WriterStatePtr Writer => new WriterStatePtr(_data + 2);
+			public ReaderStatePtr GetReaderStatePtr(int readerIdx) => new ReaderStatePtr(_data + 2 + readerIdx);
 
-			public ChunkStatePtr GetChunkStatePtr(int chunkIdx) => new ChunkStatePtr(_data + 3 + chunkIdx);
+			public ChunkStatePtr GetChunkStatePtr(int chunkIdx) => new ChunkStatePtr(_data + 2 + ComputeBuffer.MaxReaders + chunkIdx);
+		}
+
+		class BufferDebugProxy
+		{
+			public int ChunkLength { get; }
+			public int RefCount { get; }
+			public WriterState Writer { get; }
+			public ReaderState[] Readers { get; }
+			public ChunkState[] Chunks { get; }
 
-			public int AllocateReader()
+			public BufferDebugProxy(ComputeBufferDetail buffer)
 			{
-				for (int readerIdx = 0; readerIdx < NumReaders; readerIdx++)
-				{
-					ulong allocatedReaders = _data[1];
-					ulong readerFlag = 1UL << (32 + readerIdx);
+				RefCount = buffer._refCount;
 
-					if((allocatedReaders & readerFlag) == 0 && Interlocked.CompareExchange(ref _data[1], allocatedReaders | readerFlag, allocatedReaders) == allocatedReaders)
-					{
-						for (; ; )
-						{
-							WriterState state = Writer.Value;
-							if (state.HasWrapped)
-							{
-								throw new InvalidOperationException("Cannot create a new reader after writer has wrapped back to the first chunk");
-							}
+				HeaderPtr headerPtr = buffer._headerPtr;
+				ChunkLength = headerPtr.ChunkLength;
 
-							if (Writer.TryUpdate(state, new WriterState(state.ChunkIdx, state.ReaderFlags | (1 << readerIdx), state.HasWrapped)))
-							{
-								for (int WriteChunkIdx = 0; WriteChunkIdx <= state.ChunkIdx; WriteChunkIdx++)
-								{
-									GetChunkStatePtr(WriteChunkIdx).StartReading(readerIdx);
-								}
-								return readerIdx;
-							}
-						}
-					}
+				Writer = headerPtr.GetWriterStatePtr().Get();
+
+				Chunks = new ChunkState[headerPtr.NumChunks];
+				for (int chunkIdx = 0; chunkIdx < headerPtr.NumChunks; chunkIdx++)
+				{
+					Chunks[chunkIdx] = headerPtr.GetChunkStatePtr(chunkIdx).Get();
 				}
-				throw new InvalidOperationException("Unable to allocate reader; all available readers are in use.");
-			}
 
-			void ReleaseReader(int readerIdx)
-			{
-				for (int Idx = 0; Idx < NumChunks; Idx++)
+				Readers = new ReaderState[headerPtr.NumReaders];
+				for (int readerIdx = 0; readerIdx < headerPtr.NumReaders; readerIdx++)
 				{
-					GetChunkStatePtr(Idx).FinishReading(readerIdx);
+					Readers[readerIdx] = headerPtr.GetReaderStatePtr(readerIdx).Get();
 				}
-				Interlocked.And(ref _data[1], ~(1UL << (readerIdx + 32)));
 			}
 		}
 
 		readonly HeaderPtr _headerPtr;
 		readonly Memory<byte>[] _chunks;
 		int _refCount = 1;
-		int _writerRefCount = 1;
 
 		/// <summary>
 		/// Constructor
@@ -479,27 +500,6 @@ namespace EpicGames.Horde.Compute
 			}
 		}
 
-		/// <summary>
-		/// Increments the writer ref count
-		/// </summary>
-		internal void AddWriterRef()
-		{
-			AddRef();
-			Interlocked.Increment(ref _writerRefCount);
-		}
-
-		/// <summary>
-		/// Decrement the reference count on this object, and dispose of it once it reaches zero
-		/// </summary>
-		public void ReleaseWriter()
-		{
-			if (Interlocked.Decrement(ref _writerRefCount) == 0)
-			{
-				MarkComplete();
-			}
-			Release();
-		}
-
 		/// <inheritdoc/>
 		public void Dispose()
 		{
@@ -528,11 +528,6 @@ namespace EpicGames.Horde.Compute
 			}
 		}
 
-		/// <summary>
-		/// Resets a read event
-		/// </summary>
-		public abstract void ResetReadEvent(int readerIdx);
-
 		/// <summary>
 		/// Waits for a read event to be signalled
 		/// </summary>
@@ -543,47 +538,212 @@ namespace EpicGames.Horde.Compute
 		/// </summary>
 		public abstract void SetWriteEvent();
 
-		/// <summary>
-		/// Resets the write event
-		/// </summary>
-		public abstract void ResetWriteEvent();
-
 		/// <summary>
 		/// Waits for the write event to be signalled
 		/// </summary>
 		public abstract Task WaitForWriteEvent(CancellationToken cancellationToken);
 
-#pragma warning disable IDE0051 // Remove unused private members
-		// For debugging purposes only
-		ChunkState[] ChunkStates => Enumerable.Range(0, _headerPtr.NumChunks).Select(x => _headerPtr.GetChunkStatePtr(x).Value).ToArray();
-#pragma warning restore IDE0051 // Remove unused private members
-
 		/// <summary>
 		/// Allocate a new reader
 		/// </summary>
-		public int AllocateReader() => _headerPtr.AllocateReader();
+		public int CreateReader()
+		{
+			for (int readerIdx = 0; readerIdx < _headerPtr.NumReaders; readerIdx++)
+			{
+				ReaderStatePtr readerStatePtr = _headerPtr.GetReaderStatePtr(readerIdx);
+				for (; ;)
+				{
+					ReaderState readerState = readerStatePtr.Get();
+					if (readerState.RefCount > 0)
+					{
+						break;
+					}
+					if (readerStatePtr.TryUpdate(readerState, new ReaderState(0, 0, 1, false)))
+					{
+						WriterStatePtr writerStatePtr = _headerPtr.GetWriterStatePtr();
+						for (; ; )
+						{
+							WriterState writerState = writerStatePtr.Get();
+							if (writerState.HasWrapped)
+							{
+								throw new InvalidOperationException("Cannot create a new reader after writer has wrapped back to the first chunk");
+							}
+							if (writerStatePtr.TryUpdate(writerState, new WriterState(writerState.ChunkIdx, writerState.ReaderFlags | (1 << readerIdx), writerState.RefCount, writerState.HasWrapped)))
+							{
+								for (int WriteChunkIdx = 0; WriteChunkIdx <= writerState.ChunkIdx; WriteChunkIdx++)
+								{
+									_headerPtr.GetChunkStatePtr(WriteChunkIdx).StartReading(readerIdx);
+								}
+								return readerIdx;
+							}
+						}
+					}
+				}
+			}
+			throw new InvalidOperationException("Unable to allocate reader; all available readers are in use.");
+		}
+
+		public void AddReaderRef(int readerIdx)
+		{
+			ReaderStatePtr readerStatePtr = _headerPtr.GetReaderStatePtr(readerIdx);
+			for (; ; )
+			{
+				ReaderState readerState = readerStatePtr.Get();
+				if (readerState.RefCount == 0)
+				{
+					throw new InvalidOperationException("Refcount for reader is zero");
+				}
+				if (readerStatePtr.TryUpdate(readerState, new ReaderState(readerState.ChunkIdx, readerState.Offset, readerState.RefCount + 1, readerState.Detached)))
+				{
+					break;
+				}
+			}
+		}
+
+		public void ReleaseReaderRef(int readerIdx)
+		{
+			ReaderStatePtr readerStatePtr = _headerPtr.GetReaderStatePtr(readerIdx);
+			for (; ; )
+			{
+				ReaderState readerState = readerStatePtr.Get();
+				if (readerState.RefCount == 0)
+				{
+					throw new InvalidOperationException("Refcount for reader is already zero");
+				}
+
+				if (readerState.RefCount == 1)
+				{
+					for (int Idx = 0; Idx < _headerPtr.NumChunks; Idx++)
+					{
+						_headerPtr.GetChunkStatePtr(Idx).FinishReading(readerIdx);
+					}
+				}
+
+				if (readerStatePtr.TryUpdate(readerState, new ReaderState(readerState.ChunkIdx, readerState.Offset, readerState.RefCount - 1, readerState.Detached)))
+				{
+					break;
+				}
+			}
+		}
+
+		public void CreateWriter()
+		{
+			WriterStatePtr writerStatePtr = _headerPtr.GetWriterStatePtr();
+			for (; ; )
+			{
+				WriterState writerState = writerStatePtr.Get();
+				if (writerState.RefCount > 0)
+				{
+					throw new InvalidOperationException("Writer has already been created for this buffer");
+				}
+				if (writerStatePtr.TryUpdate(writerState, new WriterState(writerState.ChunkIdx, writerState.ReaderFlags, 1, writerState.HasWrapped)))
+				{
+					ChunkStatePtr chunkStatePtr = _headerPtr.GetChunkStatePtr(writerState.ChunkIdx);
+					for (; ; )
+					{
+						ChunkState chunkState = chunkStatePtr.Get();
+						if (chunkStatePtr.TryUpdate(chunkState, new ChunkState(WriteState.Writing, chunkState.ReaderFlags, chunkState.Length)))
+						{
+							break;
+						}
+					}
+					break;
+				}
+			}
+		}
+
+		public void AddWriterRef()
+		{
+			WriterStatePtr writerStatePtr = _headerPtr.GetWriterStatePtr();
+			for (; ; )
+			{
+				WriterState writerState = writerStatePtr.Get();
+				if (writerState.RefCount == 0)
+				{
+					throw new InvalidOperationException("Writer does not exist for this buffer");
+				}
+				if (writerStatePtr.TryUpdate(writerState, new WriterState(writerState.ChunkIdx, writerState.ReaderFlags, writerState.RefCount + 1, writerState.HasWrapped)))
+				{
+					break;
+				}
+			}
+		}
+
+		public void ReleaseWriterRef()
+		{
+			WriterStatePtr writerStatePtr = _headerPtr.GetWriterStatePtr();
+			for (; ; )
+			{
+				WriterState writerState = writerStatePtr.Get();
+				if (writerState.RefCount == 0)
+				{
+					throw new InvalidOperationException("Writer does not exist for this buffer");
+				}
+
+				if (writerState.RefCount == 1)
+				{
+					MarkComplete();
+				}
+
+				if (writerStatePtr.TryUpdate(writerState, new WriterState(writerState.ChunkIdx, writerState.ReaderFlags, writerState.RefCount - 1, writerState.HasWrapped)))
+				{
+					break;
+				}
+			}
+		}
 
 		/// <inheritdoc/>
-		public bool IsComplete(ReaderState readerState)
+		public bool IsComplete(int readerIdx)
 		{
-			ChunkState chunkState = _headerPtr.GetChunkStatePtr(readerState.ChunkIdx).Value;
+			ReaderState readerState = _headerPtr.GetReaderStatePtr(readerIdx).Get();
+			if(readerState.Detached)
+			{
+				return true;
+			}
+
+			ChunkState chunkState = _headerPtr.GetChunkStatePtr(readerState.ChunkIdx).Get();
 			return chunkState.WriteState == WriteState.Complete && readerState.Offset == chunkState.Length;
 		}
 
 		/// <inheritdoc/>
-		public void DetachReader(ReaderState readerState)
+		public void DetachReader(int readerIdx)
+		{
+			ReaderStatePtr readerStatePtr = _headerPtr.GetReaderStatePtr(readerIdx);
+			for (; ; )
+			{
+				ReaderState readerState = readerStatePtr.Get();
+				if (readerStatePtr.TryUpdate(readerState, new ReaderState(readerState.ChunkIdx, readerState.Offset, readerState.RefCount, true)))
+				{
+					SetReadEvent(readerIdx);
+					break;
+				}
+			}
+		}
+
+		/// <inheritdoc/>
+		public void AdvanceReadPosition(int readerIdx, int offset)
 		{
-			readerState.Detached = true;
-			SetReadEvent(readerState.ReaderIdx);
+			ReaderStatePtr readerStatePtr = _headerPtr.GetReaderStatePtr(readerIdx);
+			for (; ; )
+			{
+				ReaderState readerState = readerStatePtr.Get();
+				if (readerStatePtr.TryUpdate(readerState, new ReaderState(readerState.ChunkIdx, readerState.Offset + offset, readerState.RefCount, readerState.Detached)))
+				{
+					SetReadEvent(readerIdx);
+					break;
+				}
+			}
 		}
 
 		/// <inheritdoc/>
-		public ReadOnlyMemory<byte> GetReadBuffer(ReaderState readerState)
+		public ReadOnlyMemory<byte> GetReadBuffer(int readerIdx)
 		{
+			ReaderState readerState = _headerPtr.GetReaderStatePtr(readerIdx).Get();
+
 			ChunkStatePtr chunkStatePtr = _headerPtr.GetChunkStatePtr(readerState.ChunkIdx);
-			ChunkState chunkState = chunkStatePtr.Value;
+			ChunkState chunkState = chunkStatePtr.Get();
 
-			if (chunkState.HasReaderFlag(readerState.ReaderIdx))
+			if (chunkState.HasReaderFlag(readerIdx))
 			{
 				return _chunks[readerState.ChunkIdx].Slice(readerState.Offset, chunkState.Length - readerState.Offset);
 			}
@@ -594,27 +754,24 @@ namespace EpicGames.Horde.Compute
 		}
 
 		/// <inheritdoc/>
-		public async ValueTask<bool> WaitToReadAsync(ReaderState readerState, int minLength, CancellationToken cancellationToken = default)
+		public async ValueTask<bool> WaitToReadAsync(int readerIdx, int minLength, CancellationToken cancellationToken = default)
 		{
-			int readerIdx = readerState.ReaderIdx;
 			for (; ; )
 			{
+				ReaderStatePtr readerStatePtr = _headerPtr.GetReaderStatePtr(readerIdx);
+				ReaderState readerState = readerStatePtr.Get();
 				if (readerState.Detached)
 				{
 					return false;
 				}
 
 				ChunkStatePtr chunkStatePtr = _headerPtr.GetChunkStatePtr(readerState.ChunkIdx);
-				ChunkState chunkState = chunkStatePtr.Value;
+				ChunkState chunkState = chunkStatePtr.Get();
 
 				if (!chunkState.HasReaderFlag(readerIdx))
 				{
 					// Wait until the current chunk is readable
-					ResetReadEvent(readerIdx);
-					if (!chunkState.HasReaderFlag(readerIdx))
-					{
-						await WaitForReadEvent(readerIdx, cancellationToken);
-					}
+					await WaitForReadEvent(readerIdx, cancellationToken);
 				}
 				else if (readerState.Offset + minLength <= chunkState.Length)
 				{
@@ -624,11 +781,7 @@ namespace EpicGames.Horde.Compute
 				else if (chunkState.WriteState == WriteState.Writing)
 				{
 					// Wait until there is more data in the chunk
-					ResetReadEvent(readerIdx);
-					if (_headerPtr.GetChunkStatePtr(readerState.ChunkIdx).Value == chunkState)
-					{
-						await WaitForReadEvent(readerIdx, cancellationToken);
-					}
+					await WaitForReadEvent(readerIdx, cancellationToken);
 				}
 				else if (readerState.Offset < chunkState.Length || chunkState.WriteState == WriteState.Complete)
 				{
@@ -641,12 +794,13 @@ namespace EpicGames.Horde.Compute
 					chunkStatePtr.FinishReading(readerIdx);
 					SetWriteEvent();
 
-					if (++readerState.ChunkIdx == _chunks.Length)
+					int nextChunkIdx = readerState.ChunkIdx + 1;
+					if (nextChunkIdx == _headerPtr.NumChunks)
 					{
-						readerState.ChunkIdx = 0;
+						nextChunkIdx = 0;
 					}
 
-					readerState.Offset = 0;
+					readerStatePtr.TryUpdate(readerState, new ReaderState(nextChunkIdx, 0, readerState.RefCount, readerState.Detached));
 				}
 				else
 				{
@@ -658,10 +812,10 @@ namespace EpicGames.Horde.Compute
 		/// <inheritdoc/>
 		public bool MarkComplete()
 		{
-			WriterState writerState = _headerPtr.Writer.Value;
+			WriterState writerState = _headerPtr.GetWriterStatePtr().Get();
 
 			ChunkStatePtr chunkStatePtr = _headerPtr.GetChunkStatePtr(writerState.ChunkIdx);
-			if (chunkStatePtr.Value.WriteState != WriteState.Complete)
+			if (chunkStatePtr.Get().WriteState != WriteState.Complete)
 			{
 				chunkStatePtr.MarkComplete();
 				SetAllReadEvents();
@@ -676,10 +830,10 @@ namespace EpicGames.Horde.Compute
 		{
 			if (size > 0)
 			{
-				WriterState writerState = _headerPtr.Writer.Value;
+				WriterState writerState = _headerPtr.GetWriterStatePtr().Get();
 
 				ChunkStatePtr chunkStatePtr = _headerPtr.GetChunkStatePtr(writerState.ChunkIdx);
-				ChunkState chunkState = chunkStatePtr.Value;
+				ChunkState chunkState = chunkStatePtr.Get();
 
 				Debug.Assert(chunkState.WriteState == WriteState.Writing);
 				chunkStatePtr.Append(size);
@@ -691,9 +845,9 @@ namespace EpicGames.Horde.Compute
 		/// <inheritdoc/>
 		public Memory<byte> GetWriteBuffer()
 		{
-			WriterState writerState = _headerPtr.Writer.Value;
+			WriterState writerState = _headerPtr.GetWriterStatePtr().Get();
 
-			ChunkState chunkState = _headerPtr.GetChunkStatePtr(writerState.ChunkIdx).Value;
+			ChunkState chunkState = _headerPtr.GetChunkStatePtr(writerState.ChunkIdx).Get();
 			if (chunkState.WriteState == WriteState.Writing)
 			{
 				return _chunks[writerState.ChunkIdx].Slice(chunkState.Length);
@@ -713,13 +867,13 @@ namespace EpicGames.Horde.Compute
 			}
 
 			// Get the current chunk we're writing to
-			WriterState writerState = _headerPtr.Writer.Value;
+			WriterState writerState = _headerPtr.GetWriterStatePtr().Get();
 			int writeChunkIdx = writerState.ChunkIdx;
 
 			ChunkStatePtr writeChunkStatePtr = _headerPtr.GetChunkStatePtr(writeChunkIdx);
 
 			// Check if we can append to this chunk
-			ChunkState chunkState = writeChunkStatePtr.Value;
+			ChunkState chunkState = writeChunkStatePtr.Get();
 			if (chunkState.WriteState == WriteState.Writing)
 			{
 				int length = chunkState.Length;
@@ -748,21 +902,21 @@ namespace EpicGames.Horde.Compute
 			ChunkStatePtr nextWriteChunkStatePtr = _headerPtr.GetChunkStatePtr(nextWriteChunkIdx);
 			for (; ; )
 			{
-				ChunkState nextWriteChunkState = nextWriteChunkStatePtr.Value;
+				ChunkState nextWriteChunkState = nextWriteChunkStatePtr.Get();
 				if (nextWriteChunkState.ReaderFlags != 0)
 				{
 					await WaitForWriteEvent(cancellationToken);
-					ResetWriteEvent();
 				}
 				else if (nextWriteChunkStatePtr.TryUpdate(nextWriteChunkState, new ChunkState(WriteState.Writing, writerState.ReaderFlags, 0)))
 				{
-					if (_headerPtr.Writer.TryUpdate(writerState, new WriterState(nextWriteChunkIdx, writerState.ReaderFlags, nextWriteChunkIdx == 0)))
+					WriterStatePtr writerStatePtr = _headerPtr.GetWriterStatePtr();
+					if (writerStatePtr.TryUpdate(writerState, new WriterState(nextWriteChunkIdx, writerState.ReaderFlags, writerState.RefCount, nextWriteChunkIdx == 0)))
 					{
 						break;
 					}
 					else
 					{
-						writerState = _headerPtr.Writer.Value;
+						writerState = writerStatePtr.Get();
 					}
 				}
 			}
diff --git a/Engine/Source/Programs/Shared/EpicGames.Horde/Compute/ComputeSocket.cs b/Engine/Source/Programs/Shared/EpicGames.Horde/Compute/ComputeSocket.cs
--- a/Engine/Source/Programs/Shared/EpicGames.Horde/Compute/ComputeSocket.cs
+++ b/Engine/Source/Programs/Shared/EpicGames.Horde/Compute/ComputeSocket.cs
@@ -55,10 +55,10 @@ namespace EpicGames.Horde.Compute
 			AttachRecvBuffer(channelId, recvBuffer);
 			AttachSendBuffer(channelId, sendBuffer);
 
-			using (ComputeBufferReader recvBufferReader = recvBuffer.CreateReader())
-			{
-				return new ComputeChannel(recvBufferReader, sendBuffer.Writer);
-			}
+			using ComputeBufferReader recvBufferReader = recvBuffer.CreateReader();
+			using ComputeBufferWriter sendBufferWriter = sendBuffer.CreateWriter();
+
+			return new ComputeChannel(recvBufferReader, sendBufferWriter);
 		}
 	}
 
@@ -78,15 +78,15 @@ namespace EpicGames.Horde.Compute
 		/// </summary>
 		public const string IpcEnvVar = "UE_HORDE_COMPUTE_IPC";
 
-		readonly SharedMemoryBuffer _commandBuffer;
+		readonly ComputeBufferWriter _commandBufferWriter;
 		readonly List<ComputeBuffer> _buffers = new List<ComputeBuffer>();
 
 		/// <summary>
 		/// Creates a socket for a worker
 		/// </summary>
-		private WorkerComputeSocket(SharedMemoryBuffer commandBuffer)
+		private WorkerComputeSocket(ComputeBufferWriter commandBufferWriter)
 		{
-			_commandBuffer = commandBuffer;
+			_commandBufferWriter = commandBufferWriter;
 		}
 
 		/// <summary>
@@ -109,14 +109,14 @@ namespace EpicGames.Horde.Compute
 		/// <param name="commandBufferName">Name of the command buffer</param>
 		public static WorkerComputeSocket Open(string commandBufferName)
 		{
-			SharedMemoryBuffer commandBuffer = SharedMemoryBuffer.OpenExisting(commandBufferName);
-			return new WorkerComputeSocket(commandBuffer);
+			using SharedMemoryBuffer commandBuffer = SharedMemoryBuffer.OpenExisting(commandBufferName);
+			return new WorkerComputeSocket(commandBuffer.CreateWriter());
 		}
 
 		/// <inheritdoc/>
 		public void Dispose()
 		{
-			_commandBuffer.Dispose();
+			_commandBufferWriter.Dispose();
 		}
 
 		/// <inheritdoc/>
@@ -137,11 +137,12 @@ namespace EpicGames.Horde.Compute
 
 		void AttachBuffer(IpcMessage message, int channelId, string bufferName)
 		{
-			MemoryWriter writer = new MemoryWriter(_commandBuffer.Writer.GetWriteBuffer());
+			MemoryWriter writer = new MemoryWriter(_commandBufferWriter.GetWriteBuffer());
 			writer.WriteUnsignedVarInt((int)message);
 			writer.WriteUnsignedVarInt(channelId);
 			writer.WriteString(bufferName);
-			_commandBuffer.Writer.AdvanceWritePosition(writer.Length);
+
+			_commandBufferWriter.AdvanceWritePosition(writer.Length);
 		}
 	}
 
@@ -364,7 +365,7 @@ namespace EpicGames.Horde.Compute
 		readonly SendSegment _bodySegment = new SendSegment();
 
 		/// <inheritdoc/>
-		public async ValueTask SendAsync(int id, ReadOnlyMemory<byte> memory, CancellationToken cancellationToken = default)
+		async ValueTask SendAsync(int id, ReadOnlyMemory<byte> memory, CancellationToken cancellationToken = default)
 		{
 			if (memory.Length > 0)
 			{
@@ -397,21 +398,14 @@ namespace EpicGames.Horde.Compute
 		/// <inheritdoc/>
 		public override void AttachRecvBuffer(int channelId, ComputeBuffer recvBuffer)
 		{
-			bool complete;
 			lock (_lockObject)
 			{
-				complete = _complete;
-				if (!complete)
+				if (!_complete)
 				{
-					_recvBufferWriters.Add(channelId, recvBuffer.Writer.AddRef());
+					_recvBufferWriters.Add(channelId, recvBuffer.CreateWriter());
 					_recvTask ??= BackgroundTask.StartNew(ctx => RunRecvTaskAsync(_transport, ctx));
 				}
 			}
-
-			if (complete)
-			{
-				recvBuffer.Writer.MarkComplete();
-			}
 		}
 
 		void DetachRecvBuffer(Dictionary<int, ComputeBufferWriter> cachedWriters, int id)
