Horde: Fix remote execution errors on Mac.

- Max length for semaphore name is 31 characters on MacOS.
- Also add some custom assert/verify macros, and include errno strings in errors.

diff --git a/Engine/Source/Programs/Horde/Samples/RemoteWorkerCpp/ComputeBuffer.cpp b/Engine/Source/Programs/Horde/Samples/RemoteWorkerCpp/ComputeBuffer.cpp
--- a/Engine/Source/Programs/Horde/Samples/RemoteWorkerCpp/ComputeBuffer.cpp
+++ b/Engine/Source/Programs/Horde/Samples/RemoteWorkerCpp/ComputeBuffer.cpp
@@ -437,7 +437,7 @@ struct FComputeBufferDetail
 					for (;;)
 					{
 						FWriterState WriterState = WriterStatePtr.Get();
-						assert(!WriterState.HasWrapped());
+						UE_COMPUTE_ASSERT(!WriterState.HasWrapped());
 
 						if (WriterStatePtr.TryUpdate(WriterState, FWriterState(WriterState.GetChunkIdx(), WriterState.GetReaderFlags() | (1 << ReaderIdx), WriterState.GetRefCount(), WriterState.HasWrapped())))
 						{
@@ -460,7 +460,7 @@ struct FComputeBufferDetail
 		for (; ; )
 		{
 			FReaderState ReaderState = ReaderStatePtr.Get();
-			assert(ReaderState.GetRefCount() > 0);
+			UE_COMPUTE_ASSERT(ReaderState.GetRefCount() > 0);
 
 			if (ReaderStatePtr.TryUpdate(ReaderState, FReaderState(ReaderState.GetChunkIdx(), ReaderState.GetOffset(), ReaderState.GetRefCount() + 1, ReaderState.IsDetached())))
 			{
@@ -475,7 +475,7 @@ struct FComputeBufferDetail
 		for (; ; )
 		{
 			FReaderState ReaderState = ReaderStatePtr.Get();
-			assert(ReaderState.GetRefCount() > 0);
+			UE_COMPUTE_ASSERT(ReaderState.GetRefCount() > 0);
 
 			if (ReaderState.GetRefCount() == 1)
 			{
@@ -498,7 +498,7 @@ struct FComputeBufferDetail
 		for (; ; )
 		{
 			FWriterState WriterState = WriterStatePtr.Get();
-			assert(WriterState.GetRefCount() == 0);
+			UE_COMPUTE_ASSERT(WriterState.GetRefCount() == 0);
 
 			if (WriterStatePtr.TryUpdate(WriterState, FWriterState(WriterState.GetChunkIdx(), WriterState.GetReaderFlags(), 1, WriterState.HasWrapped())))
 			{
@@ -522,7 +522,7 @@ struct FComputeBufferDetail
 		for (; ; )
 		{
 			FWriterState WriterState = WriterStatePtr.Get();
-			assert(WriterState.GetRefCount() > 0);
+			UE_COMPUTE_ASSERT(WriterState.GetRefCount() > 0);
 
 			if (WriterStatePtr.TryUpdate(WriterState, FWriterState(WriterState.GetChunkIdx(), WriterState.GetReaderFlags(), WriterState.GetRefCount() + 1, WriterState.HasWrapped())))
 			{
@@ -537,7 +537,7 @@ struct FComputeBufferDetail
 		for (; ; )
 		{
 			FWriterState writerState = writerStatePtr.Get();
-			assert(writerState.GetRefCount() > 0);
+			UE_COMPUTE_ASSERT(writerState.GetRefCount() > 0);
 
 			if (writerState.GetRefCount() == 1)
 			{
@@ -668,7 +668,7 @@ struct FComputeBufferDetail
 			}
 			else
 			{
-				assert(false);
+				UE_COMPUTE_ASSERT(false);
 			}
 		}
 	}
@@ -697,7 +697,7 @@ struct FComputeBufferDetail
 			FChunkStatePtr ChunkStatePtr = Header->GetChunkStatePtr(WriterState.GetChunkIdx());
 			FChunkState ChunkState = ChunkStatePtr.Get();
 
-			assert(ChunkState.GetWriteState() == EWriteState::Writing);
+			UE_COMPUTE_ASSERT(ChunkState.GetWriteState() == EWriteState::Writing);
 			ChunkStatePtr.Append(Size);
 
 			SetAllReadEvents();
@@ -740,7 +740,7 @@ struct FComputeBufferDetail
 
 	unsigned char* WaitToWrite(size_t MinSize, int TimeoutMs)
 	{
-		assert(MinSize <= Header->ChunkLength);
+		UE_COMPUTE_ASSERT(MinSize <= Header->ChunkLength);
 
 		// Get the current chunk we're writing to
 		FWriterState WriterState = Header->GetWriterStatePtr().Get();
@@ -1012,25 +1012,25 @@ void FComputeBufferReader::Close()
 
 void FComputeBufferReader::Detach()
 {
-	assert(Detail);
+	UE_COMPUTE_ASSERT(Detail);
 	Detail->DetachReader(ReaderIdx);
 }
 
 bool FComputeBufferReader::IsComplete() const
 {
-	assert(Detail);
+	UE_COMPUTE_ASSERT(Detail);
 	return Detail->IsComplete(ReaderIdx);
 }
 
 void FComputeBufferReader::AdvanceReadPosition(size_t Size)
 {
-	assert(Detail);
+	UE_COMPUTE_ASSERT(Detail);
 	Detail->AdvanceReadPosition(ReaderIdx, Size);
 }
 
 size_t FComputeBufferReader::GetMaxReadSize() const
 {
-	assert(Detail);
+	UE_COMPUTE_ASSERT(Detail);
 	return Detail->GetMaxReadSize(ReaderIdx);
 }
 
@@ -1055,7 +1055,7 @@ size_t FComputeBufferReader::Read(void* Buffer, size_t MaxSize, int TimeoutMs)
 
 const unsigned char* FComputeBufferReader::WaitToRead(size_t MinSize, int TimeoutMs)
 {
-	assert(Detail);
+	UE_COMPUTE_ASSERT(Detail);
 	return Detail->WaitToRead(ReaderIdx, MinSize, TimeoutMs);
 }
 
@@ -1067,7 +1067,7 @@ FComputeBufferReader::FComputeBufferReader(FComputeBufferDetail* InDetail, int I
 
 const char* FComputeBufferReader::GetName() const
 {
-	assert(Detail);
+	UE_COMPUTE_ASSERT(Detail);
 	return Detail->Name;
 }
 
diff --git a/Engine/Source/Programs/Horde/Samples/RemoteWorkerCpp/ComputePlatform.cpp b/Engine/Source/Programs/Horde/Samples/RemoteWorkerCpp/ComputePlatform.cpp
--- a/Engine/Source/Programs/Horde/Samples/RemoteWorkerCpp/ComputePlatform.cpp
+++ b/Engine/Source/Programs/Horde/Samples/RemoteWorkerCpp/ComputePlatform.cpp
@@ -20,6 +20,14 @@
 	#include <sys/stat.h>
 	#include <time.h>
 	#include <fcntl.h>
+
+	struct FFormatErrno { };
+
+	static std::ostream &operator<<(std::ostream &os, const FFormatErrno &)
+	{
+		os << "(errno=" << errno << ": " << strerror(errno) << ")";
+		return os;
+	}
 #endif
 
 ///////////////////////////////////////////////////
@@ -120,7 +128,7 @@ bool FComputeEvent::Wait(int timeoutMs)
 	struct timespec ts;
 	if (clock_gettime(CLOCK_REALTIME, &ts) == -1)
 	{
-		assert(false);
+		UE_COMPUTE_ASSERT(false);
 		return false;
 	}
 	
@@ -196,7 +204,7 @@ bool FComputeMemoryMappedFile::Create(const char* Name, long long Capacity)
 	int Fd = shm_open(Name, O_CREAT | O_EXCL | O_RDWR, 0666);
 	if(Fd < 0)
 	{
-		std::cerr << "Unable to create shared memory object '" << Name << "' (errno=" << errno << ")" << std::endl;
+		std::cerr << "Unable to create shared memory object '" << Name << "' " << FFormatErrno() << std::endl;
 		return false;
 	}
 
@@ -206,14 +214,14 @@ bool FComputeMemoryMappedFile::Create(const char* Name, long long Capacity)
 
 	if(ftruncate(Fd, MappedSize) < 0)
 	{
-		std::cerr << "Unable to update size of shared memory object '" << Name << "' to " << MappedSize << " (errno=" << errno << ")" << std::endl;
+		std::cerr << "Unable to update size of shared memory object '" << Name << "' to " << MappedSize << " " << FFormatErrno() << std::endl;
 		return false;
 	}
 
 	Pointer = mmap(nullptr, MappedSize, PROT_READ | PROT_WRITE, MAP_SHARED, Fd, 0);
 	if(Pointer == MAP_FAILED)
 	{
-		std::cerr << "Unable to map shared memory object '" << Name << " (errno=" << errno << ")" << std::endl;
+		std::cerr << "Unable to map shared memory object '" << Name << " " << FFormatErrno() << std::endl;
 		return false;
 	}
 	
@@ -243,7 +251,7 @@ bool FComputeMemoryMappedFile::OpenExisting(const char* Name)
 	int Fd = shm_open(Name, O_RDWR, 0666);
 	if(Fd < 0)
 	{
-		std::cerr << "Unable to open shared memory object '" << Name << "' (errno=" << errno << ")" << std::endl;
+		std::cerr << "Unable to open shared memory object '" << Name << "' " << FFormatErrno() << std::endl;
 		return false;
 	}
 
@@ -256,7 +264,7 @@ bool FComputeMemoryMappedFile::OpenExisting(const char* Name)
 	Pointer = mmap(nullptr, MappedSize, PROT_READ | PROT_WRITE, MAP_SHARED, Fd, 0);
 	if(Pointer == MAP_FAILED)
 	{
-		std::cerr << "Unable to map shared memory object '" << Name << " (errno=" << errno << ")" << std::endl;
+		std::cerr << "Unable to map shared memory object '" << Name << " " << FFormatErrno() << std::endl;
 		return false;
 	}
 	
@@ -311,6 +319,16 @@ void* FComputeMemoryMappedFile::GetPointer() const
 
 /////////////////////////////////////////////////// 
 
+void FComputePlatform::AssertFailed(const char* Expr, const char* File, int Line)
+{
+	std::cerr << "Assertion failed: '" << Expr << std::endl;
+	std::cerr << "  at " << File << "(" << Line << ")" << std::endl;
+#if UE_COMPUTE_PLATFORM_WINDOWS
+	DebugBreak();
+#endif
+	exit(1);
+}
+
 bool FComputePlatform::GetEnvironmentVariable(const char* Name, char* Buffer, size_t BufferLen)
 {
 #if UE_COMPUTE_PLATFORM_WINDOWS
@@ -338,7 +356,7 @@ void FComputePlatform::CreateUniqueName(char* NameBuffer, size_t NameBufferLen)
 #else
 	struct timespec ts;
 	clock_gettime(CLOCK_MONOTONIC, &ts);
-	snprintf(NameBuffer, NameBufferLen, "/COMPUTE_%u_%zu_%zu_%lu", getpid(), (size_t)ts.tv_sec, (size_t)ts.tv_nsec, AtomicIncrement(&Counter));
+	snprintf(NameBuffer, NameBufferLen, "/UEC_%u%zu%zu_%lu", getpid(), (size_t)ts.tv_sec, (size_t)ts.tv_nsec, AtomicIncrement(&Counter));
 #endif
 }
 
diff --git a/Engine/Source/Programs/Horde/Samples/RemoteWorkerCpp/ComputePlatform.h b/Engine/Source/Programs/Horde/Samples/RemoteWorkerCpp/ComputePlatform.h
--- a/Engine/Source/Programs/Horde/Samples/RemoteWorkerCpp/ComputePlatform.h
+++ b/Engine/Source/Programs/Horde/Samples/RemoteWorkerCpp/ComputePlatform.h
@@ -21,6 +21,14 @@
 	#define UE_COMPUTE_PLATFORM_LINUX 1
 #endif
 
+#ifdef NDEBUG
+	#define UE_COMPUTE_ASSERT(x) ((void)0)
+	#define UE_COMPUTE_VERIFY(x) (!!(x))
+#else
+	#define UE_COMPUTE_ASSERT(x) ((void)UE_COMPUTE_VERIFY(x))
+	#define UE_COMPUTE_VERIFY(x) (!!(x) || (FComputePlatform::AssertFailed(#x, __FILE__, __LINE__), true))
+#endif
+
 //
 // IPC manually-reset event object
 //
@@ -86,6 +94,9 @@ struct FComputePlatform
 	// General
 	//
 
+	// Signal that an assertion has failed
+	static void AssertFailed(const char* Expr, const char* File, int Line);
+	
 	// Reads an environment variable
 	static bool GetEnvironmentVariable(const char* Name, char* Buffer, size_t BufferLen);
 
diff --git a/Engine/Source/Programs/Horde/Samples/RemoteWorkerCpp/ComputeSocket.cpp b/Engine/Source/Programs/Horde/Samples/RemoteWorkerCpp/ComputeSocket.cpp
--- a/Engine/Source/Programs/Horde/Samples/RemoteWorkerCpp/ComputeSocket.cpp
+++ b/Engine/Source/Programs/Horde/Samples/RemoteWorkerCpp/ComputeSocket.cpp
@@ -144,7 +144,7 @@ void FWorkerComputeSocket::RunServer(FComputeBufferReader& CommandBufferReader,
 				}
 				else
 				{
-					assert(false);
+					UE_COMPUTE_ASSERT(false);
 				}
 			}
 			break;
@@ -163,12 +163,12 @@ void FWorkerComputeSocket::RunServer(FComputeBufferReader& CommandBufferReader,
 				}
 				else
 				{
-					assert(false);
+					UE_COMPUTE_ASSERT(false);
 				}
 			}
 			break;
 		default:
-			assert(false);
+			UE_COMPUTE_ASSERT(false);
 			return;
 		}
 
@@ -308,7 +308,7 @@ public:
 			}
 			else
 			{
-				assert(false);
+				UE_COMPUTE_ASSERT(false);
 			}
 		}
 	}
diff --git a/Engine/Source/Programs/Horde/Samples/RemoteWorkerCpp/Tests.cpp b/Engine/Source/Programs/Horde/Samples/RemoteWorkerCpp/Tests.cpp
--- a/Engine/Source/Programs/Horde/Samples/RemoteWorkerCpp/Tests.cpp
+++ b/Engine/Source/Programs/Horde/Samples/RemoteWorkerCpp/Tests.cpp
@@ -2,11 +2,10 @@
 
 #include "ComputeBuffer.h"
 #include "ComputeSocket.h"
+#include "ComputePlatform.h"
 #include <assert.h>
 #include <iostream>
 
-#define verify(x) ((x) || (throw #x, true))
-
 void RunTests()
 {
 	void ComputeBufferTest();
@@ -23,7 +22,7 @@ void ComputeBufferTest()
 	char TestData[] = "hello world!";
 
 	FComputeBuffer Buffer;
-	verify(Buffer.CreateNew(FComputeBuffer::FParams()));
+	UE_COMPUTE_VERIFY(Buffer.CreateNew(FComputeBuffer::FParams()));
 
 	FComputeBufferWriter Writer = Buffer.CreateWriter();
 	unsigned char* WriteBuffer = Writer.WaitToWrite(sizeof(TestData));
@@ -32,14 +31,14 @@ void ComputeBufferTest()
 
 	FComputeBufferReader Reader = Buffer.CreateReader();
 	const unsigned char* ReadBuffer = Reader.WaitToRead(sizeof(TestData));
-	verify(memcmp(ReadBuffer, TestData, sizeof(TestData)) == 0);
-	assert(!Reader.IsComplete());
+	UE_COMPUTE_ASSERT(memcmp(ReadBuffer, TestData, sizeof(TestData)) == 0);
+	UE_COMPUTE_ASSERT(!Reader.IsComplete());
 
 	Reader.AdvanceReadPosition(sizeof(TestData));
-	assert(!Reader.IsComplete());
+	UE_COMPUTE_ASSERT(!Reader.IsComplete());
 
 	Writer.MarkComplete();
-	assert(Reader.IsComplete());
+	UE_COMPUTE_ASSERT(Reader.IsComplete());
 }
 
 template<size_t TestDataSize> void CheckChannelSendRecv(FComputeChannel& SendChannel, FComputeChannel& RecvChannel, const char(&TestData)[TestDataSize])
@@ -47,19 +46,19 @@ template<size_t TestDataSize> void CheckChannelSendRecv(FComputeChannel& SendCha
 	SendChannel.Send(TestData, TestDataSize);
 
 	char RecvTestData[TestDataSize];
-	verify(RecvChannel.Recv(RecvTestData, TestDataSize) == TestDataSize);
+	UE_COMPUTE_VERIFY(RecvChannel.Recv(RecvTestData, TestDataSize) == TestDataSize);
 
-	verify(memcmp(TestData, RecvTestData, TestDataSize) == 0);
+	UE_COMPUTE_VERIFY(memcmp(TestData, RecvTestData, TestDataSize) == 0);
 }
 
 void ComputeSocketTest()
 {
 	// Buffers for transferring between client and server
 	FComputeBuffer ClientToServerBuffer;
-	verify(ClientToServerBuffer.CreateNew(FComputeBuffer::FParams()));
+	UE_COMPUTE_VERIFY(ClientToServerBuffer.CreateNew(FComputeBuffer::FParams()));
 
 	FComputeBuffer ServerToClientBuffer;
-	verify(ServerToClientBuffer.CreateNew(FComputeBuffer::FParams()));
+	UE_COMPUTE_VERIFY(ServerToClientBuffer.CreateNew(FComputeBuffer::FParams()));
 
 	// Client transport
 	std::unique_ptr<FComputeSocket> ClientSocket = CreateComputeSocket(std::make_unique<FBufferTransport>(ClientToServerBuffer.CreateWriter(), ServerToClientBuffer.CreateReader()), EComputeSocketEndpoint::Local);
