Horde: Include the stream view in the AutoSDK workspace marker.

diff --git a/Engine/Source/Programs/Horde/Horde.Agent/Execution/PerforceExecutor.cs b/Engine/Source/Programs/Horde/Horde.Agent/Execution/PerforceExecutor.cs
--- a/Engine/Source/Programs/Horde/Horde.Agent/Execution/PerforceExecutor.cs
+++ b/Engine/Source/Programs/Horde/Horde.Agent/Execution/PerforceExecutor.cs
@@ -7,6 +7,7 @@ using System.IO;
 using System.Linq;
 using System.Net.Http;
 using System.Runtime.InteropServices;
+using System.Text;
 using System.Text.RegularExpressions;
 using System.Threading;
 using System.Threading.Tasks;
@@ -77,6 +78,15 @@ namespace Horde.Agent.Execution
 				int autoSdkChangeNumber = await _autoSdkWorkspace.GetLatestChangeAsync(cancellationToken);
 
 				string syncText = $"Synced to CL {autoSdkChangeNumber}";
+				if (_autoSdkWorkspaceInfo.View.Count > 0)
+				{
+					StringBuilder syncTextBuilder = new StringBuilder(syncText);
+					foreach (string line in _autoSdkWorkspaceInfo.View)
+					{
+						syncTextBuilder.Append($"\nView: {line}");
+					}
+					syncText = syncTextBuilder.ToString();
+				}
 
 				FileReference syncFile = FileReference.Combine(_autoSdkWorkspace.MetadataDir, "Synced.txt");
 				if (!FileReference.Exists(syncFile) || FileReference.ReadAllText(syncFile) != syncText)
