Horde: Add a retry loop for schema upgrades. Server pods can sometimes race on completing upgrades, and the exception will keep throwing from awaiting the upgrade task.

diff --git a/Engine/Source/Programs/Horde/Horde.Server/Server/MongoService.cs b/Engine/Source/Programs/Horde/Horde.Server/Server/MongoService.cs
--- a/Engine/Source/Programs/Horde/Horde.Server/Server/MongoService.cs
+++ b/Engine/Source/Programs/Horde/Horde.Server/Server/MongoService.cs
@@ -683,6 +683,24 @@ namespace Horde.Server.Server
 				return;
 			}
 
+			int attemptIdx = 1;
+			for(; ;)
+			{
+				try
+				{
+					await UpdateIndexesInternalAsync<T>(collectionName, collection, newIndexes, cancellationToken);
+					break;
+				}
+				catch (Exception ex)
+				{
+					_logger.LogError(ex, "Error updating indexes ({Message}) - retrying (attempt {Attempt})", ex.Message, attemptIdx);
+					attemptIdx++;
+				}
+			}
+		}
+
+		private async Task UpdateIndexesInternalAsync<T>(string collectionName, IMongoCollection<T> collection, MongoIndex<T>[] newIndexes, CancellationToken cancellationToken)
+		{
 			// Find all the current indexes, excluding the default
 			Dictionary<string, MongoIndex> nameToExistingIndex = new Dictionary<string, MongoIndex>(StringComparer.Ordinal);
 			using (IAsyncCursor<BsonDocument> cursor = await collection.Indexes.ListAsync(cancellationToken))
