UnrealAutomatedTests plugin - Seperate hover style on history changelist and swarm range button

#rnx

diff --git a/Engine/Source/Programs/Horde/HordeDashboard/src/Plugins/UnrealAutomatedTests/components/UnrealAutomatedTests.tsx b/Engine/Source/Programs/Horde/HordeDashboard/src/Plugins/UnrealAutomatedTests/components/UnrealAutomatedTests.tsx
--- a/Engine/Source/Programs/Horde/HordeDashboard/src/Plugins/UnrealAutomatedTests/components/UnrealAutomatedTests.tsx
+++ b/Engine/Source/Programs/Horde/HordeDashboard/src/Plugins/UnrealAutomatedTests/components/UnrealAutomatedTests.tsx
@@ -90,7 +90,9 @@ const styles = mergeStyleSets({
       }
    },
    itemHighlighted: {
-      background: theme.palette.neutralLighter
+      borderWidth: "1px",
+      borderColor: "#888",
+      borderStyle: "solid"
    },
    historyList: {
       zIndex: 1,
@@ -253,14 +255,14 @@ const HistoryItem: React.FC<{ item: TestStateHistoryItem, testName: string, sele
    const gutterStyle = gutterHystoryStyles.get(item.State);
 
    return (
-      <Stack horizontal className={`${styles.itemHover} ${selected ? styles.itemHighlighted : ""}`}>
+      <Stack horizontal className={`${selected ? styles.itemHighlighted : ""}`}>
          <Stack className={gutterStyle}></Stack>
          <Stack.Item>
-            <Link to={`/testreport/${item.TestdataId}?test=${testName}`}>
+            <Link to={`/testreport/${item.TestdataId}?test=${testName}`} className={`${styles.itemHover}`}>
                <Text variant="smallPlus">{getStateLabel(item.State)} on {item.Change}</Text>
             </Link>
             {item.RangeUrl &&
-               <a href={item.RangeUrl} target="blank">
+               <a href={item.RangeUrl} target="blank" className={`${styles.itemHover}`}>
                   <Text variant="smallPlus" styles={{ root: {color: theme.palette.neutralDark,  paddingLeft: 4, paddingRight: 4 } }}>[ Swarm Range ]</Text>
                </a>}
          </Stack.Item>
