Horde: Support specifying a template id from P4V
#rnx
#skipci

diff --git a/Engine/Source/Programs/Horde/HordeDashboard/src/components/NewBuild.tsx b/Engine/Source/Programs/Horde/HordeDashboard/src/components/NewBuild.tsx
--- a/Engine/Source/Programs/Horde/HordeDashboard/src/components/NewBuild.tsx
+++ b/Engine/Source/Programs/Horde/HordeDashboard/src/components/NewBuild.tsx
@@ -643,6 +643,7 @@ export const NewBuild: React.FC<{ streamId: string; show: boolean; onClose: (new
    // @todo: better way of checking whether we are in a preflight submit
    const isPreflightSubmit = !!query.get("shelvedchange");
    const isFromP4V = !!query.get("p4v");
+   const queryTemplateId = !query.get("templateId") ? "" : query.get("templateId")!;
 
    let defaultShelvedChange: string | undefined = query.get("shelvedchange") ? query.get("shelvedchange")! : undefined;
 
@@ -746,15 +747,15 @@ export const NewBuild: React.FC<{ streamId: string; show: boolean; onClose: (new
             const defaultTemplate = templateData.templates.find(t => t.id === defaultChange.templateId);
             if (defaultTemplate) {
                defaultChange.name = "Latest Success - " + defaultTemplate.name;
-            }          
+            }
          }
-         
+
          if (defaultChange.name) {
             defaultPreflightQuery = [defaultChange];
-         }         
-      }      
+         }
+      }
    }
-   
+
    const defaultStreamPreflightTemplate = stream.templates.find(t => t.id === stream.defaultPreflight?.templateId);
 
    if (!defaultPreflightQuery && stream.defaultPreflight?.templateId) {
@@ -766,7 +767,7 @@ export const NewBuild: React.FC<{ streamId: string; show: boolean; onClose: (new
       if (!defaultPreflightQuery) {
          console.error(`Unable to find default stream preflight template ${stream.defaultPreflight?.templateId} in stream templates`);
       }
-   }   
+   }
 
    if (buildParams?.preflight) {
       templates = templates.filter(t => t.allowPreflights);
@@ -783,9 +784,22 @@ export const NewBuild: React.FC<{ streamId: string; show: boolean; onClose: (new
       // handle preflight redirect case
       if (!t && defaultShelvedChange && !jobDetails) {
 
-         t = defaultStreamPreflightTemplate;
+         if (queryTemplateId) {
+            t = templateData.templates?.find(t => t.id === queryTemplateId);
+            if (!t) {
+               console.error(`Unable to find queryTemplateId ${queryTemplateId} in stream ${streamId}`)
+            } else if (!t.allowPreflights) {
+               console.error(`Template does not allow preflights: queryTemplateId ${queryTemplateId} in stream ${streamId}`)
+               t = undefined;
+            }
+
+         }
+
          if (!t) {
-            console.error(`Stream default preflight template cannot be found for stream ${stream.fullname} : stream defaultPreflightTemplate ${stream.defaultPreflight?.templateId}, will use first template in list`);
+            t = defaultStreamPreflightTemplate;
+            if (!t) {
+               console.error(`Stream default preflight template cannot be found for stream ${stream.fullname} : stream defaultPreflightTemplate ${stream.defaultPreflight?.templateId}, will use first template in list`);
+            }
          }
       }
 
@@ -1401,7 +1415,7 @@ export const NewBuild: React.FC<{ streamId: string; show: boolean; onClose: (new
                   if (user.jobTemplateSettings) {
                      dashboard.jobTemplateSettings = user.jobTemplateSettings;
                   }
-                  
+
                } catch (reason) {
                   console.log(`Error on updating notifications: ${reason}`);
                }
diff --git a/Engine/Source/Programs/Horde/HordeDashboard/src/components/Preflight.tsx b/Engine/Source/Programs/Horde/HordeDashboard/src/components/Preflight.tsx
--- a/Engine/Source/Programs/Horde/HordeDashboard/src/components/Preflight.tsx
+++ b/Engine/Source/Programs/Horde/HordeDashboard/src/components/Preflight.tsx
@@ -24,9 +24,7 @@ function setError(message: string) {
 // redirect from external source, where horde stream id, etc are not known by that application
 export const PreflightRedirector: React.FC = () => {
 
-
    const [state, setState] = useState({ preflightQueried: false })
-
    const navigate = useNavigate();
    const query = useQuery();
    const { projectStore } = useBackend();
@@ -37,6 +35,9 @@ export const PreflightRedirector: React.FC = () => {
    // whether to autosubmit
    const autosubmit = !query.get("submit") ? "" : query.get("submit")!;
 
+   // whether a template is specified
+   const templateId = !query.get("templateId") ? "" : query.get("templateId")!;
+
    if (!change) {
       setError("No preflight change specified");
       return null;
@@ -80,6 +81,8 @@ export const PreflightRedirector: React.FC = () => {
 
    if (!state.preflightQueried) {
 
+      console.log(`Redirecting preflight: ${window.location.href}`);
+
       backend.getJobs({ filter: "id", count: 1, preflightChange: cl }).then(result => {
 
          if (result && result.length === 1) {
@@ -90,6 +93,10 @@ export const PreflightRedirector: React.FC = () => {
                url += "&autosubmit=true";
             }
 
+            if (templateId) {
+               url += `&templateId=${templateId}`;
+            }
+
             navigate(url, { replace: true });
             return;
          }
@@ -100,6 +107,10 @@ export const PreflightRedirector: React.FC = () => {
             url += "&autosubmit=true";
          }
 
+         if (templateId) {
+            url += `&templateId=${templateId}`;
+         }
+
          navigate(url, { replace: true });
 
 
