Horde: Surface prelight template issues usig the error dialog
#skipci
#rnx

diff --git a/Engine/Source/Programs/Horde/HordeDashboard/src/components/NewBuild.tsx b/Engine/Source/Programs/Horde/HordeDashboard/src/components/NewBuild.tsx
--- a/Engine/Source/Programs/Horde/HordeDashboard/src/components/NewBuild.tsx
+++ b/Engine/Source/Programs/Horde/HordeDashboard/src/components/NewBuild.tsx
@@ -785,14 +785,25 @@ export const NewBuild: React.FC<{ streamId: string; show: boolean; onClose: (new
       if (!t && defaultShelvedChange && !jobDetails) {
 
          if (queryTemplateId) {
+            let errorReason = "";
             t = templateData.templates?.find(t => t.id === queryTemplateId);
             if (!t) {
-               console.error(`Unable to find queryTemplateId ${queryTemplateId} in stream ${streamId}`)
+               errorReason = `Unable to find queryTemplateId ${queryTemplateId} in stream ${streamId}`;
+               console.error(errorReason);
+
             } else if (!t.allowPreflights) {
-               console.error(`Template does not allow preflights: queryTemplateId ${queryTemplateId} in stream ${streamId}`)
+               errorReason = `Template does not allow preflights: queryTemplateId ${queryTemplateId} in stream ${streamId}`;
+               console.error(errorReason);
                t = undefined;
             }
 
+            if (errorReason) {
+               ErrorHandler.set({
+                  reason: `${errorReason}`,
+                  title: `Preflight Template Error`,
+                  message: `There was an issue with the specified preflight template.\n\nReason: ${errorReason}\n\nTime: ${moment.utc().format("MMM Do, HH:mm z")}`
+               }, true);   
+            }
          }
 
          if (!t) {
