Horde: Make a distinction between passwords and ticket for Perforce cluster credentials

When password is set for a service account, a Perforce login needs to be made. But in case of a ticket that's not necessary (and will not be accepted by login).

diff --git a/Engine/Source/Programs/Horde/Horde.Server/Perforce/PerforceService.cs b/Engine/Source/Programs/Horde/Horde.Server/Perforce/PerforceService.cs
--- a/Engine/Source/Programs/Horde/Horde.Server/Perforce/PerforceService.cs
+++ b/Engine/Source/Programs/Horde/Horde.Server/Perforce/PerforceService.cs
@@ -130,12 +130,14 @@ namespace Horde.Server.Perforce
 		{
 			public string UserName { get; }
 			public string? Password { get; }
+			public string? Ticket { get; }
 			public DateTime? ExpiresAt { get; }
 
-			public Credentials(string userName, string? password, DateTime? expiresAt)
+			public Credentials(string userName, string? password, string? ticket, DateTime? expiresAt)
 			{
 				UserName = userName;
 				Password = password;
+				Ticket = ticket;
 				ExpiresAt = expiresAt;
 			}
 		}
@@ -380,7 +382,7 @@ namespace Horde.Server.Perforce
 		{
 			if (_perforceUserOverride != null)
 			{
-				return new Credentials(_perforceUserOverride, null, null);
+				return new Credentials(_perforceUserOverride, null, null, null);
 			}
 			else if (cluster.ServiceAccount != null)
 			{
@@ -389,11 +391,11 @@ namespace Horde.Server.Perforce
 				{
 					throw new Exception($"No credentials defined for {cluster.ServiceAccount} on {cluster.Name}");
 				}
-				return new Credentials(credentials.UserName, credentials.Password, null);
+				return new Credentials(credentials.UserName, credentials.Password, credentials.Ticket, null);
 			}
 			else
 			{
-				return new Credentials(PerforceSettings.Default.UserName, null, null);
+				return new Credentials(PerforceSettings.Default.UserName, null, null, null);
 			}
 		}
 
@@ -401,7 +403,7 @@ namespace Horde.Server.Perforce
 		{
 			if (_perforceUserOverride != null)
 			{
-				return new Credentials(_perforceUserOverride, null, null);
+				return new Credentials(_perforceUserOverride, null, null, null);
 			}
 			else if (userName != null)
 			{
@@ -411,7 +413,7 @@ namespace Horde.Server.Perforce
 				}
 				else
 				{
-					return new Credentials(userName, null, null);
+					return new Credentials(userName, null, null, null);
 				}
 			}
 			else
@@ -434,6 +436,11 @@ namespace Horde.Server.Perforce
 				IPerforceServer server = await GetServerAsync(cluster);
 				Credentials credentials = await GetCredentialsAsync(cluster, userName, cancellationToken);
 				handle = await CreatePooledConnectionAsync(server.ServerAndPort, credentials, null, cancellationToken);
+
+				if (credentials.Password != null && credentials.Ticket == null)
+				{
+					await handle.LoginAsync(credentials.Password, cancellationToken);
+				}
 			}
 			return handle;
 		}
@@ -561,7 +568,7 @@ namespace Horde.Server.Perforce
 				}
 
 				DateTime expiresAt = DateTime.UtcNow + new TimeSpan(response.Data.TicketExpiration * TimeSpan.TicksPerSecond) - TimeSpan.FromMinutes(15.0);
-				ticketInfo = new Credentials(userName, response.Data.Ticket, expiresAt);
+				ticketInfo = new Credentials(userName, null, response.Data.Ticket, expiresAt);
 
 				lock (_userCredentialsByCluster)
 				{
diff --git a/Engine/Source/Programs/Horde/Horde.Server/Server/GlobalConfig.cs b/Engine/Source/Programs/Horde/Horde.Server/Server/GlobalConfig.cs
--- a/Engine/Source/Programs/Horde/Horde.Server/Server/GlobalConfig.cs
+++ b/Engine/Source/Programs/Horde/Horde.Server/Server/GlobalConfig.cs
@@ -765,7 +765,12 @@ namespace Horde.Server.Server
 		/// <summary>
 		/// Password for the user
 		/// </summary>
-		public string Password { get; set; } = String.Empty;
+		public string? Password { get; set; } = String.Empty;
+
+		/// <summary>
+		/// Login ticket for the user (will be used instead of password if set)
+		/// </summary>
+		public string? Ticket { get; set; } = String.Empty;
 	}
 
 	/// <summary>
