Horde: Add a 'paused' state for agents, which cancels all current leases and prevents the agent taking on work.

diff --git a/Engine/Source/Programs/Horde/Horde.Agent/Leases/LeaseManager.cs b/Engine/Source/Programs/Horde/Horde.Agent/Leases/LeaseManager.cs
--- a/Engine/Source/Programs/Horde/Horde.Agent/Leases/LeaseManager.cs
+++ b/Engine/Source/Programs/Horde/Horde.Agent/Leases/LeaseManager.cs
@@ -233,6 +233,9 @@ namespace Horde.Agent.Leases
 			Stopwatch updateTimer = Stopwatch.StartNew();
 			Queue<TimeSpan> updateTimes = new Queue<TimeSpan>();
 
+			// Whether there is an active user on this machine. For workstations, this determines whether the agent will take on work.
+			bool paused = false;
+
 			// Loop until we're ready to exit
 			Stopwatch updateCapabilitiesTimer = Stopwatch.StartNew();
 			for (; ; )
@@ -280,6 +283,10 @@ namespace Horde.Agent.Leases
 				{
 					updateSessionRequest.Status = AgentStatus.Unhealthy;
 				}
+				else if (paused)
+				{
+					updateSessionRequest.Status = AgentStatus.Paused;
+				}
 				else
 				{
 					updateSessionRequest.Status = AgentStatus.Ok;
@@ -373,6 +380,10 @@ namespace Horde.Agent.Leases
 					{
 						_statusService.Set(false, _activeLeases.Count, "Attempting to connect to server...");
 					}
+					else if (paused)
+					{
+						_statusService.Set(true, 0, "Paused");
+					}
 					else if (_activeLeases.Count == 0)
 					{
 						_statusService.Set(true, 0, "Waiting for work");
diff --git a/Engine/Source/Programs/Horde/Horde.Server/Agents/AgentService.cs b/Engine/Source/Programs/Horde/Horde.Server/Agents/AgentService.cs
--- a/Engine/Source/Programs/Horde/Horde.Server/Agents/AgentService.cs
+++ b/Engine/Source/Programs/Horde/Horde.Server/Agents/AgentService.cs
@@ -378,6 +378,10 @@ namespace Horde.Server.Agents
 		bool CanUseTaskSource(IAgent agent, ITaskSource taskSource)
 		{
 			TaskSourceFlags flags = taskSource.Flags;
+			if ((flags & TaskSourceFlags.AllowWhenPaused) == 0 && agent.Status == AgentStatus.Paused)
+			{
+				return false;
+			}
 			if ((flags & TaskSourceFlags.AllowWhenDisabled) == 0 && !agent.Enabled)
 			{
 				return false;
@@ -703,7 +707,7 @@ namespace Horde.Server.Agents
 				}
 
 				// If the agent is stopping, cancel all the leases. Clear out the current session once it's complete.
-				if (status == AgentStatus.Stopping)
+				if (status == AgentStatus.Stopping || status == AgentStatus.Paused)
 				{
 					foreach (AgentLease lease in leases)
 					{
@@ -1048,6 +1052,7 @@ namespace Horde.Server.Agents
 			int numAgentsTotalOk = agentList.Count(a => a.Enabled && a.Status == AgentStatus.Ok);
 			int numAgentsTotalStopping = agentList.Count(a => a.Enabled && a.Status == AgentStatus.Stopping);
 			int numAgentsTotalUnhealthy = agentList.Count(a => a.Enabled && a.Status == AgentStatus.Unhealthy);
+			int numAgentsTotalPaused = agentList.Count(a => a.Enabled && a.Status == AgentStatus.Paused);
 			int numAgentsTotalUnspecified = agentList.Count(a => a.Enabled && a.Status == AgentStatus.Unspecified);
 			
 			_agentTotalCount.Record(numAgentsTotal);
@@ -1056,6 +1061,7 @@ namespace Horde.Server.Agents
 			_agentTotalCount.Record(numAgentsTotalOk, new KeyValuePair<string, object?>("status", "ok"));
 			_agentTotalCount.Record(numAgentsTotalStopping, new KeyValuePair<string, object?>("status", "stopping"));
 			_agentTotalCount.Record(numAgentsTotalUnhealthy, new KeyValuePair<string, object?>("status", "unhealthy"));
+			_agentTotalCount.Record(numAgentsTotalPaused, new KeyValuePair<string, object?>("status", "paused"));
 			_agentTotalCount.Record(numAgentsTotalUnspecified, new KeyValuePair<string, object?>("status", "unspecified"));
 		}
 
diff --git a/Engine/Source/Programs/Horde/Horde.Server/Tasks/ITaskSource.cs b/Engine/Source/Programs/Horde/Horde.Server/Tasks/ITaskSource.cs
--- a/Engine/Source/Programs/Horde/Horde.Server/Tasks/ITaskSource.cs
+++ b/Engine/Source/Programs/Horde/Horde.Server/Tasks/ITaskSource.cs
@@ -14,7 +14,6 @@ using Google.Protobuf.Reflection;
 using Google.Protobuf.WellKnownTypes;
 using Horde.Server.Agents;
 using Horde.Server.Agents.Leases;
-using Horde.Server.Utilities;
 using HordeCommon;
 using Microsoft.Extensions.Logging;
 
@@ -40,6 +39,11 @@ namespace Horde.Server.Tasks
 		/// Allows this source to schedule tasks when the agent is disabled
 		/// </summary>
 		AllowWhenDisabled = 2,
+
+		/// <summary>
+		/// Allows this source to schedule tasks when the agent is paused
+		/// </summary>
+		AllowWhenPaused = 4,
 	}
 
 	/// <summary>
