Horde: Add functionality to the Horde tray app which will disable the agent when a workstation is active.

diff --git a/Engine/Source/Programs/Horde/Horde.Agent.TrayApp/Program.cs b/Engine/Source/Programs/Horde/Horde.Agent.TrayApp/Program.cs
--- a/Engine/Source/Programs/Horde/Horde.Agent.TrayApp/Program.cs
+++ b/Engine/Source/Programs/Horde/Horde.Agent.TrayApp/Program.cs
@@ -2,11 +2,16 @@
 
 using System.Diagnostics;
 using System.IO.Pipes;
+using System.Runtime.InteropServices;
 using EpicGames.Core;
+using Horde.Agent.TrayApp.Forms;
 using Horde.Agent.TrayApp.Properties;
+using Microsoft.Win32;
 
 namespace Horde.Agent.TrayApp
 {
+	record struct IdleStat(string Name, long Value, long MinValue);
+
 	static class Program
 	{
 		const string MutexName = "Horde.Agent.TrayApp-Mutex";
@@ -84,13 +89,38 @@ namespace Horde.Agent.TrayApp
 	{
 		readonly NotifyIcon _trayIcon;
 		readonly BackgroundTask _clientTask;
+		readonly BackgroundTask _tickPauseStateTask;
 		readonly BackgroundTask _waitForExitTask;
 		readonly Control _mainThreadInvokeTarget;
 
+		readonly ToolStripMenuItem _statusEnabled;
+		readonly ToolStripMenuItem _statusDisabled;
+		readonly ToolStripMenuItem _statusWhenIdle;
+
+		IdleForm? _idleForm;
 		bool _disposed;
 
 		public CustomApplicationContext(EventWaitHandle eventHandle)
 		{
+			_statusEnabled = new ToolStripMenuItem("Enabled");
+			_statusEnabled.Click += (s, e) => SetUserStatus(UserStatus.Enabled);
+
+			_statusDisabled = new ToolStripMenuItem("Disabled");
+			_statusDisabled.Click += (s, e) => SetUserStatus(UserStatus.Disabled);
+
+			_statusWhenIdle = new ToolStripMenuItem("When Idle");
+			_statusWhenIdle.Click += (s, e) => SetUserStatus(UserStatus.WhenIdle);
+
+			ToolStripMenuItem showStatsMenuItem = new ToolStripMenuItem("Stats...");
+			showStatsMenuItem.Click += Status_Stats_OnClick;
+
+			ToolStripMenuItem statusMenuItem = new ToolStripMenuItem("Status");
+			statusMenuItem.DropDownItems.Add(_statusEnabled);
+			statusMenuItem.DropDownItems.Add(_statusDisabled);
+			statusMenuItem.DropDownItems.Add(_statusWhenIdle);
+			statusMenuItem.DropDownItems.Add(new ToolStripSeparator());
+			statusMenuItem.DropDownItems.Add(showStatsMenuItem);
+
 			ToolStripMenuItem logsMenuItem = new ToolStripMenuItem("View logs");
 			logsMenuItem.Click += OnOpenLogs;
 
@@ -98,6 +128,8 @@ namespace Horde.Agent.TrayApp
 			exitMenuItem.Click += OnExit;
 
 			ContextMenuStrip menu = new ContextMenuStrip();
+			menu.Items.Add(statusMenuItem);
+			menu.Items.Add(new ToolStripSeparator());
 			menu.Items.Add(logsMenuItem);
 			menu.Items.Add(new ToolStripSeparator());
 			menu.Items.Add(exitMenuItem);
@@ -111,9 +143,19 @@ namespace Horde.Agent.TrayApp
 				ContextMenuStrip = menu,
 				Visible = true
 			};
+			_trayIcon.Click += TrayIcon_Click;
 
 			_clientTask = BackgroundTask.StartNew(StatusTask);
-			_waitForExitTask = BackgroundTask.StartNew(ctx => WaitForExitTask(eventHandle, ctx));
+			_tickPauseStateTask = BackgroundTask.StartNew(ctx => TickPauseStateAsync(ctx));
+			_waitForExitTask = BackgroundTask.StartNew(ctx => WaitForExitAsync(eventHandle, ctx));
+		}
+
+		private void TrayIcon_Click(object? sender, EventArgs e)
+		{
+			UserStatus status = GetUserStatus();
+			_statusEnabled.Checked = (status == UserStatus.Enabled);
+			_statusDisabled.Checked = (status == UserStatus.Disabled);
+			_statusWhenIdle.Checked = (status == UserStatus.WhenIdle);
 		}
 
 		protected override void Dispose(bool disposing)
@@ -122,6 +164,16 @@ namespace Horde.Agent.TrayApp
 
 			if (disposing)
 			{
+				if (_idleForm != null)
+				{
+					_idleForm.Dispose();
+					_idleForm = null;
+				}
+
+				_statusEnabled.Dispose();
+				_statusDisabled.Dispose();
+				_statusWhenIdle.Dispose();
+
 				_mainThreadInvokeTarget.Dispose();
 				_trayIcon.Dispose();
 				_disposed = true;
@@ -131,6 +183,7 @@ namespace Horde.Agent.TrayApp
 		public async ValueTask DisposeAsync()
 		{
 			await _waitForExitTask.DisposeAsync();
+			await _tickPauseStateTask.DisposeAsync();
 			await _clientTask.DisposeAsync();
 
 			Dispose();
@@ -147,6 +200,41 @@ namespace Horde.Agent.TrayApp
 			}
 		}
 
+		enum UserStatus
+		{
+			Enabled = 0,
+			Disabled = 1,
+			WhenIdle = 2,
+		}
+
+		const string RegistryKey = "HKEY_CURRENT_USER\\Software\\Epic Games\\Horde\\TrayApp";
+		const string RegistryStatusValue = "Status";
+
+		private static UserStatus GetUserStatus()
+		{
+			return (UserStatus)((Registry.GetValue(RegistryKey, RegistryStatusValue, null) as int?) ?? 0);
+		}
+
+		private void SetUserStatus(UserStatus status)
+		{
+			Registry.SetValue(RegistryKey, RegistryStatusValue, (int)status);
+			_statusChangedEvent.Set();
+		}
+
+		private void Status_Stats_OnClick(object? sender, EventArgs e)
+		{
+			if (_idleForm == null)
+			{
+				_idleForm = new IdleForm();
+				_idleForm.FormClosed += (s, e) =>
+				{
+					_idleForm.Dispose(); 
+					_idleForm = null;
+				};
+				_idleForm.Show();
+			}
+		}
+
 		private void OnExit(object? sender, EventArgs e)
 		{
 			ExitThread();
@@ -171,11 +259,16 @@ namespace Horde.Agent.TrayApp
 					_trayIcon.Icon = Resources.StatusBusy;
 					_trayIcon.Text = (status.NumLeases == 1) ? "Currently handling 1 lease" : $"Currently handling {status.NumLeases} leases";
 				}
-				else
+				else if (_enabled)
 				{
 					_trayIcon.Icon = Resources.StatusNormal;
 					_trayIcon.Text = "Agent is operating normally";
 				}
+				else
+				{
+					_trayIcon.Icon = Resources.StatusPaused;
+					_trayIcon.Text = "Agent is paused";
+				}
 			}
 		}
 
@@ -200,12 +293,176 @@ namespace Horde.Agent.TrayApp
 			}
 		}
 
-		async Task WaitForExitTask(EventWaitHandle eventHandle, CancellationToken cancellationToken)
+		async Task WaitForExitAsync(EventWaitHandle eventHandle, CancellationToken cancellationToken)
 		{
 			await eventHandle.WaitOneAsync(cancellationToken);
 			_mainThreadInvokeTarget.BeginInvoke(() => Exit_MainThread());
 		}
 
+		[StructLayout(LayoutKind.Sequential)]
+		struct LASTINPUTINFO
+		{
+			public int cbSize;
+			public uint dwTime;
+		}
+
+		[DllImport("user32.dll")]
+		static extern bool GetLastInputInfo(ref LASTINPUTINFO plii);
+
+		[DllImport("kernel32.dll")]
+		static extern uint GetTickCount();
+
+		[StructLayout(LayoutKind.Sequential)]
+		struct FILETIME
+		{
+			public uint dwLowDateTime;
+			public uint dwHighDateTime;
+
+			public ulong Total => dwLowDateTime | ((ulong)dwHighDateTime << 32);
+		};
+
+		[StructLayout(LayoutKind.Sequential)]
+		struct MEMORYSTATUSEX
+		{
+			public int dwLength;
+			public uint dwMemoryLoad;
+			public ulong ullTotalPhys;
+			public ulong ullAvailPhys;
+			public ulong ullTotalPageFile;
+			public ulong ullAvailPageFile;
+			public ulong ullTotalVirtual;
+			public ulong ullAvailVirtual;
+			public ulong ullAvailExtendedVirtual;
+		}
+
+		[DllImport("kernel32.dll", CharSet = CharSet.Auto, SetLastError = true)]
+		static extern bool GlobalMemoryStatusEx(ref MEMORYSTATUSEX lpBuffer);
+
+		[DllImport("kernel32.dll", SetLastError = true)]
+		static extern bool GetSystemTimes(out FILETIME lpIdleTime, out FILETIME lpKernelTime, out FILETIME lpUserTime);
+
+		bool _enabled;
+		readonly AsyncEvent _statusChangedEvent = new AsyncEvent();
+		readonly AsyncEvent _enabledChangedEvent = new AsyncEvent();
+
+		async Task TickPauseStateAsync(CancellationToken cancellationToken)
+		{
+			await using BackgroundTask cpuStatsTask = BackgroundTask.StartNew(ctx => TickCpuStatsAsync(ctx));
+
+			TimeSpan pollInterval = TimeSpan.FromSeconds(0.25);
+
+			List<IdleStat> idleStats = new List<IdleStat>();
+
+			Stopwatch stateChangeTimer = Stopwatch.StartNew();
+			while (!cancellationToken.IsCancellationRequested)
+			{
+				Task statusChangedTask = _statusChangedEvent.Task;
+
+				UserStatus userStatus = GetUserStatus();
+				if (userStatus == UserStatus.Enabled)
+				{
+					if (!_enabled)
+					{
+						_enabled = true;
+						_enabledChangedEvent.Set();
+					}
+				}
+				else if (userStatus == UserStatus.Disabled)
+				{
+					if (_enabled)
+					{
+						_enabled = false;
+						_enabledChangedEvent.Set();
+					}
+				}
+
+				DateTime utcNow = DateTime.UtcNow;
+				GetIdleStats(idleStats);
+
+				bool idle = idleStats.All(x => x.Value >= x.MinValue);
+				if (idle == _enabled)
+				{
+					stateChangeTimer.Restart();
+				}
+
+				const int WakeTimeSecs = 2;
+				const int IdleTimeSecs = 30;
+				int stateChangeTime = (int)stateChangeTimer.Elapsed.TotalSeconds;
+				int stateChangeMaxTime = _enabled ? WakeTimeSecs : IdleTimeSecs;
+				_idleForm?.TickStats(_enabled, stateChangeTime, stateChangeMaxTime, idleStats);
+
+				if (userStatus == UserStatus.WhenIdle && stateChangeTime >= stateChangeMaxTime)
+				{
+					_enabled ^= true;
+					_enabledChangedEvent.Set();
+					stateChangeTimer.Restart();
+				}
+
+				await Task.WhenAny(statusChangedTask, Task.Delay(pollInterval, cancellationToken));
+			}
+		}
+
+		void GetIdleStats(List<IdleStat> idleStats)
+		{
+			idleStats.Clear();
+
+			// Check there has been no input for a while
+			LASTINPUTINFO lastInputInfo = new LASTINPUTINFO();
+			lastInputInfo.cbSize = Marshal.SizeOf<LASTINPUTINFO>();
+
+			if (GetLastInputInfo(ref lastInputInfo))
+			{
+				const int MinIdleTimeSecs = 2;
+				idleStats.Add(new IdleStat("LastInputTime", (GetTickCount() - lastInputInfo.dwTime) / 1000, MinIdleTimeSecs));
+			}
+
+			// Only look at memory/CPU usage if we're not paused; executing jobs will increase them
+			if (!_enabled)
+			{
+				// Check the CPU usage doesn't exceed the limit
+				int MinIdleCpuPct = 70;
+				idleStats.Add(new IdleStat("IdleCpuPct", _idleCpuPct, MinIdleCpuPct));
+
+				// Check there's enough available virtual memory 
+				MEMORYSTATUSEX memoryStatus = new MEMORYSTATUSEX();
+				memoryStatus.dwLength = Marshal.SizeOf<MEMORYSTATUSEX>();
+
+				if (GlobalMemoryStatusEx(ref memoryStatus))
+				{
+					const long MinFreeVirtualMemMb = 256;
+					idleStats.Add(new IdleStat("VirtualMemMb", (long)(memoryStatus.ullAvailPhys + memoryStatus.ullAvailPageFile) / (1024 * 1024), MinFreeVirtualMemMb));
+				}
+			}
+		}
+
+		int _idleCpuPct = 0;
+
+		async Task TickCpuStatsAsync(CancellationToken cancellationToken)
+		{
+			const int NumSamples = 10;
+			TimeSpan sampleInterval = TimeSpan.FromSeconds(0.2);
+			(ulong IdleTime, ulong TotalTime)[] samples = new (ulong IdleTime, ulong TotalTime)[NumSamples];
+
+			int sampleIdx = 0;
+			for (; ; )
+			{
+				if (GetSystemTimes(out FILETIME idleTime, out FILETIME kernelTime, out FILETIME userTime))
+				{
+					(ulong prevIdleTime, ulong prevTotalTime) = samples[sampleIdx];
+					(ulong nextIdleTime, ulong nextTotalTime) = (idleTime.Total, kernelTime.Total + userTime.Total);
+
+					samples[sampleIdx] = (nextIdleTime, nextTotalTime);
+					sampleIdx = (sampleIdx + 1) % NumSamples;
+
+					if (prevTotalTime > 0 && nextTotalTime > prevTotalTime)
+					{
+						_idleCpuPct = (int)(((nextIdleTime - prevIdleTime) * 100) / (nextTotalTime - prevTotalTime));
+					}
+				}
+				await Task.Delay(sampleInterval, cancellationToken);
+			}
+		}
+
 		async Task PollForStatusUpdatesAsync(CancellationToken cancellationToken)
 		{
 			AgentMessageBuffer message = new AgentMessageBuffer();
@@ -217,13 +474,18 @@ namespace Horde.Agent.TrayApp
 				SetStatus(new AgentStatusMessage(false, 0, "Waiting for status update."));
 				for (; ; )
 				{
+					Task idleChangeTask = _enabledChangedEvent.Task;
+
+					bool enabled = _enabled;
+					message.Set(AgentMessageType.SetEnabledRequest, new AgentEnabledMessage(enabled));
+
 					message.Set(AgentMessageType.GetStatusRequest);
 					await message.SendAsync(pipeClient, cancellationToken);
 
 					if (!await message.TryReadAsync(pipeClient, cancellationToken))
 					{
 						break;
-					}					
+					}
 					
 					switch (message.Type)
 					{
@@ -233,7 +495,7 @@ namespace Horde.Agent.TrayApp
 							break;
 					}
 
-					await Task.Delay(TimeSpan.FromSeconds(5.0), cancellationToken);
+					await Task.WhenAny(idleChangeTask, Task.Delay(TimeSpan.FromSeconds(5.0), cancellationToken));
 				}
 			}
 		}
diff --git a/Engine/Source/Programs/Horde/Horde.Agent/AgentSettings.cs b/Engine/Source/Programs/Horde/Horde.Agent/AgentSettings.cs
--- a/Engine/Source/Programs/Horde/Horde.Agent/AgentSettings.cs
+++ b/Engine/Source/Programs/Horde/Horde.Agent/AgentSettings.cs
@@ -39,7 +39,7 @@ namespace Horde.Agent
 		public string Name { get; set; } = null!;
 
 		/// <summary>
-		/// Name of the environment (currrently just used for tracing)
+		/// Name of the environment (currently just used for tracing)
 		/// </summary>
 		[Required]
 		public string Environment { get; set; } = "prod";
@@ -101,7 +101,7 @@ namespace Horde.Agent
 	}
 
 	/// <summary>
-	/// Setttings for the perforce executor
+	/// Settings for the perforce executor
 	/// </summary>
 	public class PerforceExecutorSettings
 	{
diff --git a/Engine/Source/Programs/Horde/Horde.Agent/IpcMessages.cs b/Engine/Source/Programs/Horde/Horde.Agent/IpcMessages.cs
--- a/Engine/Source/Programs/Horde/Horde.Agent/IpcMessages.cs
+++ b/Engine/Source/Programs/Horde/Horde.Agent/IpcMessages.cs
@@ -41,6 +41,11 @@ namespace Horde.Agent
 		/// Returns the current status, as a Json-encoded <see cref="AgentStatusMessage"/> message
 		/// </summary>
 		GetStatusResponse = 2,
+
+		/// <summary>
+		/// Sets the paused state
+		/// </summary>
+		SetEnabledRequest = 3,
 	}
 
 	/// <summary>
@@ -140,6 +145,12 @@ namespace Horde.Agent
 		}
 	}
 
+	/// <summary>
+	/// Sets the enabled state for the agent
+	/// </summary>
+	/// <param name="IsEnabled">Whether the agent is enabled or not</param>
+	public record class AgentEnabledMessage(bool IsEnabled);
+
 	/// <summary>
 	/// Current status of the agent
 	/// </summary>
diff --git a/Engine/Source/Programs/Horde/Horde.Agent/Leases/LeaseManager.cs b/Engine/Source/Programs/Horde/Horde.Agent/Leases/LeaseManager.cs
--- a/Engine/Source/Programs/Horde/Horde.Agent/Leases/LeaseManager.cs
+++ b/Engine/Source/Programs/Horde/Horde.Agent/Leases/LeaseManager.cs
@@ -233,9 +233,6 @@ namespace Horde.Agent.Leases
 			Stopwatch updateTimer = Stopwatch.StartNew();
 			Queue<TimeSpan> updateTimes = new Queue<TimeSpan>();
 
-			// Whether there is an active user on this machine. For workstations, this determines whether the agent will take on work.
-			bool paused = false;
-
 			// Loop until we're ready to exit
 			Stopwatch updateCapabilitiesTimer = Stopwatch.StartNew();
 			for (; ; )
@@ -275,6 +272,7 @@ namespace Horde.Agent.Leases
 				}
 
 				// Get the new agent status
+				bool paused = !_statusService.IsEnabled;
 				if (stopping)
 				{
 					updateSessionRequest.Status = AgentStatus.Stopping;
diff --git a/Engine/Source/Programs/Horde/Horde.Agent/Services/StatusService.cs b/Engine/Source/Programs/Horde/Horde.Agent/Services/StatusService.cs
--- a/Engine/Source/Programs/Horde/Horde.Agent/Services/StatusService.cs
+++ b/Engine/Source/Programs/Horde/Horde.Agent/Services/StatusService.cs
@@ -2,6 +2,7 @@
 
 using System;
 using System.Collections.Generic;
+using System.Diagnostics;
 using System.IO.Pipes;
 using System.Runtime.InteropServices;
 using System.Runtime.Versioning;
@@ -25,6 +26,7 @@ namespace Horde.Agent.Services
 		private AgentStatusMessage _current;
 
 		readonly BackgroundTask _task;
+		bool? _enabled;
 		readonly ILogger _logger;
 
 		/// <summary>
@@ -32,6 +34,11 @@ namespace Horde.Agent.Services
 		/// </summary>
 		public AgentStatusMessage Current => _current;
 
+		/// <summary>
+		/// Whether the agent is currently enabled.
+		/// </summary>
+		public bool IsEnabled => _enabled ?? true;
+
 		/// <summary>
 		/// Constructor
 		/// </summary>
@@ -81,7 +88,7 @@ namespace Horde.Agent.Services
 		public void SetDescription(string description) => Update(status => new AgentStatusMessage(status.Healthy, status.NumLeases, description));
 
 		/// <summary>
-		/// Updates the status using a custom funciton
+		/// Updates the status using a custom function
 		/// </summary>
 		/// <param name="updateFunc">Function to take the existing status and create an updated version</param>
 		public void Update(Func<AgentStatusMessage, AgentStatusMessage> updateFunc)
@@ -152,14 +159,18 @@ namespace Horde.Agent.Services
 				{
 					switch (request.Type)
 					{
+						case AgentMessageType.SetEnabledRequest:
+							_enabled = request.Parse<AgentEnabledMessage>().IsEnabled;
+							break;
 						case AgentMessageType.GetStatusRequest:
 							response.Set(AgentMessageType.GetStatusResponse, Current);
+							await response.SendAsync(pipeServer, cancellationToken);
 							break;
 						default:
 							response.Set(AgentMessageType.InvalidResponse);
+							await response.SendAsync(pipeServer, cancellationToken);
 							break;
 					}
-					await response.SendAsync(pipeServer, cancellationToken);
 				}
 			}
 		}
diff --git a/Engine/Source/Programs/Shared/EpicGames.Core/AsyncUtils.cs b/Engine/Source/Programs/Shared/EpicGames.Core/AsyncUtils.cs
--- a/Engine/Source/Programs/Shared/EpicGames.Core/AsyncUtils.cs
+++ b/Engine/Source/Programs/Shared/EpicGames.Core/AsyncUtils.cs
@@ -216,14 +216,14 @@ namespace EpicGames.Core
 		}
 
 		/// <summary>
-		/// Waits for a native wait handle to be signalled
+		/// Waits for a native wait handle to be signaled
 		/// </summary>
 		/// <param name="handle">Handle to wait for</param>
 		/// <param name="cancellationToken">Cancellation token for the operation</param>
 		public static Task WaitOneAsync(this WaitHandle handle, CancellationToken cancellationToken = default) => handle.WaitOneAsync(-1, cancellationToken);
 
 		/// <summary>
-		/// Waits for a native wait handle to be signalled
+		/// Waits for a native wait handle to be signaled
 		/// </summary>
 		/// <param name="handle">Handle to wait for</param>
 		/// <param name="timeoutMs">Timeout for the wait</param>
