Horde: Use color preferences for agent view status icons
#rnx
#skipci

diff --git a/Engine/Source/Programs/Horde/HordeDashboard/src/components/AgentView.tsx b/Engine/Source/Programs/Horde/HordeDashboard/src/components/AgentView.tsx
--- a/Engine/Source/Programs/Horde/HordeDashboard/src/components/AgentView.tsx
+++ b/Engine/Source/Programs/Horde/HordeDashboard/src/components/AgentView.tsx
@@ -1,5 +1,5 @@
 // Copyright Epic Games, Inc. All Rights Reserved.
-import { Checkbox, CommandButton, ConstrainMode, ContextualMenu, DefaultButton, DetailsHeader, DetailsList, DetailsListLayoutMode, Dialog, DialogType, DirectionalHint, Dropdown, FontSizes, FontWeights, getTheme, IBasePickerProps, IColumn, Icon, IconButton, IContextualMenuItem, IContextualMenuProps, IDetailsHeaderProps, IDetailsHeaderStyles, IDetailsListProps, ITag, ITagItemStyles, ITooltipHostStyles, Link as ReactLink, mergeStyles, mergeStyleSets, PrimaryButton, ProgressIndicator, ScrollablePane, ScrollbarVisibility, SearchBox, Selection, SelectionMode, Slider, Spinner, SpinnerSize, Stack, Sticky, StickyPositionType, TagItem, TagPicker, Text, TextField } from '@fluentui/react';
+import { Checkbox, CommandButton, ConstrainMode, ContextualMenu, DefaultButton, DetailsHeader, DetailsList, DetailsListLayoutMode, Dialog, DialogType, DirectionalHint, Dropdown, FontSizes, FontWeights, IBasePickerProps, IColumn, IContextualMenuItem, IContextualMenuProps, IDetailsHeaderProps, IDetailsHeaderStyles, IDetailsListProps, ITag, ITagItemStyles, ITooltipHostStyles, Icon, IconButton, PrimaryButton, ProgressIndicator, Link as ReactLink, ScrollablePane, ScrollbarVisibility, SearchBox, Selection, SelectionMode, Slider, Spinner, SpinnerSize, Stack, Sticky, StickyPositionType, TagItem, TagPicker, Text, TextField, getTheme, mergeStyleSets } from '@fluentui/react';
 import { action, makeObservable, observable } from 'mobx';
 import { observer } from 'mobx-react-lite';
 import moment from 'moment-timezone';
@@ -9,6 +9,7 @@ import Marquee from 'react-text-marquee';
 import backend from '../backend';
 import { agentStore } from '../backend/AgentStore';
 import { AgentData, BatchUpdatePoolRequest, GetAgentResponse, LeaseState, PoolData } from '../backend/Api';
+import dashboard, { StatusColor } from '../backend/Dashboard';
 import { copyToClipboard } from '../base/utilities/clipboard';
 import { useWindowSize } from '../base/utilities/hooks';
 import { getShortNiceTime } from '../base/utilities/timeUtils';
@@ -21,23 +22,6 @@ import { TopNav } from './TopNav';
 
 const theme = getTheme();
 
-const iconClass = mergeStyles({
-   fontSize: 16,
-   marginRight: "13px",
-   paddingTop: "2px",
-});
-
-const detailClassNames = mergeStyleSets({
-   success: [{ color: theme.palette.green }, iconClass],
-   warnings: [{ color: theme.palette.yellow }, iconClass],
-   failure: [{ color: theme.palette.red }, iconClass],
-   offline: [{ color: theme.palette.neutralTertiary }, iconClass],
-   waiting: [{ color: theme.palette.neutralLighter }, iconClass],
-   ready: [{ color: theme.palette.neutralLight }, iconClass],
-   skipped: [{ color: theme.palette.neutralTertiary }, iconClass],
-   running: [{ color: theme.palette.blueLight }, iconClass]
-});
-
 const agentStyles = mergeStyleSets({
 
    ticker: {
@@ -907,7 +891,7 @@ class EditPoolsModalState {
             if (pool) {
                if (pool.name !== item.pool.name || pool.properties!["Color"] !== item.pool.properties!["Color"]) {
                   value.modifiedPools.push(item.pool.name);
-               }   
+               }
             }
          }
 
@@ -1437,7 +1421,7 @@ const PoolEditorConfirmation: React.FC = observer(() => {
       if (mods.newPools.length > 1) {
          title += "s";
       }
-   } 
+   }
 
    title += "?";
 
@@ -1834,7 +1818,7 @@ export const SearchUpdate: React.FC = observer(() => {
 export const AgentViewInner: React.FC<{ agentId?: string, poolId?: string, searchParams?: URLSearchParams, agentView?: boolean }> = observer(({ agentId, poolId, searchParams, agentView }) => {
 
    poolId = poolId?.toLowerCase();
-
+   
    const [initAgentUpdater, setInitAgentUpdater] = useState(false);
 
    useEffect(() => {
@@ -2319,18 +2303,17 @@ export const AgentViewInner: React.FC<{ agentId?: string, poolId?: string, searc
    }
 
    function getAgentStatusIcon(agent: AgentData) {
-      let className = detailClassNames.success;
-      //let title = agentReadyStates.enabledReady;
-      // do filtering here
+
+      const statusColors = dashboard.getStatusColors();
+
+      let color = statusColors.get(StatusColor.Success);
       if (agent.enabled === false) {
-         className = detailClassNames.failure;
-         //title = agentReadyStates.disabled;
+         color = statusColors.get(StatusColor.Failure);         
       }
       if (!agent.online) {
-         className = detailClassNames.offline;
-         //title = agentReadyStates.offline;
+         color = statusColors.get(StatusColor.Skipped);         
       }
-      return <Icon iconName="FullCircle" className={className} />;
+      return <Icon iconName="FullCircle" style={{color: color, fontSize:16, marginRight: "13px", paddingTop:"2px"}} />;
    }
 
    function getPropFromDevice(agent: AgentData, propKey: string, propValue: string | null = null) {
