Horde: Missing files for tray app.

diff --git a/Engine/Source/Programs/Horde/Horde.Agent.TrayApp/Forms/IdleForm.cs b/Engine/Source/Programs/Horde/Horde.Agent.TrayApp/Forms/IdleForm.cs
new file mode 100644
--- /dev/null
+++ b/Engine/Source/Programs/Horde/Horde.Agent.TrayApp/Forms/IdleForm.cs
@@ -0,0 +1,83 @@
+// Copyright Epic Games, Inc. All Rights Reserved.
+
+using System;
+using System.Collections.Generic;
+using System.ComponentModel;
+using System.Data;
+using System.Drawing;
+using System.Linq;
+using System.Text;
+using System.Threading.Tasks;
+using System.Windows.Forms;
+using EpicGames.Core;
+
+namespace Horde.Agent.TrayApp.Forms
+{
+	partial class IdleForm : Form
+	{
+		readonly Dictionary<string, ListViewItem> _items = new Dictionary<string, ListViewItem>(StringComparer.Ordinal);
+
+		public IdleForm()
+		{
+			InitializeComponent();
+
+			Rectangle rect = Screen.GetWorkingArea(this);
+			Location = new Point(rect.X + rect.Width - Size.Width - 20, rect.Y + rect.Height - Size.Height - 20);
+
+			float scale = DeviceDpi / 96.0f;
+			nameColumnHeader.Width = (int)(nameColumnHeader.Width * scale);
+			valueColumnHeader.Width = (int)(valueColumnHeader.Width * scale);
+			minValueColumnHeader.Width = (int)(minValueColumnHeader.Width * scale);
+			resultColumnHeader.Width = (int)(resultColumnHeader.Width * scale);
+		}
+
+		public void TickStats(bool enabled, int stateChangeTime, int stateChangeMaxTime, List<IdleStat> stats)
+		{
+			BeginInvoke(() => TickStatsMainThread(enabled, stateChangeTime, stateChangeMaxTime, stats));
+		}
+
+		void TickStatsMainThread(bool enabled, int stateChangeTime, int stateChangeMaxTime, List<IdleStat> stats)
+		{
+			if (IsDisposed)
+			{
+				return;
+			}
+
+			string activeText = enabled ? "Enabled" : "Paused";
+			if (stateChangeTime == 0)
+			{
+				stateTextBox.Text = $"{activeText}";
+			}
+			else
+			{
+				stateTextBox.Text = $"{activeText} (state change in {Math.Max(stateChangeMaxTime - stateChangeTime, 0)}s)";
+			}
+
+			HashSet<string> removeItems = new HashSet<string>(_items.Keys, StringComparer.Ordinal);
+			foreach (IdleStat stat in stats)
+			{
+				ListViewItem? item;
+				if (!_items.TryGetValue(stat.Name, out item))
+				{
+					item = statsListView.Items.Add(stat.Name);
+					item.SubItems.Add("");
+					item.SubItems.Add("");
+					item.SubItems.Add("");
+					_items.Add(stat.Name, item);
+				}
+
+				item.SubItems[1].Text = stat.Value.ToString();
+				item.SubItems[2].Text = stat.MinValue.ToString();
+				item.SubItems[3].Text = (stat.Value < stat.MinValue) ? "" : "Idle";
+
+				removeItems.Remove(stat.Name);
+			}
+
+			foreach (string removeItem in removeItems)
+			{
+				statsListView.Items.Remove(_items[removeItem]);
+				_items.Remove(removeItem);
+			}
+		}
+	}
+}
