Horde: Refine naming on Horde artifacts v2 modal artifact types, add a notice to artifacts panel to help steer users to new location
#skipci
#rnx

diff --git a/Engine/Source/Programs/Horde/HordeDashboard/src/components/LogView.tsx b/Engine/Source/Programs/Horde/HordeDashboard/src/components/LogView.tsx
--- a/Engine/Source/Programs/Horde/HordeDashboard/src/components/LogView.tsx
+++ b/Engine/Source/Programs/Horde/HordeDashboard/src/components/LogView.tsx
@@ -879,7 +879,7 @@ export const LogList: React.FC<{ logId: string }> = observer(({ logId }) => {
 
          opsList.push({
             key: 'stepops_artifacts_step',
-            text: "Step Artifacts",
+            text: "Logs",
             iconProps: { iconName: "Folder" },
             disabled: !atypes.get("step-saved"),
             onClick: () => { setLogArtifacts("step-saved") }
@@ -887,7 +887,7 @@ export const LogList: React.FC<{ logId: string }> = observer(({ logId }) => {
 
          opsList.push({
             key: 'stepops_artifacts_output',
-            text: "Output Artifacts",
+            text: "Temp Storage",
             iconProps: { iconName: "MenuOpen" },
             disabled: !atypes.get("step-output"),
             onClick: () => { setLogArtifacts("step-output") }
@@ -895,7 +895,7 @@ export const LogList: React.FC<{ logId: string }> = observer(({ logId }) => {
 
          opsList.push({
             key: 'stepops_artifacts_trace',
-            text: "Trace Artifacts",
+            text: "Traces",
             iconProps: { iconName: "SearchTemplate" },
             disabled: !atypes.get("step-trace"),
             onClick: () => { setLogArtifacts("step-trace") }
diff --git a/Engine/Source/Programs/Horde/HordeDashboard/src/components/jobDetailsV2/JobDetailArtifactsV2.tsx b/Engine/Source/Programs/Horde/HordeDashboard/src/components/jobDetailsV2/JobDetailArtifactsV2.tsx
--- a/Engine/Source/Programs/Horde/HordeDashboard/src/components/jobDetailsV2/JobDetailArtifactsV2.tsx
+++ b/Engine/Source/Programs/Horde/HordeDashboard/src/components/jobDetailsV2/JobDetailArtifactsV2.tsx
@@ -1,6 +1,6 @@
 // Copyright Epic Games, Inc. All Rights Reserved.
 
-import { Stack, Text, IColumn, mergeStyleSets, Icon, DetailsList, Selection, SelectionMode, DetailsListLayoutMode, ScrollablePane, ScrollbarVisibility, StickyPositionType, IDetailsListProps, IDetailsHeaderStyles, Sticky, DetailsHeader, PrimaryButton, SpinnerSize, Spinner, Link, TextField } from '@fluentui/react';
+import { Stack, Text, IColumn, mergeStyleSets, Icon, DetailsList, Selection, SelectionMode, DetailsListLayoutMode, ScrollablePane, ScrollbarVisibility, StickyPositionType, IDetailsListProps, IDetailsHeaderStyles, Sticky, DetailsHeader, PrimaryButton, SpinnerSize, Spinner, Link, TextField, FontIcon } from '@fluentui/react';
 import React, { useEffect, useState } from 'react';
 import { ArtifactData, GetArtifactResponse, GetArtifactZipRequest } from '../../backend/Api';
 import { hordeClasses } from '../../styles/Styles';
@@ -181,7 +181,7 @@ class ArtifactsDataView extends JobDataView {
       }).catch((reason) => {
          console.log(`Error getting artifacts jobId: ${details.jobId} stepId: ${this.stepId} error: ${reason}`);
       }).finally(() => {
-         this.initialize(this.artifacts?.length ? [sideRail] : undefined)
+         this.initialize((this.artifacts?.length || details.jobData?.useArtifactsV2) ? [sideRail] : undefined)
       });
 
    }
@@ -223,9 +223,28 @@ export const JobDetailArtifactsV2: React.FC<{ jobDetails: JobDetailsV2; stepId:
    }, [dataView]);
 
    dataView.subscribe();
-
    dataView.set(stepId);
 
+   if (jobDetails.jobData?.useArtifactsV2) {
+      return <Stack id={sideRail.url} styles={{ root: { paddingTop: 18, paddingRight: 12 } }}>
+         <Stack className={hordeClasses.raised}>
+            <Stack tokens={{ childrenGap: 12 }}>
+               <Stack horizontalAlign="space-between" styles={{ root: { minHeight: 32 } }}>
+                  <Text variant="mediumPlus" styles={{ root: { fontFamily: "Horde Open Sans SemiBold" } }}>Artifacts</Text>
+                  <Stack tokens={{ childrenGap: 24 }}>
+                     <Stack horizontal verticalAlign='center' style={{ paddingTop: 24, paddingLeft: 12, paddingBottom: 52 }}>
+                        <FontIcon style={{ paddingRight: 12 }} iconName="Error" />
+                        <Stack>
+                           <Text variant='medium'>This job uses the new Horde artifact storage backend.  These artifacts are now accessible in the upper right, underneath the breadcrumb area.</Text>
+                        </Stack>
+                     </Stack>
+                  </Stack>
+               </Stack>
+            </Stack>
+         </Stack>
+      </Stack>
+   }
+
    if (!dataView.artifacts?.length || !jobDetails?.viewsReady) {
       return null;
    }
@@ -295,7 +314,7 @@ export const JobDetailArtifactsV2: React.FC<{ jobDetails: JobDetailsV2; stepId:
 
                            const filter = newValue ? newValue : undefined;
                            if (filter !== state.filter) {
-                              setState({...state, filter: filter});
+                              setState({ ...state, filter: filter });
                            }
 
                            return undefined;
diff --git a/Engine/Source/Programs/Horde/HordeDashboard/src/components/jobDetailsV2/JobOperationsBar.tsx b/Engine/Source/Programs/Horde/HordeDashboard/src/components/jobDetailsV2/JobOperationsBar.tsx
--- a/Engine/Source/Programs/Horde/HordeDashboard/src/components/jobDetailsV2/JobOperationsBar.tsx
+++ b/Engine/Source/Programs/Horde/HordeDashboard/src/components/jobDetailsV2/JobOperationsBar.tsx
@@ -239,7 +239,7 @@ const StepArtifactsOperations: React.FC<{ jobDetails: JobDetailsV2, stepId: stri
 
    opsList.push({
       key: 'stepops_artifacts_step',
-      text: "Step Artifacts",
+      text: "Logs",
       iconProps: { iconName: "Folder" },
       disabled: !atypes.get("step-saved"),
       onClick: () => { setArtifactsShown("step-saved") }
@@ -247,7 +247,7 @@ const StepArtifactsOperations: React.FC<{ jobDetails: JobDetailsV2, stepId: stri
 
    opsList.push({
       key: 'stepops_artifacts_output',
-      text: "Output Artifacts",
+      text: "Temp Storage",
       iconProps: { iconName: "MenuOpen" },
       disabled: !atypes.get("step-output"),
       onClick: () => { setArtifactsShown("step-output") }
@@ -255,7 +255,7 @@ const StepArtifactsOperations: React.FC<{ jobDetails: JobDetailsV2, stepId: stri
 
    opsList.push({
       key: 'stepops_artifacts_trace',
-      text: "Trace Artifacts",
+      text: "Traces",
       iconProps: { iconName: "SearchTemplate" },
       disabled: !atypes.get("step-trace"),
       onClick: () => { setArtifactsShown("step-trace") }
diff --git a/Engine/Source/Programs/Horde/HordeDashboard/src/components/jobDetailsV2/StepDetailView.tsx b/Engine/Source/Programs/Horde/HordeDashboard/src/components/jobDetailsV2/StepDetailView.tsx
--- a/Engine/Source/Programs/Horde/HordeDashboard/src/components/jobDetailsV2/StepDetailView.tsx
+++ b/Engine/Source/Programs/Horde/HordeDashboard/src/components/jobDetailsV2/StepDetailView.tsx
@@ -163,7 +163,7 @@ const StepDetailViewInner: React.FC<{ jobDetails: JobDetailsV2, stepId: string }
       {<Stack>
          <StepHistoryPanel jobDetails={jobDetails} stepId={stepId} />
       </Stack>}
-      {!jobData.useArtifactsV2 && <Stack>
+      {<Stack>
          <JobDetailArtifactsV2 jobDetails={jobDetails} stepId={stepId} />
       </Stack>}
       {!!step && <Stack>
