Horde: Reformat DDC source files with tabs instead of spaces.

diff --git a/Engine/Source/Programs/Horde/Horde.Server/Ddc/BlobId.cs b/Engine/Source/Programs/Horde/Horde.Server/Ddc/BlobId.cs
--- a/Engine/Source/Programs/Horde/Horde.Server/Ddc/BlobId.cs
+++ b/Engine/Source/Programs/Horde/Horde.Server/Ddc/BlobId.cs
@@ -34,6 +34,11 @@ namespace Horde.Server.Ddc
 		/// </summary>
 		public static BlobId Parse(string text) => new BlobId(IoHash.Parse(text));
 
+		/// <summary>
+		/// Creates a BlobId from a block of data
+		/// </summary>
+		public static BlobId FromBlob(ReadOnlySpan<byte> span) => new BlobId(IoHash.Compute(span));
+
 		/// <inheritdoc/>
 		public override readonly int GetHashCode() => Hash.GetHashCode();
 
diff --git a/Engine/Source/Programs/Horde/Horde.Server/Ddc/RefsController.cs b/Engine/Source/Programs/Horde/Horde.Server/Ddc/RefsController.cs
--- a/Engine/Source/Programs/Horde/Horde.Server/Ddc/RefsController.cs
+++ b/Engine/Source/Programs/Horde/Horde.Server/Ddc/RefsController.cs
@@ -34,712 +34,712 @@ namespace Horde.Server.Ddc
 {
 	using IDiagnosticContext = Serilog.IDiagnosticContext;
 
-    [ApiController]
-    [FormatFilter]
-    [Produces(MediaTypeNames.Application.Json, MediaTypeNames.Application.Octet, CustomMediaTypeNames.UnrealCompactBinary)]
-    [Route("api/v1/refs")]
-    [Authorize]
-    public class RefsController : ControllerBase
-    {
-        private readonly IDiagnosticContext _diagnosticContext;
-        private readonly FormatResolver _formatResolver;
-        private readonly BufferedPayloadFactory _bufferedPayloadFactory;
-        private readonly IReferenceResolver _referenceResolver;
-        private readonly NginxRedirectHelper _nginxRedirectHelper;
-        private readonly IRequestHelper _requestHelper;
-        private readonly Tracer _tracer;
-
-        private readonly ILogger _logger;
-        private readonly IRefService _refService;
-        private readonly IBlobService _blobStore;
-
-        public RefsController(IRefService objectService, IBlobService blobStore, IDiagnosticContext diagnosticContext, FormatResolver formatResolver, BufferedPayloadFactory bufferedPayloadFactory, IReferenceResolver referenceResolver, NginxRedirectHelper nginxRedirectHelper, IRequestHelper requestHelper, Tracer tracer, ILogger<RefsController> logger)
-        {
-            _refService = objectService;
-            _blobStore = blobStore;
-            _diagnosticContext = diagnosticContext;
-            _formatResolver = formatResolver;
-            _bufferedPayloadFactory = bufferedPayloadFactory;
-            _referenceResolver = referenceResolver;
-            _nginxRedirectHelper = nginxRedirectHelper;
-            _requestHelper = requestHelper;
-            _tracer = tracer;
-            _logger = logger;
-        }
-
-        /// <summary>
-        /// Returns a refs key
-        /// </summary>
-        /// <param name="ns">Namespace. Each namespace is completely separated from each other. Use for different types of data that is never expected to be similar (between two different games for instance). Example: `uc4.ddc`</param>
-        /// <param name="bucket">The category/type of record you are caching. Is a clustered key together with the actual key, but all records in the same bucket can be dropped easily. Example: `terrainTexture` </param>
-        /// <param name="key">The unique name of this particular key. `iAmAVeryValidKey`</param>
-        /// <param name="format">Optional specifier to set which output format is used json/raw/cb</param>
-        [HttpGet("{ns}/{bucket}/{key}.{format?}", Order = 500)]
-        [Produces(MediaTypeNames.Application.Json, MediaTypeNames.Application.Octet, CustomMediaTypeNames.UnrealCompactBinary, CustomMediaTypeNames.JupiterInlinedPayload, CustomMediaTypeNames.UnrealCompactBinaryPackage)]
-        public async Task<IActionResult> Get(
-            [FromRoute] [Required] NamespaceId ns,
-            [FromRoute] [Required] BucketId bucket,
-            [FromRoute] [Required] RefId key,
-            [FromRoute] string? format = null)
-        {
-            ActionResult? accessResult = await _requestHelper.HasAccessToNamespaceAsync(User, Request, ns, new [] { StorageAclAction.ReadRefs });
-            if (accessResult != null)
-            {
-                return accessResult;
-            }
-
-            try
-            {
-                (RefRecord objectRecord, BlobContents? blob) = await _refService.GetAsync(ns, bucket, key, Array.Empty<string>());
-
-                if (blob == null)
-                {
-                    throw new InvalidOperationException($"Blob was null when attempting to fetch {ns} {bucket} {key}");
-                }
-
-                if (!objectRecord.IsFinalized)
-                {
-                    // we do not consider un-finalized objects as valid
-                    return BadRequest(new ProblemDetails { Title = $"Object {objectRecord.Bucket} {objectRecord.Name} is not finalized." });
-                }
-
-                Response.Headers[CommonHeaders.HashHeaderName] = objectRecord.BlobIdentifier.ToString();
-                Response.Headers[CommonHeaders.LastAccessHeaderName] = objectRecord.LastAccess.ToString(CultureInfo.InvariantCulture);
-
-                async Task WriteBody(BlobContents blobContents, string contentType)
-                {
-                    IServerTiming? serverTiming = Request.HttpContext.RequestServices.GetService<IServerTiming>();
-                    using ServerTimingMetricScoped? serverTimingScope = serverTiming?.CreateServerTimingMetricScope("body.write", "Time spent writing body");
-
-                    long contentLength = blobContents.Length;
-                    using TelemetrySpan scope = _tracer.StartActiveSpan("body.write").SetAttribute("operation.name", "body.write");
-                    scope.SetAttribute("content-length", contentLength);
-                    const int BufferSize = 64 * 1024;
-                    Stream outputStream = Response.Body;
-                    Response.ContentLength = contentLength;
-                    Response.ContentType = contentType;
-                    Response.StatusCode = StatusCodes.Status200OK;
-                    try
-                    {
-                        await StreamCopyOperation.CopyToAsync(blobContents.Stream, outputStream, count: null, bufferSize: BufferSize, cancel: Response.HttpContext.RequestAborted);
-                    }
-                    catch (OperationCanceledException)
-                    {
-                        // do not raise exceptions for cancelled writes
-                        // as we have already started writing a response we can not change the status code
-                        // so we just drop a warning and proceed
-                        _logger.LogWarning("The operation was canceled while writing the body");
-                    }
-                }
-
-                string responseType = _formatResolver.GetResponseType(Request, format, CustomMediaTypeNames.UnrealCompactBinary);
-                Tracer.CurrentSpan.SetAttribute("response-type", responseType);
-                switch (responseType)
-                {
-                    case CustomMediaTypeNames.UnrealCompactBinary:
-                    {
-                        // for compact binary we can just serialize our internal object
-                        await WriteBody(blob, CustomMediaTypeNames.UnrealCompactBinary);
-
-                        break;
-                    }
-                    case MediaTypeNames.Application.Octet:
-                    {
-                        byte[] blobMemory = await blob.Stream.ToByteArray();
-                        CbObject cb = new CbObject(blobMemory);
-
-                        (int,CbField?) CountFields(CbObject o)
-                        {
-                            int count = 0;
-                            CbField? foundField = null;
-                            cb.IterateAttachments(field =>
-                            {
-                                ++count;
-                                if (field.IsBinaryAttachment())
-                                {
-                                    foundField = field;
-                                }
-                            });
-
-                            return (count, foundField);
-                        }
-
-                        // breaking lambda call into private method to workaround incorrect triggering of CA1508 - [REDACTED_PROFILE_URL]/roslyn-analyzers/issues/5254
-                        (int countOfAttachmentFields,CbField? binaryAttachmentField) = CountFields(cb);
-
-                        if (countOfAttachmentFields == 1 && binaryAttachmentField != null)
-                        {
-                            // there is a single attachment field and that is of the binary attachment type, fetch that attachment and return it instead of the compact binary
-                            // this is so that we match the uploaded that happened as a octet-stream which generates a small cb object with a single attachment
-
-                            IoHash hash = binaryAttachmentField.AsBinaryAttachment();
-
-                            BlobContents referencedBlobContents = await _blobStore.GetObjectAsync(ns, new BlobId(hash));
-
-                            if (_nginxRedirectHelper.CanRedirect(Request, referencedBlobContents))
-                            {
-                                return _nginxRedirectHelper.CreateActionResult(referencedBlobContents, MediaTypeNames.Application.Octet);
-                            }
-                            await WriteBody(referencedBlobContents, MediaTypeNames.Application.Octet);
-                            break;
-                        }
-
-                        // this doesn't look like the generated compact binary so we just return the payload
-                        await using BlobContents contents = new(blobMemory);
-                        await WriteBody(contents, MediaTypeNames.Application.Octet);
-                        break;
-                    }
-                    case MediaTypeNames.Application.Json:
-                    {
-                        byte[] blobMemory;
-                        {
-                            using TelemetrySpan scope = _tracer.StartActiveSpan("json.readblob").SetAttribute("operation.name", "json.readblob");
-                            blobMemory = await blob.Stream.ToByteArray();
-                        }
-                        CbObject cb = new CbObject(blobMemory);
-                        string s = cb.ToJson();
-                        await using BlobContents contents = new BlobContents(Encoding.UTF8.GetBytes(s));
-                        await WriteBody(contents, MediaTypeNames.Application.Json);
-                        break;
-
-                    }
-                    case CustomMediaTypeNames.UnrealCompactBinaryPackage:
-                    {
-                        using TelemetrySpan packageScope = _tracer.StartActiveSpan("cbpackage.fetch").SetAttribute("operation.name", "cbpackage.fetch");
-                        byte[] blobMemory = await blob.Stream.ToByteArray();
-                        CbObject cb = new CbObject(blobMemory);
-
-                        IAsyncEnumerable<Attachment> attachments = _referenceResolver.GetAttachmentsAsync(ns, cb);
-
-                        using CbPackageBuilder writer = new CbPackageBuilder();
-                        writer.AddAttachment(objectRecord.BlobIdentifier.Hash, CbPackageAttachmentFlags.IsObject, blobMemory);
-
-                        await Parallel.ForEachAsync(attachments, async (attachment, token) =>
-                        {
-                            IoHash attachmentHash = attachment.AsIoHash();
-                            CbPackageAttachmentFlags flags = 0;
-
-                            try
-                            {
-                                BlobContents attachmentContents;
-                                if (attachment is BlobAttachment blobAttachment)
-                                {
-                                    BlobId referencedBlob = blobAttachment.Identifier;
-                                    attachmentContents = await _blobStore.GetObjectAsync(ns, referencedBlob);
-                                }
-                                else if (attachment is ObjectAttachment objectAttachment)
-                                {
-                                    flags |= CbPackageAttachmentFlags.IsObject;
-                                    BlobId referencedBlob = objectAttachment.Identifier;
-                                    attachmentContents = await _blobStore.GetObjectAsync(ns, referencedBlob);
-                                }
-                                else if (attachment is ContentIdAttachment contentIdAttachment)
-                                {
-
-                                    ContentId contentId = contentIdAttachment.Identifier;
-                                    (attachmentContents, string mime) = await _blobStore.GetCompressedObject(ns, contentId, HttpContext.RequestServices);
-                                    if (mime == CustomMediaTypeNames.UnrealCompressedBuffer)
-                                    {
-                                        flags |= CbPackageAttachmentFlags.IsCompressed;
-                                    }
-                                    else
-                                    {
-                                        // this resolved to a uncompressed blob, the content id existed the the compressed blob didn't
-                                        // so resetting flags to indicate this.
-                                        flags = 0;
-                                    }
-                                }
-                                else
-                                {
-                                    throw new NotSupportedException($"Unknown attachment type {attachment.GetType()}");
-                                }
-
-                                writer.AddAttachment(attachmentHash, flags, attachmentContents.Stream, (ulong)attachmentContents.Length);
-                            }
-                            catch (Exception e)
-                            {
-                                (CbObject errorObject, HttpStatusCode _) = ToErrorResult(e);
-                                writer.AddAttachment(attachmentHash,
-                                    CbPackageAttachmentFlags.IsError | CbPackageAttachmentFlags.IsObject,
-                                    errorObject.GetView().ToArray());
-                            }
-                        });
-
-                        byte[] packageBytes;
-                        {
-                            using TelemetrySpan _ = _tracer.StartActiveSpan("cbpackage.buffer").SetAttribute("operation.name", "cbpackage.buffer");
-                            packageBytes = await writer.ToByteArray();
-                        }
-                        await using BlobContents contents = new BlobContents(packageBytes);
-                        await WriteBody(contents, CustomMediaTypeNames.UnrealCompactBinaryPackage);
-                        break;
-                    }
-                    case CustomMediaTypeNames.JupiterInlinedPayload:
-                    {
-                        byte[] blobMemory = await blob.Stream.ToByteArray();
-                        CbObject cb = new CbObject(blobMemory);
-
-                        static (int, int) CountFields(CbObject o)
-                        {
-                            int countOfBinaryAttachmentFields = 0;
-                            int countOfAttachmentFields = 0;
-
-                            o.IterateAttachments(field =>
-                            {
-                                if (field.IsBinaryAttachment())
-                                {
-                                    ++countOfBinaryAttachmentFields;
-                                }
-
-                                if (field.IsAttachment())
-                                {
-                                    ++countOfAttachmentFields;
-                                }
-                            });
-
-                            return (countOfBinaryAttachmentFields, countOfAttachmentFields);
-                        }
-                        // breaking lambda call into private method to workaround incorrect triggering of CA1508 - [REDACTED_PROFILE_URL]/roslyn-analyzers/issues/5254
-                        (int countOfAttachmentFields, int countOfBinaryAttachmentFields) = CountFields(cb);
-
-                        // if the object consists of a single attachment field we return this attachment field instead
-                        if (countOfBinaryAttachmentFields == 1 && countOfAttachmentFields == 1)
-                        {
-                            // fetch the blob so we can resolve any content ids in it
-                            List<BlobId> referencedBlobs;
-                            try
-                            {
-                                IAsyncEnumerable<BlobId> referencedBlobsEnumerable = _referenceResolver.GetReferencedBlobsAsync(ns, cb);
-                                referencedBlobs = await referencedBlobsEnumerable.ToListAsync();
-                            }
-                            catch (PartialReferenceResolveException)
-                            {
-                                return NotFound(new ProblemDetails {Title = $"Object {bucket} {key} in namespace {ns} was missing some content ids"});
-                            }
-                            catch (ReferenceIsMissingBlobsException)
-                            {
-                                return NotFound(new ProblemDetails {Title = $"Object {bucket} {key} in namespace {ns} was missing some blobs"});
-                            }
-
-                            if (referencedBlobs.Count == 1)
-                            {
-                                BlobId attachmentToSend = referencedBlobs.First();
-                                try
-                                {
-                                    BlobContents referencedBlobContents = await _blobStore.GetObjectAsync(ns, attachmentToSend);
-                                    Response.Headers[CommonHeaders.InlinePayloadHash] = attachmentToSend.ToString();
-
-                                    if (_nginxRedirectHelper.CanRedirect(Request, referencedBlobContents))
-                                    {
-                                        return _nginxRedirectHelper.CreateActionResult(referencedBlobContents, CustomMediaTypeNames.JupiterInlinedPayload);
-                                    }
-                                    await WriteBody(referencedBlobContents, CustomMediaTypeNames.JupiterInlinedPayload);
-                                }
-                                catch (BlobNotFoundException)
-                                {
-                                    return NotFound(new ProblemDetails {Title = $"Object {bucket} {key} in namespace {ns} was missing blob {attachmentToSend}"});
-                                }
-                                catch (Exception ex)
-                                {
-                                    _logger.LogError(ex, "Unknown exception encountered while writing body for jupiter inlined payload.");
-                                    throw;
-                                }
-                                return new EmptyResult();
-                            }
-                            else if (referencedBlobs.Count == 0)
-                            {
-                                return NotFound(new ProblemDetails
-                                {
-                                    Title =
-                                        $"Object {objectRecord.Bucket} {objectRecord.Name} did not resolve into any objects that we could find."
-                                });
-                            }
-
-                            return BadRequest(new ProblemDetails
-                            {
-                                Title =
-                                    $"Object {objectRecord.Bucket} {objectRecord.Name} contained a content id which resolved to more then 1 blob, unable to inline this object. Use compact object response instead."
-                            });
-                        }
-                        else if (countOfBinaryAttachmentFields == 0 && countOfAttachmentFields == 0)
-                        {
-                            // no attachments so we just return the compact object instead
-                            await using BlobContents contents = new BlobContents(blobMemory);
-                            await WriteBody(contents, CustomMediaTypeNames.JupiterInlinedPayload);
-                            return new EmptyResult();
-                        }
-
-                        return BadRequest(new ProblemDetails
-                        {
-                            Title =
-                                $"Object {objectRecord.Bucket} {objectRecord.Name} had more then 1 binary attachment field, unable to inline this object. Use compact object response instead."
-                        });
-                    }
-                    default:
-                        throw new NotImplementedException($"Unknown expected response type {responseType}");
-                }
-                
-                // this result is ignored as we write to the body explicitly
-                return new EmptyResult();
-
-            }
-            catch (NamespaceNotFoundException e)
-            {
-                return NotFound(new ProblemDetails {Title = $"Namespace {e.Namespace} did not exist"});
-            }
-            catch (RefNotFoundException e)
-            {
-                return NotFound(new ProblemDetails { Title = $"Object {e.Bucket} {e.Key} did not exist" });
-            }
-            catch (BlobNotFoundException e)
-            {
-                return NotFound(new ProblemDetails { Title = $"Object {e.Blob} in {e.Ns} not found" });
-            }
-        }
-
-       
-
-    /// <summary>
-    /// Returns the metadata about a ref key
-    /// </summary>
-    /// <param name="ns">Namespace. Each namespace is completely separated from each other. Use for different types of data that is never expected to be similar (between two different games for instance). Example: `uc4.ddc`</param>
-    /// <param name="bucket">The category/type of record you are caching. Is a clustered key together with the actual key, but all records in the same bucket can be dropped easily. Example: `terrainTexture` </param>
-    /// <param name="key">The unique name of this particular key. `iAmAVeryValidKey`</param>
-    /// <param name="fields">The fields to include in the response, omit this to include everything.</param>
-    [HttpGet("{ns}/{bucket}/{key}/metadata", Order = 500)]
-    public async Task<IActionResult> GetMetadata(
-        [FromRoute] [Required] NamespaceId ns,
-        [FromRoute] [Required] BucketId bucket,
-        [FromRoute] [Required] RefId key,
-        [FromQuery] string[] fields)
-    {
-        ActionResult? accessResult = await _requestHelper.HasAccessToNamespaceAsync(User, Request, ns, new [] { StorageAclAction.ReadRefs });
-        if (accessResult != null)
-        {
-            return accessResult;
-        }
-
-        try
-        {
-            (RefRecord objectRecord, BlobContents? _) = await _refService.GetAsync(ns, bucket, key, fields);
-
-            return Ok(new RefMetadataResponse(objectRecord));
-        }
-        catch (NamespaceNotFoundException e)
-        {
-            return NotFound(new ProblemDetails {Title = $"Namespace {e.Namespace} did not exist"});
-        }
-        catch (RefNotFoundException e)
-        {
-            return NotFound(new ProblemDetails { Title = $"Object {e.Bucket} {e.Key} did not exist" });
-        }
-        catch (BlobNotFoundException e)
-        {
-            return NotFound(new ProblemDetails { Title = $"Object {e.Blob} in {e.Ns} not found" });
-        }
-    }
-
-        /// <summary>
-        /// Checks if a object exists
-        /// </summary>
-        /// <param name="ns">Namespace. Each namespace is completely separated from each other. Use for different types of data that is never expected to be similar (between two different games for instance). Example: `uc4.ddc`</param>
-        /// <param name="bucket">The category/type of record you are caching. Is a clustered key together with the actual key, but all records in the same bucket can be dropped easily. Example: `terrainTexture` </param>
-        /// <param name="key">The unique name of this particular key. `iAmAVeryValidKey`</param>
-        /// <returns>200 if it existed, 400 otherwise</returns>
-        [HttpHead("{ns}/{bucket}/{key}", Order = 500)]
-        [ProducesResponseType(type: typeof(OkResult), 200)]
-        [ProducesResponseType(type: typeof(ValidationProblemDetails), 400)]
-        public async Task<IActionResult> Head(
-            [FromRoute] [Required] NamespaceId ns,
-            [FromRoute] [Required] BucketId bucket,
-            [FromRoute] [Required] RefId key)
-        {
-            ActionResult? accessResult = await _requestHelper.HasAccessToNamespaceAsync(User, Request, ns, new [] { StorageAclAction.ReadRefs });
-            if (accessResult != null)
-            {
-                return accessResult;
-            }
-
-            try
-            {
-                (RefRecord record, BlobContents? blob) = await _refService.GetAsync(ns, bucket, key, new string[] {"blobIdentifier", "IsFinalized"});
-                Response.Headers[CommonHeaders.HashHeaderName] = record.BlobIdentifier.ToString();
-
-                if (!record.IsFinalized)
-                {
-                    return NotFound(new ProblemDetails {Title = $"Object {bucket} {key} in namespace {ns} is not finalized."});
-                }
-
-                blob ??= await _blobStore.GetObjectAsync(ns, record.BlobIdentifier);
-
-                // we have to verify the blobs are available locally, as the record of the key is replicated a head of the content
-                // TODO: Once we support inline replication this step is not needed as at least one region as this blob, just maybe not this current one
-                byte[] blobContents = await blob.Stream.ToByteArray();
-                CbObject compactBinaryObject = new CbObject(blobContents);
-                // the reference resolver will throw if any blob is missing, so no need to do anything other then process each reference
-                IAsyncEnumerable<BlobId> references = _referenceResolver.GetReferencedBlobsAsync(ns, compactBinaryObject);
-                List<BlobId>? _ = await references.ToListAsync();
-
-                // we have to verify the blobs are available locally, as the record of the key is replicated a head of the content
-                // TODO: Once we support inline replication this step is not needed as at least one region as this blob, just maybe not this current one
-                BlobId[] unknownBlobs = await _blobStore.FilterOutKnownBlobsAsync(ns, new BlobId[] { record.BlobIdentifier });
-                if (unknownBlobs.Length != 0)
-                {
-                    return NotFound(new ProblemDetails {Title = $"Object {bucket} {key} in namespace {ns} had at least one missing blob."});
-                }
-            }
-            catch (NamespaceNotFoundException e)
-            {
-                return NotFound(new ProblemDetails {Title = $"Namespace {e.Namespace} did not exist"});
-            }
-            catch (BlobNotFoundException e)
-            {
-                return NotFound(new ProblemDetails { Title = $"Blob {e.Blob} in namespace {ns} did not exist" });
-            }
-            catch (RefNotFoundException e)
-            {
-                return NotFound(new ProblemDetails {Title = $"Object {e.Bucket} {e.Key} in namespace {e.Namespace} did not exist"});
-            }
-            catch (PartialReferenceResolveException)
-            {
-                return NotFound(new ProblemDetails {Title = $"Object {bucket} {key} in namespace {ns} was missing some content ids"});
-            }
-            catch (ReferenceIsMissingBlobsException)
-            {
-                return NotFound(new ProblemDetails {Title = $"Object {bucket} {key} in namespace {ns} was missing some blobs"});
-            }
-
-            return Ok();
-        }
-
-        [HttpGet("{ns}/exists")]
-        [ProducesDefaultResponseType]
-        public async Task<IActionResult> ExistsMultiple(
-            [FromRoute] [Required] NamespaceId ns,
-            [FromQuery] [Required] List<string> names)
-        {
-            ActionResult? accessResult = await _requestHelper.HasAccessToNamespaceAsync(User, Request, ns, new [] { StorageAclAction.ReadRefs });
-            if (accessResult != null)
-            {
-                return accessResult;
-            }
-
-            ConcurrentBag<(BucketId, RefId)> missingObject = new ();
-
-            List<(BucketId, RefId)> requestedNames = new List<(BucketId, RefId)>();
-            foreach (string name in names)
-            {
-                int separatorIndex = name.IndexOf(".", StringComparison.Ordinal);
-                if (separatorIndex == -1)
-                {
-                    return BadRequest(new ProblemDetails() { Title = $"Key {name} did not contain a '.' separator" });
-                }
-
-                BucketId bucket = new BucketId(name.Substring(0, separatorIndex));
+	[ApiController]
+	[FormatFilter]
+	[Produces(MediaTypeNames.Application.Json, MediaTypeNames.Application.Octet, CustomMediaTypeNames.UnrealCompactBinary)]
+	[Route("api/v1/refs")]
+	[Authorize]
+	public class RefsController : ControllerBase
+	{
+		private readonly IDiagnosticContext _diagnosticContext;
+		private readonly FormatResolver _formatResolver;
+		private readonly BufferedPayloadFactory _bufferedPayloadFactory;
+		private readonly IReferenceResolver _referenceResolver;
+		private readonly NginxRedirectHelper _nginxRedirectHelper;
+		private readonly IRequestHelper _requestHelper;
+		private readonly Tracer _tracer;
+
+		private readonly ILogger _logger;
+		private readonly IRefService _refService;
+		private readonly IBlobService _blobStore;
+
+		public RefsController(IRefService objectService, IBlobService blobStore, IDiagnosticContext diagnosticContext, FormatResolver formatResolver, BufferedPayloadFactory bufferedPayloadFactory, IReferenceResolver referenceResolver, NginxRedirectHelper nginxRedirectHelper, IRequestHelper requestHelper, Tracer tracer, ILogger<RefsController> logger)
+		{
+			_refService = objectService;
+			_blobStore = blobStore;
+			_diagnosticContext = diagnosticContext;
+			_formatResolver = formatResolver;
+			_bufferedPayloadFactory = bufferedPayloadFactory;
+			_referenceResolver = referenceResolver;
+			_nginxRedirectHelper = nginxRedirectHelper;
+			_requestHelper = requestHelper;
+			_tracer = tracer;
+			_logger = logger;
+		}
+
+		/// <summary>
+		/// Returns a refs key
+		/// </summary>
+		/// <param name="ns">Namespace. Each namespace is completely separated from each other. Use for different types of data that is never expected to be similar (between two different games for instance). Example: `uc4.ddc`</param>
+		/// <param name="bucket">The category/type of record you are caching. Is a clustered key together with the actual key, but all records in the same bucket can be dropped easily. Example: `terrainTexture` </param>
+		/// <param name="key">The unique name of this particular key. `iAmAVeryValidKey`</param>
+		/// <param name="format">Optional specifier to set which output format is used json/raw/cb</param>
+		[HttpGet("{ns}/{bucket}/{key}.{format?}", Order = 500)]
+		[Produces(MediaTypeNames.Application.Json, MediaTypeNames.Application.Octet, CustomMediaTypeNames.UnrealCompactBinary, CustomMediaTypeNames.JupiterInlinedPayload, CustomMediaTypeNames.UnrealCompactBinaryPackage)]
+		public async Task<IActionResult> Get(
+			[FromRoute][Required] NamespaceId ns,
+			[FromRoute][Required] BucketId bucket,
+			[FromRoute][Required] RefId key,
+			[FromRoute] string? format = null)
+		{
+			ActionResult? accessResult = await _requestHelper.HasAccessToNamespaceAsync(User, Request, ns, new[] { StorageAclAction.ReadRefs });
+			if (accessResult != null)
+			{
+				return accessResult;
+			}
+
+			try
+			{
+				(RefRecord objectRecord, BlobContents? blob) = await _refService.GetAsync(ns, bucket, key, Array.Empty<string>());
+
+				if (blob == null)
+				{
+					throw new InvalidOperationException($"Blob was null when attempting to fetch {ns} {bucket} {key}");
+				}
+
+				if (!objectRecord.IsFinalized)
+				{
+					// we do not consider un-finalized objects as valid
+					return BadRequest(new ProblemDetails { Title = $"Object {objectRecord.Bucket} {objectRecord.Name} is not finalized." });
+				}
+
+				Response.Headers[CommonHeaders.HashHeaderName] = objectRecord.BlobIdentifier.ToString();
+				Response.Headers[CommonHeaders.LastAccessHeaderName] = objectRecord.LastAccess.ToString(CultureInfo.InvariantCulture);
+
+				async Task WriteBody(BlobContents blobContents, string contentType)
+				{
+					IServerTiming? serverTiming = Request.HttpContext.RequestServices.GetService<IServerTiming>();
+					using ServerTimingMetricScoped? serverTimingScope = serverTiming?.CreateServerTimingMetricScope("body.write", "Time spent writing body");
+
+					long contentLength = blobContents.Length;
+					using TelemetrySpan scope = _tracer.StartActiveSpan("body.write").SetAttribute("operation.name", "body.write");
+					scope.SetAttribute("content-length", contentLength);
+					const int BufferSize = 64 * 1024;
+					Stream outputStream = Response.Body;
+					Response.ContentLength = contentLength;
+					Response.ContentType = contentType;
+					Response.StatusCode = StatusCodes.Status200OK;
+					try
+					{
+						await StreamCopyOperation.CopyToAsync(blobContents.Stream, outputStream, count: null, bufferSize: BufferSize, cancel: Response.HttpContext.RequestAborted);
+					}
+					catch (OperationCanceledException)
+					{
+						// do not raise exceptions for cancelled writes
+						// as we have already started writing a response we can not change the status code
+						// so we just drop a warning and proceed
+						_logger.LogWarning("The operation was canceled while writing the body");
+					}
+				}
+
+				string responseType = _formatResolver.GetResponseType(Request, format, CustomMediaTypeNames.UnrealCompactBinary);
+				Tracer.CurrentSpan.SetAttribute("response-type", responseType);
+				switch (responseType)
+				{
+					case CustomMediaTypeNames.UnrealCompactBinary:
+						{
+							// for compact binary we can just serialize our internal object
+							await WriteBody(blob, CustomMediaTypeNames.UnrealCompactBinary);
+
+							break;
+						}
+					case MediaTypeNames.Application.Octet:
+						{
+							byte[] blobMemory = await blob.Stream.ToByteArray();
+							CbObject cb = new CbObject(blobMemory);
+
+							(int, CbField?) CountFields(CbObject o)
+							{
+								int count = 0;
+								CbField? foundField = null;
+								cb.IterateAttachments(field =>
+								{
+									++count;
+									if (field.IsBinaryAttachment())
+									{
+										foundField = field;
+									}
+								});
+
+								return (count, foundField);
+							}
+
+							// breaking lambda call into private method to workaround incorrect triggering of CA1508 - [REDACTED_PROFILE_URL]/roslyn-analyzers/issues/5254
+							(int countOfAttachmentFields, CbField? binaryAttachmentField) = CountFields(cb);
+
+							if (countOfAttachmentFields == 1 && binaryAttachmentField != null)
+							{
+								// there is a single attachment field and that is of the binary attachment type, fetch that attachment and return it instead of the compact binary
+								// this is so that we match the uploaded that happened as a octet-stream which generates a small cb object with a single attachment
+
+								IoHash hash = binaryAttachmentField.AsBinaryAttachment();
+
+								BlobContents referencedBlobContents = await _blobStore.GetObjectAsync(ns, new BlobId(hash));
+
+								if (_nginxRedirectHelper.CanRedirect(Request, referencedBlobContents))
+								{
+									return _nginxRedirectHelper.CreateActionResult(referencedBlobContents, MediaTypeNames.Application.Octet);
+								}
+								await WriteBody(referencedBlobContents, MediaTypeNames.Application.Octet);
+								break;
+							}
+
+							// this doesn't look like the generated compact binary so we just return the payload
+							await using BlobContents contents = new(blobMemory);
+							await WriteBody(contents, MediaTypeNames.Application.Octet);
+							break;
+						}
+					case MediaTypeNames.Application.Json:
+						{
+							byte[] blobMemory;
+							{
+								using TelemetrySpan scope = _tracer.StartActiveSpan("json.readblob").SetAttribute("operation.name", "json.readblob");
+								blobMemory = await blob.Stream.ToByteArray();
+							}
+							CbObject cb = new CbObject(blobMemory);
+							string s = cb.ToJson();
+							await using BlobContents contents = new BlobContents(Encoding.UTF8.GetBytes(s));
+							await WriteBody(contents, MediaTypeNames.Application.Json);
+							break;
+
+						}
+					case CustomMediaTypeNames.UnrealCompactBinaryPackage:
+						{
+							using TelemetrySpan packageScope = _tracer.StartActiveSpan("cbpackage.fetch").SetAttribute("operation.name", "cbpackage.fetch");
+							byte[] blobMemory = await blob.Stream.ToByteArray();
+							CbObject cb = new CbObject(blobMemory);
+
+							IAsyncEnumerable<Attachment> attachments = _referenceResolver.GetAttachmentsAsync(ns, cb);
+
+							using CbPackageBuilder writer = new CbPackageBuilder();
+							writer.AddAttachment(objectRecord.BlobIdentifier.Hash, CbPackageAttachmentFlags.IsObject, blobMemory);
+
+							await Parallel.ForEachAsync(attachments, async (attachment, token) =>
+							{
+								IoHash attachmentHash = attachment.AsIoHash();
+								CbPackageAttachmentFlags flags = 0;
+
+								try
+								{
+									BlobContents attachmentContents;
+									if (attachment is BlobAttachment blobAttachment)
+									{
+										BlobId referencedBlob = blobAttachment.Identifier;
+										attachmentContents = await _blobStore.GetObjectAsync(ns, referencedBlob);
+									}
+									else if (attachment is ObjectAttachment objectAttachment)
+									{
+										flags |= CbPackageAttachmentFlags.IsObject;
+										BlobId referencedBlob = objectAttachment.Identifier;
+										attachmentContents = await _blobStore.GetObjectAsync(ns, referencedBlob);
+									}
+									else if (attachment is ContentIdAttachment contentIdAttachment)
+									{
+
+										ContentId contentId = contentIdAttachment.Identifier;
+										(attachmentContents, string mime) = await _blobStore.GetCompressedObject(ns, contentId, HttpContext.RequestServices);
+										if (mime == CustomMediaTypeNames.UnrealCompressedBuffer)
+										{
+											flags |= CbPackageAttachmentFlags.IsCompressed;
+										}
+										else
+										{
+											// this resolved to a uncompressed blob, the content id existed the the compressed blob didn't
+											// so resetting flags to indicate this.
+											flags = 0;
+										}
+									}
+									else
+									{
+										throw new NotSupportedException($"Unknown attachment type {attachment.GetType()}");
+									}
+
+									writer.AddAttachment(attachmentHash, flags, attachmentContents.Stream, (ulong)attachmentContents.Length);
+								}
+								catch (Exception e)
+								{
+									(CbObject errorObject, HttpStatusCode _) = ToErrorResult(e);
+									writer.AddAttachment(attachmentHash,
+										CbPackageAttachmentFlags.IsError | CbPackageAttachmentFlags.IsObject,
+										errorObject.GetView().ToArray());
+								}
+							});
+
+							byte[] packageBytes;
+							{
+								using TelemetrySpan _ = _tracer.StartActiveSpan("cbpackage.buffer").SetAttribute("operation.name", "cbpackage.buffer");
+								packageBytes = await writer.ToByteArray();
+							}
+							await using BlobContents contents = new BlobContents(packageBytes);
+							await WriteBody(contents, CustomMediaTypeNames.UnrealCompactBinaryPackage);
+							break;
+						}
+					case CustomMediaTypeNames.JupiterInlinedPayload:
+						{
+							byte[] blobMemory = await blob.Stream.ToByteArray();
+							CbObject cb = new CbObject(blobMemory);
+
+							static (int, int) CountFields(CbObject o)
+							{
+								int countOfBinaryAttachmentFields = 0;
+								int countOfAttachmentFields = 0;
+
+								o.IterateAttachments(field =>
+								{
+									if (field.IsBinaryAttachment())
+									{
+										++countOfBinaryAttachmentFields;
+									}
+
+									if (field.IsAttachment())
+									{
+										++countOfAttachmentFields;
+									}
+								});
+
+								return (countOfBinaryAttachmentFields, countOfAttachmentFields);
+							}
+							// breaking lambda call into private method to workaround incorrect triggering of CA1508 - [REDACTED_PROFILE_URL]/roslyn-analyzers/issues/5254
+							(int countOfAttachmentFields, int countOfBinaryAttachmentFields) = CountFields(cb);
+
+							// if the object consists of a single attachment field we return this attachment field instead
+							if (countOfBinaryAttachmentFields == 1 && countOfAttachmentFields == 1)
+							{
+								// fetch the blob so we can resolve any content ids in it
+								List<BlobId> referencedBlobs;
+								try
+								{
+									IAsyncEnumerable<BlobId> referencedBlobsEnumerable = _referenceResolver.GetReferencedBlobsAsync(ns, cb);
+									referencedBlobs = await referencedBlobsEnumerable.ToListAsync();
+								}
+								catch (PartialReferenceResolveException)
+								{
+									return NotFound(new ProblemDetails { Title = $"Object {bucket} {key} in namespace {ns} was missing some content ids" });
+								}
+								catch (ReferenceIsMissingBlobsException)
+								{
+									return NotFound(new ProblemDetails { Title = $"Object {bucket} {key} in namespace {ns} was missing some blobs" });
+								}
+
+								if (referencedBlobs.Count == 1)
+								{
+									BlobId attachmentToSend = referencedBlobs.First();
+									try
+									{
+										BlobContents referencedBlobContents = await _blobStore.GetObjectAsync(ns, attachmentToSend);
+										Response.Headers[CommonHeaders.InlinePayloadHash] = attachmentToSend.ToString();
+
+										if (_nginxRedirectHelper.CanRedirect(Request, referencedBlobContents))
+										{
+											return _nginxRedirectHelper.CreateActionResult(referencedBlobContents, CustomMediaTypeNames.JupiterInlinedPayload);
+										}
+										await WriteBody(referencedBlobContents, CustomMediaTypeNames.JupiterInlinedPayload);
+									}
+									catch (BlobNotFoundException)
+									{
+										return NotFound(new ProblemDetails { Title = $"Object {bucket} {key} in namespace {ns} was missing blob {attachmentToSend}" });
+									}
+									catch (Exception ex)
+									{
+										_logger.LogError(ex, "Unknown exception encountered while writing body for jupiter inlined payload.");
+										throw;
+									}
+									return new EmptyResult();
+								}
+								else if (referencedBlobs.Count == 0)
+								{
+									return NotFound(new ProblemDetails
+									{
+										Title =
+											$"Object {objectRecord.Bucket} {objectRecord.Name} did not resolve into any objects that we could find."
+									});
+								}
+
+								return BadRequest(new ProblemDetails
+								{
+									Title =
+										$"Object {objectRecord.Bucket} {objectRecord.Name} contained a content id which resolved to more then 1 blob, unable to inline this object. Use compact object response instead."
+								});
+							}
+							else if (countOfBinaryAttachmentFields == 0 && countOfAttachmentFields == 0)
+							{
+								// no attachments so we just return the compact object instead
+								await using BlobContents contents = new BlobContents(blobMemory);
+								await WriteBody(contents, CustomMediaTypeNames.JupiterInlinedPayload);
+								return new EmptyResult();
+							}
+
+							return BadRequest(new ProblemDetails
+							{
+								Title =
+									$"Object {objectRecord.Bucket} {objectRecord.Name} had more then 1 binary attachment field, unable to inline this object. Use compact object response instead."
+							});
+						}
+					default:
+						throw new NotImplementedException($"Unknown expected response type {responseType}");
+				}
+
+				// this result is ignored as we write to the body explicitly
+				return new EmptyResult();
+
+			}
+			catch (NamespaceNotFoundException e)
+			{
+				return NotFound(new ProblemDetails { Title = $"Namespace {e.Namespace} did not exist" });
+			}
+			catch (RefNotFoundException e)
+			{
+				return NotFound(new ProblemDetails { Title = $"Object {e.Bucket} {e.Key} did not exist" });
+			}
+			catch (BlobNotFoundException e)
+			{
+				return NotFound(new ProblemDetails { Title = $"Object {e.Blob} in {e.Ns} not found" });
+			}
+		}
+
+
+
+		/// <summary>
+		/// Returns the metadata about a ref key
+		/// </summary>
+		/// <param name="ns">Namespace. Each namespace is completely separated from each other. Use for different types of data that is never expected to be similar (between two different games for instance). Example: `uc4.ddc`</param>
+		/// <param name="bucket">The category/type of record you are caching. Is a clustered key together with the actual key, but all records in the same bucket can be dropped easily. Example: `terrainTexture` </param>
+		/// <param name="key">The unique name of this particular key. `iAmAVeryValidKey`</param>
+		/// <param name="fields">The fields to include in the response, omit this to include everything.</param>
+		[HttpGet("{ns}/{bucket}/{key}/metadata", Order = 500)]
+		public async Task<IActionResult> GetMetadata(
+			[FromRoute][Required] NamespaceId ns,
+			[FromRoute][Required] BucketId bucket,
+			[FromRoute][Required] RefId key,
+			[FromQuery] string[] fields)
+		{
+			ActionResult? accessResult = await _requestHelper.HasAccessToNamespaceAsync(User, Request, ns, new[] { StorageAclAction.ReadRefs });
+			if (accessResult != null)
+			{
+				return accessResult;
+			}
+
+			try
+			{
+				(RefRecord objectRecord, BlobContents? _) = await _refService.GetAsync(ns, bucket, key, fields);
+
+				return Ok(new RefMetadataResponse(objectRecord));
+			}
+			catch (NamespaceNotFoundException e)
+			{
+				return NotFound(new ProblemDetails { Title = $"Namespace {e.Namespace} did not exist" });
+			}
+			catch (RefNotFoundException e)
+			{
+				return NotFound(new ProblemDetails { Title = $"Object {e.Bucket} {e.Key} did not exist" });
+			}
+			catch (BlobNotFoundException e)
+			{
+				return NotFound(new ProblemDetails { Title = $"Object {e.Blob} in {e.Ns} not found" });
+			}
+		}
+
+		/// <summary>
+		/// Checks if a object exists
+		/// </summary>
+		/// <param name="ns">Namespace. Each namespace is completely separated from each other. Use for different types of data that is never expected to be similar (between two different games for instance). Example: `uc4.ddc`</param>
+		/// <param name="bucket">The category/type of record you are caching. Is a clustered key together with the actual key, but all records in the same bucket can be dropped easily. Example: `terrainTexture` </param>
+		/// <param name="key">The unique name of this particular key. `iAmAVeryValidKey`</param>
+		/// <returns>200 if it existed, 400 otherwise</returns>
+		[HttpHead("{ns}/{bucket}/{key}", Order = 500)]
+		[ProducesResponseType(type: typeof(OkResult), 200)]
+		[ProducesResponseType(type: typeof(ValidationProblemDetails), 400)]
+		public async Task<IActionResult> Head(
+			[FromRoute][Required] NamespaceId ns,
+			[FromRoute][Required] BucketId bucket,
+			[FromRoute][Required] RefId key)
+		{
+			ActionResult? accessResult = await _requestHelper.HasAccessToNamespaceAsync(User, Request, ns, new[] { StorageAclAction.ReadRefs });
+			if (accessResult != null)
+			{
+				return accessResult;
+			}
+
+			try
+			{
+				(RefRecord record, BlobContents? blob) = await _refService.GetAsync(ns, bucket, key, new string[] { "blobIdentifier", "IsFinalized" });
+				Response.Headers[CommonHeaders.HashHeaderName] = record.BlobIdentifier.ToString();
+
+				if (!record.IsFinalized)
+				{
+					return NotFound(new ProblemDetails { Title = $"Object {bucket} {key} in namespace {ns} is not finalized." });
+				}
+
+				blob ??= await _blobStore.GetObjectAsync(ns, record.BlobIdentifier);
+
+				// we have to verify the blobs are available locally, as the record of the key is replicated a head of the content
+				// TODO: Once we support inline replication this step is not needed as at least one region as this blob, just maybe not this current one
+				byte[] blobContents = await blob.Stream.ToByteArray();
+				CbObject compactBinaryObject = new CbObject(blobContents);
+				// the reference resolver will throw if any blob is missing, so no need to do anything other then process each reference
+				IAsyncEnumerable<BlobId> references = _referenceResolver.GetReferencedBlobsAsync(ns, compactBinaryObject);
+				List<BlobId>? _ = await references.ToListAsync();
+
+				// we have to verify the blobs are available locally, as the record of the key is replicated a head of the content
+				// TODO: Once we support inline replication this step is not needed as at least one region as this blob, just maybe not this current one
+				BlobId[] unknownBlobs = await _blobStore.FilterOutKnownBlobsAsync(ns, new BlobId[] { record.BlobIdentifier });
+				if (unknownBlobs.Length != 0)
+				{
+					return NotFound(new ProblemDetails { Title = $"Object {bucket} {key} in namespace {ns} had at least one missing blob." });
+				}
+			}
+			catch (NamespaceNotFoundException e)
+			{
+				return NotFound(new ProblemDetails { Title = $"Namespace {e.Namespace} did not exist" });
+			}
+			catch (BlobNotFoundException e)
+			{
+				return NotFound(new ProblemDetails { Title = $"Blob {e.Blob} in namespace {ns} did not exist" });
+			}
+			catch (RefNotFoundException e)
+			{
+				return NotFound(new ProblemDetails { Title = $"Object {e.Bucket} {e.Key} in namespace {e.Namespace} did not exist" });
+			}
+			catch (PartialReferenceResolveException)
+			{
+				return NotFound(new ProblemDetails { Title = $"Object {bucket} {key} in namespace {ns} was missing some content ids" });
+			}
+			catch (ReferenceIsMissingBlobsException)
+			{
+				return NotFound(new ProblemDetails { Title = $"Object {bucket} {key} in namespace {ns} was missing some blobs" });
+			}
+
+			return Ok();
+		}
+
+		[HttpGet("{ns}/exists")]
+		[ProducesDefaultResponseType]
+		public async Task<IActionResult> ExistsMultiple(
+			[FromRoute][Required] NamespaceId ns,
+			[FromQuery][Required] List<string> names)
+		{
+			ActionResult? accessResult = await _requestHelper.HasAccessToNamespaceAsync(User, Request, ns, new[] { StorageAclAction.ReadRefs });
+			if (accessResult != null)
+			{
+				return accessResult;
+			}
+
+			ConcurrentBag<(BucketId, RefId)> missingObject = new();
+
+			List<(BucketId, RefId)> requestedNames = new List<(BucketId, RefId)>();
+			foreach (string name in names)
+			{
+				int separatorIndex = name.IndexOf(".", StringComparison.Ordinal);
+				if (separatorIndex == -1)
+				{
+					return BadRequest(new ProblemDetails() { Title = $"Key {name} did not contain a '.' separator" });
+				}
+
+				BucketId bucket = new BucketId(name.Substring(0, separatorIndex));
 				RefId key = RefId.Parse(name.Substring(separatorIndex + 1));
-                requestedNames.Add((bucket, key));
-            }
-
-            IEnumerable<Task> tasks = requestedNames.Select(async pair =>
-            {
-                (BucketId bucket, RefId key) = pair;
-                try
-                {
-                    (RefRecord record, BlobContents? blob) =
-                        await _refService.GetAsync(ns, bucket, key, new string[] { "blobIdentifier" });
-
-                    blob ??= await _blobStore.GetObjectAsync(ns, record.BlobIdentifier);
-
-                    // we have to verify the blobs are available locally, as the record of the key is replicated a head of the content
-                    // TODO: Once we support inline replication this step is not needed as at least one region as this blob, just maybe not this current one
-                    byte[] blobContents = await blob.Stream.ToByteArray();
-                    CbObject cb = new CbObject(blobContents);
-                    // the reference resolver will throw if any blob is missing, so no need to do anything other then process each reference
-                    IAsyncEnumerable<BlobId> references = _referenceResolver.GetReferencedBlobsAsync(ns, cb);
-                    List<BlobId>? _ = await references.ToListAsync();
-                }
-                catch (RefNotFoundException)
-                {
-                    missingObject.Add((bucket, key));
-                }
-                catch (PartialReferenceResolveException)
-                {
-                    missingObject.Add((bucket, key));
-                }
-                catch (ReferenceIsMissingBlobsException)
-                {
-                    missingObject.Add((bucket, key));
-                }
-                catch (BlobNotFoundException)
-                {
-                    missingObject.Add((bucket, key));
-                }
-            });
-            await Task.WhenAll(tasks);
-
-            return Ok(new ExistCheckMultipleRefsResponse(missingObject.ToList()));
-        }
-
-        [HttpPut("{ns}/{bucket}/{key}.{format?}", Order = 500)]
-        [DisableRequestSizeLimit]
-        public async Task<IActionResult> PutObject(
-            [FromRoute] [Required] NamespaceId ns,
-            [FromRoute] [Required] BucketId bucket,
-            [FromRoute] [Required] RefId key,
+				requestedNames.Add((bucket, key));
+			}
+
+			IEnumerable<Task> tasks = requestedNames.Select(async pair =>
+			{
+				(BucketId bucket, RefId key) = pair;
+				try
+				{
+					(RefRecord record, BlobContents? blob) =
+						await _refService.GetAsync(ns, bucket, key, new string[] { "blobIdentifier" });
+
+					blob ??= await _blobStore.GetObjectAsync(ns, record.BlobIdentifier);
+
+					// we have to verify the blobs are available locally, as the record of the key is replicated a head of the content
+					// TODO: Once we support inline replication this step is not needed as at least one region as this blob, just maybe not this current one
+					byte[] blobContents = await blob.Stream.ToByteArray();
+					CbObject cb = new CbObject(blobContents);
+					// the reference resolver will throw if any blob is missing, so no need to do anything other then process each reference
+					IAsyncEnumerable<BlobId> references = _referenceResolver.GetReferencedBlobsAsync(ns, cb);
+					List<BlobId>? _ = await references.ToListAsync();
+				}
+				catch (RefNotFoundException)
+				{
+					missingObject.Add((bucket, key));
+				}
+				catch (PartialReferenceResolveException)
+				{
+					missingObject.Add((bucket, key));
+				}
+				catch (ReferenceIsMissingBlobsException)
+				{
+					missingObject.Add((bucket, key));
+				}
+				catch (BlobNotFoundException)
+				{
+					missingObject.Add((bucket, key));
+				}
+			});
+			await Task.WhenAll(tasks);
+
+			return Ok(new ExistCheckMultipleRefsResponse(missingObject.ToList()));
+		}
+
+		[HttpPut("{ns}/{bucket}/{key}.{format?}", Order = 500)]
+		[DisableRequestSizeLimit]
+		public async Task<IActionResult> PutObject(
+			[FromRoute][Required] NamespaceId ns,
+			[FromRoute][Required] BucketId bucket,
+			[FromRoute][Required] RefId key,
 			CancellationToken cancellationToken)
-        {
-            ActionResult? accessResult = await _requestHelper.HasAccessToNamespaceAsync(User, Request, ns, new [] { StorageAclAction.WriteRefs });
-            if (accessResult != null)
-            {
-                return accessResult;
-            }
-
-            _diagnosticContext.Set("Content-Length", Request.ContentLength ?? -1);
-
-            CbObject payloadObject;
-            BlobId blobHeader;
-
-            try
-            {
-                using BufferedPayload payload = await _bufferedPayloadFactory.CreateFromRequest(Request);
-
-                BlobId headerHash;
-                if (Request.Headers.ContainsKey(CommonHeaders.HashHeaderName))
-                {
-                    headerHash = BlobId.Parse(Request.Headers[CommonHeaders.HashHeaderName]);
-                }
-                else
-                {
-                    return BadRequest(new ProblemDetails
-                    {
-                        Title = $"Missing expected header {CommonHeaders.HashHeaderName}"
-                    });
-                }
-
-                blobHeader = headerHash;
-
-                switch (Request.ContentType)
-                {
-                    case MediaTypeNames.Application.Json:
-                    {
-                        // TODO: define a scheme for how a json object specifies references
-
-                        blobHeader = await _blobStore.PutObjectAsync(ns, payload, headerHash, cancellationToken);
-
-                        // TODO: convert the json object into a compact binary instead
-                        CbWriter writer = new CbWriter();
-                        writer.BeginObject();
-						writer.WriteBinaryAttachmentValue(blobHeader.Hash);
-                        writer.EndObject();
-
-                        byte[] blob = writer.ToByteArray();
-                        payloadObject = new CbObject(blob);
-						blobHeader = new BlobId(IoHash.Compute(blob));
-                        break;
-                    }
-                    case CustomMediaTypeNames.UnrealCompactBinary:
-                    {
-                        await using MemoryStream ms = new MemoryStream();
-                        await using Stream payloadStream = payload.GetStream();
-                        await payloadStream.CopyToAsync(ms, cancellationToken);
-                        payloadObject = new CbObject(ms.ToArray());
-                        break;
-                    }
-                    case MediaTypeNames.Application.Octet:
-                    {
-                        blobHeader = await _blobStore.PutObjectAsync(ns, payload, headerHash, cancellationToken);
-
-                        CbWriter writer = new CbWriter();
-                        writer.BeginObject();
-						writer.WriteBinaryAttachment("RawHash", blobHeader.Hash);
-                        writer.WriteInteger("RawSize", payload.Length);
-                        writer.EndObject();
-
-                        byte[] blob = writer.ToByteArray();
-                        payloadObject = new CbObject(blob);
-						blobHeader = new BlobId(IoHash.Compute(blob));
-                        break;
-                    }
-                    default:
-                        throw new Exception($"Unknown request type {Request.ContentType}, if submitting a blob please use {MediaTypeNames.Application.Octet}");
-                }
-            }
-            catch (HashMismatchException e)
-            {
-                return BadRequest(new ProblemDetails
-                {
-                    Title = $"Incorrect hash, got hash \"{e.SuppliedHash}\" but hash of content was determined to be \"{e.ContentHash}\""
-                });
-            }
-            catch (ClientSendSlowException e)
-            {
-                return Problem(e.Message, null, (int)HttpStatusCode.RequestTimeout);
-            }
-
-            (ContentId[] missingReferences, BlobId[] missingBlobs) = await _refService.PutAsync(ns, bucket, key, blobHeader, payloadObject);
+		{
+			ActionResult? accessResult = await _requestHelper.HasAccessToNamespaceAsync(User, Request, ns, new[] { StorageAclAction.WriteRefs });
+			if (accessResult != null)
+			{
+				return accessResult;
+			}
+
+			_diagnosticContext.Set("Content-Length", Request.ContentLength ?? -1);
+
+			CbObject payloadObject;
+			BlobId blobHeader;
+
+			try
+			{
+				using BufferedPayload payload = await _bufferedPayloadFactory.CreateFromRequest(Request);
+
+				BlobId headerHash;
+				if (Request.Headers.ContainsKey(CommonHeaders.HashHeaderName))
+				{
+					headerHash = BlobId.Parse(Request.Headers[CommonHeaders.HashHeaderName]);
+				}
+				else
+				{
+					return BadRequest(new ProblemDetails
+					{
+						Title = $"Missing expected header {CommonHeaders.HashHeaderName}"
+					});
+				}
+
+				blobHeader = headerHash;
+
+				switch (Request.ContentType)
+				{
+					case MediaTypeNames.Application.Json:
+						{
+							// TODO: define a scheme for how a json object specifies references
+
+							blobHeader = await _blobStore.PutObjectAsync(ns, payload, headerHash, cancellationToken);
+
+							// TODO: convert the json object into a compact binary instead
+							CbWriter writer = new CbWriter();
+							writer.BeginObject();
+							writer.WriteBinaryAttachmentValue(blobHeader.Hash);
+							writer.EndObject();
+
+							byte[] blob = writer.ToByteArray();
+							payloadObject = new CbObject(blob);
+							blobHeader = new BlobId(IoHash.Compute(blob));
+							break;
+						}
+					case CustomMediaTypeNames.UnrealCompactBinary:
+						{
+							await using MemoryStream ms = new MemoryStream();
+							await using Stream payloadStream = payload.GetStream();
+							await payloadStream.CopyToAsync(ms, cancellationToken);
+							payloadObject = new CbObject(ms.ToArray());
+							break;
+						}
+					case MediaTypeNames.Application.Octet:
+						{
+							blobHeader = await _blobStore.PutObjectAsync(ns, payload, headerHash, cancellationToken);
+
+							CbWriter writer = new CbWriter();
+							writer.BeginObject();
+							writer.WriteBinaryAttachment("RawHash", blobHeader.Hash);
+							writer.WriteInteger("RawSize", payload.Length);
+							writer.EndObject();
+
+							byte[] blob = writer.ToByteArray();
+							payloadObject = new CbObject(blob);
+							blobHeader = new BlobId(IoHash.Compute(blob));
+							break;
+						}
+					default:
+						throw new Exception($"Unknown request type {Request.ContentType}, if submitting a blob please use {MediaTypeNames.Application.Octet}");
+				}
+			}
+			catch (HashMismatchException e)
+			{
+				return BadRequest(new ProblemDetails
+				{
+					Title = $"Incorrect hash, got hash \"{e.SuppliedHash}\" but hash of content was determined to be \"{e.ContentHash}\""
+				});
+			}
+			catch (ClientSendSlowException e)
+			{
+				return Problem(e.Message, null, (int)HttpStatusCode.RequestTimeout);
+			}
+
+			(ContentId[] missingReferences, BlobId[] missingBlobs) = await _refService.PutAsync(ns, bucket, key, blobHeader, payloadObject);
 
 			List<IoHash> missingHashes = new List<IoHash>();
 			missingHashes.AddRange(missingReferences.Select(x => x.Hash));
-            missingHashes.AddRange(missingBlobs.Select(x => x.Hash));
-            return Ok(new PutObjectResponse(missingHashes.ToArray()));
-        }
-
-        [HttpPut("{ns}/{bucket}/{key}", Order = 300)]
-        [DisableRequestSizeLimit]
-        [RequiredContentType(CustomMediaTypeNames.UnrealCompactBinaryPackage)]
-        public async Task<IActionResult> PutPackage(
-            [FromRoute][Required] NamespaceId ns,
-            [FromRoute][Required] BucketId bucket,
-            [FromRoute][Required] RefId key,
+			missingHashes.AddRange(missingBlobs.Select(x => x.Hash));
+			return Ok(new PutObjectResponse(missingHashes.ToArray()));
+		}
+
+		[HttpPut("{ns}/{bucket}/{key}", Order = 300)]
+		[DisableRequestSizeLimit]
+		[RequiredContentType(CustomMediaTypeNames.UnrealCompactBinaryPackage)]
+		public async Task<IActionResult> PutPackage(
+			[FromRoute][Required] NamespaceId ns,
+			[FromRoute][Required] BucketId bucket,
+			[FromRoute][Required] RefId key,
 			CancellationToken cancellationToken)
-        {
-            ActionResult? accessResult = await _requestHelper.HasAccessToNamespaceAsync(User, Request, ns, new [] { StorageAclAction.WriteRefs });
-            if (accessResult != null)
-            {
-                return accessResult;
-            }
-
-            _diagnosticContext.Set("Content-Length", Request.ContentLength ?? -1);
-
-            byte[] b = await ReadRawBodyAsync(Request);
-            CbPackageReader packageReader = await CbPackageReader.Create(new MemoryStream(b));
-
-            try
-            {
-                await foreach ((CbPackageAttachmentEntry entry, byte[] blob) in packageReader.IterateAttachments())
-                {
-                    if (entry.Flags.HasFlag(CbPackageAttachmentFlags.IsError))
-                    {
-                        return BadRequest(new ProblemDetails
-                        {
-                            Title = $"Package contained attachment with error {entry.AttachmentHash}\""
-                        });
-                    }
-                    if (entry.Flags.HasFlag(CbPackageAttachmentFlags.IsCompressed))
-                    {
-                        using MemoryBufferedPayload payload = new MemoryBufferedPayload(blob);
-                        await _blobStore.PutCompressedObject(ns, payload, new ContentId(entry.AttachmentHash), HttpContext.RequestServices);
-                    }
-                    else
-                    {
-                        await _blobStore.PutObjectAsync(ns, blob, new BlobId(entry.AttachmentHash), cancellationToken);
-                    }
-                }
-            }
-            catch (HashMismatchException e)
-            {
-                return BadRequest(new ProblemDetails
-                {
-                    Title = $"Incorrect hash, got hash \"{e.SuppliedHash}\" but hash of content was determined to be \"{e.ContentHash}\""
-                });
-            }
-            
-            CbObject rootObject = packageReader.RootObject;
-            BlobId rootObjectHash = new BlobId(packageReader.RootHash);
-
-            (ContentId[] missingReferences, BlobId[] missingBlobs) = await _refService.PutAsync(ns, bucket, key, rootObjectHash, rootObject);
+		{
+			ActionResult? accessResult = await _requestHelper.HasAccessToNamespaceAsync(User, Request, ns, new[] { StorageAclAction.WriteRefs });
+			if (accessResult != null)
+			{
+				return accessResult;
+			}
+
+			_diagnosticContext.Set("Content-Length", Request.ContentLength ?? -1);
+
+			byte[] b = await ReadRawBodyAsync(Request);
+			CbPackageReader packageReader = await CbPackageReader.Create(new MemoryStream(b));
+
+			try
+			{
+				await foreach ((CbPackageAttachmentEntry entry, byte[] blob) in packageReader.IterateAttachments())
+				{
+					if (entry.Flags.HasFlag(CbPackageAttachmentFlags.IsError))
+					{
+						return BadRequest(new ProblemDetails
+						{
+							Title = $"Package contained attachment with error {entry.AttachmentHash}\""
+						});
+					}
+					if (entry.Flags.HasFlag(CbPackageAttachmentFlags.IsCompressed))
+					{
+						using MemoryBufferedPayload payload = new MemoryBufferedPayload(blob);
+						await _blobStore.PutCompressedObject(ns, payload, new ContentId(entry.AttachmentHash), HttpContext.RequestServices);
+					}
+					else
+					{
+						await _blobStore.PutObjectAsync(ns, blob, new BlobId(entry.AttachmentHash), cancellationToken);
+					}
+				}
+			}
+			catch (HashMismatchException e)
+			{
+				return BadRequest(new ProblemDetails
+				{
+					Title = $"Incorrect hash, got hash \"{e.SuppliedHash}\" but hash of content was determined to be \"{e.ContentHash}\""
+				});
+			}
+
+			CbObject rootObject = packageReader.RootObject;
+			BlobId rootObjectHash = new BlobId(packageReader.RootHash);
+
+			(ContentId[] missingReferences, BlobId[] missingBlobs) = await _refService.PutAsync(ns, bucket, key, rootObjectHash, rootObject);
 			return Ok(new PutObjectResponse(missingReferences, missingBlobs));
-        }
+		}
 
 		static async Task<byte[]> ReadRawBodyAsync(HttpRequest request)
 		{
@@ -768,522 +768,522 @@ namespace Horde.Server.Ddc
 		}
 
 		[HttpPost("{ns}/{bucket}/{key}/finalize/{hash}.{format?}")]
-        public async Task<IActionResult> FinalizeObject(
-            [FromRoute] [Required] NamespaceId ns,
-            [FromRoute] [Required] BucketId bucket,
-            [FromRoute] [Required] RefId key,
-            [FromRoute] [Required] BlobId hash)
-        {
-            ActionResult? accessResult = await _requestHelper.HasAccessToNamespaceAsync(User, Request, ns, new [] { StorageAclAction.WriteRefs });
-            if (accessResult != null)
-            {
-                return accessResult;
-            }
-
-            try
-            {
-                (ContentId[] missingReferences, BlobId[] missingBlobs) = await _refService.FinalizeAsync(ns, bucket, key, hash);
+		public async Task<IActionResult> FinalizeObject(
+			[FromRoute][Required] NamespaceId ns,
+			[FromRoute][Required] BucketId bucket,
+			[FromRoute][Required] RefId key,
+			[FromRoute][Required] BlobId hash)
+		{
+			ActionResult? accessResult = await _requestHelper.HasAccessToNamespaceAsync(User, Request, ns, new[] { StorageAclAction.WriteRefs });
+			if (accessResult != null)
+			{
+				return accessResult;
+			}
+
+			try
+			{
+				(ContentId[] missingReferences, BlobId[] missingBlobs) = await _refService.FinalizeAsync(ns, bucket, key, hash);
 				return Ok(new PutObjectResponse(missingReferences, missingBlobs));
-            }
-            catch (ObjectHashMismatchException e)
-            {
-                return BadRequest(e.Message);
-            }
-            catch (RefNotFoundException e)
-            {
-                return NotFound(e.Message);
-            }
-        }
-
-        [HttpPost("{ns}")]
-        [Consumes(CustomMediaTypeNames.UnrealCompactBinary)]
-        [Produces(CustomMediaTypeNames.UnrealCompactBinary)]
-        public async Task<IActionResult> Batch(
-            [FromRoute] [Required] NamespaceId ns,
-            [FromBody] [Required] RefBatchOps ops)
-        {
-            AclAction ActionForOp(RefBatchOps.RefBatchOp.RefOperation op)
-            {
-                switch (op)
-                {
-                    case RefBatchOps.RefBatchOp.RefOperation.GET:
-                        return StorageAclAction.ReadRefs;
-                    case RefBatchOps.RefBatchOp.RefOperation.PUT:
-                        return StorageAclAction.WriteRefs;
-                    case RefBatchOps.RefBatchOp.RefOperation.HEAD:
+			}
+			catch (ObjectHashMismatchException e)
+			{
+				return BadRequest(e.Message);
+			}
+			catch (RefNotFoundException e)
+			{
+				return NotFound(e.Message);
+			}
+		}
+
+		[HttpPost("{ns}")]
+		[Consumes(CustomMediaTypeNames.UnrealCompactBinary)]
+		[Produces(CustomMediaTypeNames.UnrealCompactBinary)]
+		public async Task<IActionResult> Batch(
+			[FromRoute][Required] NamespaceId ns,
+			[FromBody][Required] RefBatchOps ops)
+		{
+			AclAction ActionForOp(RefBatchOps.RefBatchOp.RefOperation op)
+			{
+				switch (op)
+				{
+					case RefBatchOps.RefBatchOp.RefOperation.GET:
 						return StorageAclAction.ReadRefs;
-                    default:
-                        throw new ArgumentOutOfRangeException(nameof(op), op, null);
-                }
-            }
-
-            AclAction[] requiredActions = ops.Ops.Select(op => ActionForOp(op.Op)).ToArray();
-
-            ActionResult? accessResult = await _requestHelper.HasAccessToNamespaceAsync(User, Request, ns, requiredActions);
-            if (accessResult != null)
-            {
-                return accessResult;
-            }
-
-            HashSet<uint> usedOpIds = new HashSet<uint>();
-            foreach (RefBatchOps.RefBatchOp batchOp in ops.Ops)
-            {
-                bool added = usedOpIds.Add(batchOp.OpId);
-                if (!added)
-                {
-                    return BadRequest(new ProblemDetails { Title = $"Duplicate op ids used for id: {batchOp.OpId}" });
-                }
-            }
-            ConcurrentDictionary<uint, (CbObject, HttpStatusCode)> results = new();
-
-            async Task<(CbObject, HttpStatusCode)> BatchGetOp(RefBatchOps.RefBatchOp op)
-            {
-                try
-                {
-                    (RefRecord objectRecord, BlobContents? blob) = await _refService.GetAsync(ns, op.Bucket, op.Key, Array.Empty<string>());
-
-                    if (!objectRecord.IsFinalized)
-                    {
-                        throw new Exception("Object is not finalized");
-                    }
-
-                    if (blob == null)
-                    {
-                        throw new Exception();
-                    }
-
-                    CbObject cb = new CbObject(await blob.Stream.ToByteArray());
-
-                    if (op.ResolveAttachments ?? false)
-                    {
-                        IAsyncEnumerable<BlobId> references = _referenceResolver.GetReferencedBlobsAsync(ns, cb);
-                        List<BlobId>? _ = await references.ToListAsync();
-                    }
-
-                    return (cb, HttpStatusCode.OK);
-                }
-                catch (Exception ex) when( ex is RefNotFoundException or PartialReferenceResolveException or ReferenceIsMissingBlobsException)
-                {
-                    return ToErrorResult(ex, HttpStatusCode.NotFound);
-                }
-                catch (Exception e)
-                {
-                    return ToErrorResult(e);
-                }
-            }
-
-            async Task<(CbObject, HttpStatusCode)> BatchHeadOp(RefBatchOps.RefBatchOp op)
-            {
-                try
-                {
-                    (RefRecord record, BlobContents? blob) = await _refService.GetAsync(ns, op.Bucket, op.Key, new string[] { "blobIdentifier" });
-
-                    if (!record.IsFinalized)
-                    {
-                        return (CbObject.Build(writer => writer.WriteBool("exists", false)), HttpStatusCode.NotFound);
-                    }
-
-                    blob ??= await _blobStore.GetObjectAsync(ns, record.BlobIdentifier);
-
-                    if (op.ResolveAttachments ?? false)
-                    {
-                        byte[] blobContents = await blob.Stream.ToByteArray();
-                        CbObject cb = new CbObject(blobContents);
-                        // the reference resolver will throw if any blob is missing, so no need to do anything other then process each reference
-                        IAsyncEnumerable<BlobId> references = _referenceResolver.GetReferencedBlobsAsync(ns, cb);
-                        List<BlobId>? _ = await references.ToListAsync();
-                    }
-
-                    if (blob == null)
-                    {
-                        throw new Exception();
-                    }
-
-                    return (CbObject.Build(writer => writer.WriteBool("exists", true)), HttpStatusCode.OK);
-                }
-                catch (Exception ex) when( ex is RefNotFoundException or PartialReferenceResolveException or ReferenceIsMissingBlobsException)
-                {
-                    return (CbObject.Build(writer => writer.WriteBool("exists", false)), HttpStatusCode.NotFound);
-                }
-                catch (Exception e)
-                {
-                    return ToErrorResult(e);
-                }
-            }
-
-            async Task<(CbObject, HttpStatusCode)> BatchPutOp(RefBatchOps.RefBatchOp op)
-            {
-                try
-                {
-                    if (op.Payload == null || op.Payload.Equals(CbObject.Empty))
-                    {
-                        throw new Exception($"Missing payload for operation: {op.OpId}");
-                    }
-
-                    if (op.PayloadHash == null)
-                    {
-                        throw new Exception($"Missing payload hash for operation: {op.OpId}");
-                    }
-
-                    IoHash headerHash = op.PayloadHash.Value;
+					case RefBatchOps.RefBatchOp.RefOperation.PUT:
+						return StorageAclAction.WriteRefs;
+					case RefBatchOps.RefBatchOp.RefOperation.HEAD:
+						return StorageAclAction.ReadRefs;
+					default:
+						throw new ArgumentOutOfRangeException(nameof(op), op, null);
+				}
+			}
+
+			AclAction[] requiredActions = ops.Ops.Select(op => ActionForOp(op.Op)).ToArray();
+
+			ActionResult? accessResult = await _requestHelper.HasAccessToNamespaceAsync(User, Request, ns, requiredActions);
+			if (accessResult != null)
+			{
+				return accessResult;
+			}
+
+			HashSet<uint> usedOpIds = new HashSet<uint>();
+			foreach (RefBatchOps.RefBatchOp batchOp in ops.Ops)
+			{
+				bool added = usedOpIds.Add(batchOp.OpId);
+				if (!added)
+				{
+					return BadRequest(new ProblemDetails { Title = $"Duplicate op ids used for id: {batchOp.OpId}" });
+				}
+			}
+			ConcurrentDictionary<uint, (CbObject, HttpStatusCode)> results = new();
+
+			async Task<(CbObject, HttpStatusCode)> BatchGetOp(RefBatchOps.RefBatchOp op)
+			{
+				try
+				{
+					(RefRecord objectRecord, BlobContents? blob) = await _refService.GetAsync(ns, op.Bucket, op.Key, Array.Empty<string>());
+
+					if (!objectRecord.IsFinalized)
+					{
+						throw new Exception("Object is not finalized");
+					}
+
+					if (blob == null)
+					{
+						throw new Exception();
+					}
+
+					CbObject cb = new CbObject(await blob.Stream.ToByteArray());
+
+					if (op.ResolveAttachments ?? false)
+					{
+						IAsyncEnumerable<BlobId> references = _referenceResolver.GetReferencedBlobsAsync(ns, cb);
+						List<BlobId>? _ = await references.ToListAsync();
+					}
+
+					return (cb, HttpStatusCode.OK);
+				}
+				catch (Exception ex) when (ex is RefNotFoundException or PartialReferenceResolveException or ReferenceIsMissingBlobsException)
+				{
+					return ToErrorResult(ex, HttpStatusCode.NotFound);
+				}
+				catch (Exception e)
+				{
+					return ToErrorResult(e);
+				}
+			}
+
+			async Task<(CbObject, HttpStatusCode)> BatchHeadOp(RefBatchOps.RefBatchOp op)
+			{
+				try
+				{
+					(RefRecord record, BlobContents? blob) = await _refService.GetAsync(ns, op.Bucket, op.Key, new string[] { "blobIdentifier" });
+
+					if (!record.IsFinalized)
+					{
+						return (CbObject.Build(writer => writer.WriteBool("exists", false)), HttpStatusCode.NotFound);
+					}
+
+					blob ??= await _blobStore.GetObjectAsync(ns, record.BlobIdentifier);
+
+					if (op.ResolveAttachments ?? false)
+					{
+						byte[] blobContents = await blob.Stream.ToByteArray();
+						CbObject cb = new CbObject(blobContents);
+						// the reference resolver will throw if any blob is missing, so no need to do anything other then process each reference
+						IAsyncEnumerable<BlobId> references = _referenceResolver.GetReferencedBlobsAsync(ns, cb);
+						List<BlobId>? _ = await references.ToListAsync();
+					}
+
+					if (blob == null)
+					{
+						throw new Exception();
+					}
+
+					return (CbObject.Build(writer => writer.WriteBool("exists", true)), HttpStatusCode.OK);
+				}
+				catch (Exception ex) when (ex is RefNotFoundException or PartialReferenceResolveException or ReferenceIsMissingBlobsException)
+				{
+					return (CbObject.Build(writer => writer.WriteBool("exists", false)), HttpStatusCode.NotFound);
+				}
+				catch (Exception e)
+				{
+					return ToErrorResult(e);
+				}
+			}
+
+			async Task<(CbObject, HttpStatusCode)> BatchPutOp(RefBatchOps.RefBatchOp op)
+			{
+				try
+				{
+					if (op.Payload == null || op.Payload.Equals(CbObject.Empty))
+					{
+						throw new Exception($"Missing payload for operation: {op.OpId}");
+					}
+
+					if (op.PayloadHash == null)
+					{
+						throw new Exception($"Missing payload hash for operation: {op.OpId}");
+					}
+
+					IoHash headerHash = op.PayloadHash.Value;
 					IoHash objectHash = IoHash.Compute(op.Payload.GetView().Span);
-				
-                    if (!headerHash.Equals(objectHash))
-                    {
-                        throw new HashMismatchException(op.PayloadHash.Value, objectHash);
-                    }
-
-                    (ContentId[] missingReferences, BlobId[] missingBlobs) = await _refService.PutAsync(ns, op.Bucket, op.Key, new BlobId(objectHash), op.Payload);
-                    return (CbSerializer.Serialize(new PutObjectResponse(missingReferences, missingBlobs)), HttpStatusCode.OK);
-                }
-                catch (Exception e)
-                {
-                    return ToErrorResult(e);
-                }
-            }
-
-            await Parallel.ForEachAsync(ops.Ops, CancellationToken.None, async (op, token) =>
-            {
-                switch (op.Op)
-                {
-                    case RefBatchOps.RefBatchOp.RefOperation.GET:
-                        results.TryAdd(op.OpId, await BatchGetOp(op));
-                        break;
-                    case RefBatchOps.RefBatchOp.RefOperation.PUT:
-                        results.TryAdd(op.OpId, await BatchPutOp(op));
-                        break;
-                    case RefBatchOps.RefBatchOp.RefOperation.HEAD:
-                        results.TryAdd(op.OpId, await BatchHeadOp(op));
-                        break;
-                    case RefBatchOps.RefBatchOp.RefOperation.INVALID:
-                    default:
-                        throw new NotImplementedException($"Unknown op type {op.Op}");
-                }
-
-                await Task.CompletedTask;
-            });
-
-            return Ok(new BatchOpsResponse()
-            {
-                Results = results.Select(result =>
-                {
-                    return new BatchOpsResponse.OpResponses()
-                    {
-                        OpId = result.Key,
-                        Response = result.Value.Item1,
-                        StatusCode = (int)result.Value.Item2
-                    };
-                }).ToList()
-            });
-        }
-
-        private static (CbObject, HttpStatusCode) ToErrorResult(Exception exception, HttpStatusCode statusCode = HttpStatusCode.InternalServerError)
-        {
-            Exception e = exception;
-            CbWriter writer = new CbWriter();
-            writer.BeginObject();
-            writer.WriteString("title", e.Message);
-            writer.WriteInteger("status", (int)statusCode);
-            if (e.StackTrace != null)
-            {
-                writer.WriteString("stackTrace", e.StackTrace);
-            }
-            writer.EndObject();
-            return (writer.ToObject(), statusCode);
-        }
-
-        /// <summary>
-        /// Delete a individual refs key
-        /// </summary>
-        /// <param name="ns">Namespace. Each namespace is completely separated from each other. Use for different types of data that is never expected to be similar (between two different games for instance)</param>
-        /// <param name="bucket">The category/type of record you are caching. Is a clustered key together with the actual key, but all records in the same bucket can be dropped easily.</param>
-        /// <param name="key">The unique name of this particular key</param>
-        [HttpDelete("{ns}/{bucket}/{key}", Order = 500)]
-        [ProducesResponseType(200)]
-        [ProducesResponseType(404)]
-        public async Task<IActionResult> Delete(
-            [FromRoute] [Required] NamespaceId ns,
-            [FromRoute] [Required] BucketId bucket,
-            [FromRoute] [Required] RefId key)
-        {
-            ActionResult? accessResult = await _requestHelper.HasAccessToNamespaceAsync(User, Request, ns, new [] { StorageAclAction.DeleteRefs });
-            if (accessResult != null)
-            {
-                return accessResult;
-            }
-
-            try
-            {
-                bool deleted = await _refService.DeleteAsync(ns, bucket, key);
-                if (!deleted)
-                {
-                    return NotFound(new ProblemDetails { Title = $"Object {key} in bucket {bucket} and namespace {ns} did not exist" });
-                }
-                return Ok(new RefDeletedResponse(deleted ? 1 : 0));
-            }
-            catch (NamespaceNotFoundException e)
-            {
-                return NotFound(new ProblemDetails { Title = $"Namespace {e.Namespace} did not exist" });
-            }
-            catch (RefNotFoundException)
-            {
-                return NotFound(new ProblemDetails { Title = $"Object {key} in bucket {bucket} and namespace {ns} did not exist" });
-            }
-        }
-    }
-
-    public class RefDeletedResponse
-    {
-        public RefDeletedResponse()
-        {
-
-        }
-
-        public RefDeletedResponse(int deletedCount)
-        {
-            DeletedCount = deletedCount;
-        }
-
-        [CbField("deletedCount")]
-        public int DeletedCount { get; set; }
-    }
-
-    public class BucketDeletedResponse
-    {
-        public BucketDeletedResponse()
-        {
-
-        }
-
-        public BucketDeletedResponse(long countOfDeletedRecords)
-        {
-            CountOfDeletedRecords = countOfDeletedRecords;
-        }
-
-        [CbField("countOfDeletedRecords")]
-        public long CountOfDeletedRecords { get; set; }
-    }
-
-    public class RefBatchOps
-    {
-        public RefBatchOps()
-        {
-            Ops = Array.Empty<RefBatchOp>();
-        }
-
-        [System.Diagnostics.CodeAnalysis.SuppressMessage("Design", "CA1034:Nested types should not be visible", Justification = "For serialization only")]
-        public class RefBatchOp
-        {
-            public RefBatchOp()
-            {
-                Payload = null;
-                PayloadHash = null;
-            }
-
-            public enum RefOperation
-            {
-                INVALID,
-                GET,
-                PUT,
-                HEAD,
-            }
-
-            [Required]
-            [CbField("opId")]
-            public uint OpId { get; set; }
-
-            [CbField("op")]
-            [JsonIgnore]
-            [UsedImplicitly]
-            public string OpString
-            {
-                get => Op.ToString();
-                set => Op = Enum.Parse<RefOperation>(value);
-            }
-
-            [Required]
-            public RefOperation Op { get; set; } = RefOperation.INVALID;
-
-            [Required]
-            [CbField("bucket")]
-            public BucketId Bucket { get; set; }
-
-            [Required]
-            [CbField("key")]
-            public RefId Key { get; set; }
-
-            [CbField("resolveAttachments")]
-            public bool? ResolveAttachments { get; set; } = null;
-
-            [CbField("payload")]
-            public CbObject? Payload { get; set; } = null;
-
-            [CbField("payloadHash")] 
-            public IoHash? PayloadHash { get; set; } = null;
-        }
-
-        [CbField("ops")]
-        public RefBatchOp[] Ops { get; set; }
-    }
-
-    public class BatchOpsResponse
-    {
-        public BatchOpsResponse()
-        {
-
-        }
-
-        [System.Diagnostics.CodeAnalysis.SuppressMessage("Design", "CA1034:Nested types should not be visible", Justification = "For serialization only")]
-        public class OpResponses
-        {
-            public OpResponses()
-            {
-
-            }
-            [CbField("opId")]
-            public uint OpId { get; set; }
-
-            [CbField("response")]
-            public CbObject Response { get; set; } = null!;
-
-            [CbField("statusCode")]
-            public int StatusCode { get; set; }
-        }
-
-        [CbField("results")]
-        [System.Diagnostics.CodeAnalysis.SuppressMessage("Usage", "CA2227:Collection properties should be read only", Justification = "Used by serialization")]
-        public List<OpResponses> Results { get; set; } = new List<OpResponses>();
-    }
-
-    public class RefMetadataResponse
-    {
-        public RefMetadataResponse()
-        {
-            InlinePayload = null!;
-        }
-
-        [JsonConstructor]
-        public RefMetadataResponse(NamespaceId ns, BucketId bucket, RefId name, BlobId payloadIdentifier, DateTime lastAccess, bool isFinalized, byte[]? inlinePayload)
-        {
-            Ns = ns;
-            Bucket = bucket;
-            Name = name;
-            PayloadIdentifier = payloadIdentifier;
-            LastAccess = lastAccess;
-            IsFinalized = isFinalized;
-            InlinePayload = inlinePayload;
-        }
-
-        public RefMetadataResponse(RefRecord objectRecord)
-        {
-            Ns = objectRecord.Namespace;
-            Bucket = objectRecord.Bucket;
-            Name = objectRecord.Name;
-            PayloadIdentifier = objectRecord.BlobIdentifier;
-            LastAccess = objectRecord.LastAccess;
-            IsFinalized = objectRecord.IsFinalized;
-            InlinePayload = objectRecord.InlinePayload;
-        }
-
-        [CbField("ns")]
-        public NamespaceId Ns { get; set; }
-        
-        [CbField("bucket")]
-        public BucketId Bucket { get; set; }
-
-        [CbField("name")]
-        public RefId Name { get; set; }
-
-        [CbField("payloadIdentifier")]
-        public BlobId PayloadIdentifier { get; set; }
-
-        [CbField("lastAccess")]
-        public DateTime LastAccess { get; set; }
-
-        [CbField("isFinalized")]
-        public bool IsFinalized { get; set; }
-
-        [CbField("inlinePayload")]
-        public byte[]? InlinePayload { get; set; }
-    }
-
-    public class PutObjectResponse
-    {
-        public PutObjectResponse()
-        {
-            Needs = null!;
-        }
-
-        public PutObjectResponse(IoHash[] missingReferences)
-        {
-            Needs = missingReferences;
-        }
+
+					if (!headerHash.Equals(objectHash))
+					{
+						throw new HashMismatchException(op.PayloadHash.Value, objectHash);
+					}
+
+					(ContentId[] missingReferences, BlobId[] missingBlobs) = await _refService.PutAsync(ns, op.Bucket, op.Key, new BlobId(objectHash), op.Payload);
+					return (CbSerializer.Serialize(new PutObjectResponse(missingReferences, missingBlobs)), HttpStatusCode.OK);
+				}
+				catch (Exception e)
+				{
+					return ToErrorResult(e);
+				}
+			}
+
+			await Parallel.ForEachAsync(ops.Ops, CancellationToken.None, async (op, token) =>
+			{
+				switch (op.Op)
+				{
+					case RefBatchOps.RefBatchOp.RefOperation.GET:
+						results.TryAdd(op.OpId, await BatchGetOp(op));
+						break;
+					case RefBatchOps.RefBatchOp.RefOperation.PUT:
+						results.TryAdd(op.OpId, await BatchPutOp(op));
+						break;
+					case RefBatchOps.RefBatchOp.RefOperation.HEAD:
+						results.TryAdd(op.OpId, await BatchHeadOp(op));
+						break;
+					case RefBatchOps.RefBatchOp.RefOperation.INVALID:
+					default:
+						throw new NotImplementedException($"Unknown op type {op.Op}");
+				}
+
+				await Task.CompletedTask;
+			});
+
+			return Ok(new BatchOpsResponse()
+			{
+				Results = results.Select(result =>
+				{
+					return new BatchOpsResponse.OpResponses()
+					{
+						OpId = result.Key,
+						Response = result.Value.Item1,
+						StatusCode = (int)result.Value.Item2
+					};
+				}).ToList()
+			});
+		}
+
+		private static (CbObject, HttpStatusCode) ToErrorResult(Exception exception, HttpStatusCode statusCode = HttpStatusCode.InternalServerError)
+		{
+			Exception e = exception;
+			CbWriter writer = new CbWriter();
+			writer.BeginObject();
+			writer.WriteString("title", e.Message);
+			writer.WriteInteger("status", (int)statusCode);
+			if (e.StackTrace != null)
+			{
+				writer.WriteString("stackTrace", e.StackTrace);
+			}
+			writer.EndObject();
+			return (writer.ToObject(), statusCode);
+		}
+
+		/// <summary>
+		/// Delete a individual refs key
+		/// </summary>
+		/// <param name="ns">Namespace. Each namespace is completely separated from each other. Use for different types of data that is never expected to be similar (between two different games for instance)</param>
+		/// <param name="bucket">The category/type of record you are caching. Is a clustered key together with the actual key, but all records in the same bucket can be dropped easily.</param>
+		/// <param name="key">The unique name of this particular key</param>
+		[HttpDelete("{ns}/{bucket}/{key}", Order = 500)]
+		[ProducesResponseType(200)]
+		[ProducesResponseType(404)]
+		public async Task<IActionResult> Delete(
+			[FromRoute][Required] NamespaceId ns,
+			[FromRoute][Required] BucketId bucket,
+			[FromRoute][Required] RefId key)
+		{
+			ActionResult? accessResult = await _requestHelper.HasAccessToNamespaceAsync(User, Request, ns, new[] { StorageAclAction.DeleteRefs });
+			if (accessResult != null)
+			{
+				return accessResult;
+			}
+
+			try
+			{
+				bool deleted = await _refService.DeleteAsync(ns, bucket, key);
+				if (!deleted)
+				{
+					return NotFound(new ProblemDetails { Title = $"Object {key} in bucket {bucket} and namespace {ns} did not exist" });
+				}
+				return Ok(new RefDeletedResponse(deleted ? 1 : 0));
+			}
+			catch (NamespaceNotFoundException e)
+			{
+				return NotFound(new ProblemDetails { Title = $"Namespace {e.Namespace} did not exist" });
+			}
+			catch (RefNotFoundException)
+			{
+				return NotFound(new ProblemDetails { Title = $"Object {key} in bucket {bucket} and namespace {ns} did not exist" });
+			}
+		}
+	}
+
+	public class RefDeletedResponse
+	{
+		public RefDeletedResponse()
+		{
+
+		}
+
+		public RefDeletedResponse(int deletedCount)
+		{
+			DeletedCount = deletedCount;
+		}
+
+		[CbField("deletedCount")]
+		public int DeletedCount { get; set; }
+	}
+
+	public class BucketDeletedResponse
+	{
+		public BucketDeletedResponse()
+		{
+
+		}
+
+		public BucketDeletedResponse(long countOfDeletedRecords)
+		{
+			CountOfDeletedRecords = countOfDeletedRecords;
+		}
+
+		[CbField("countOfDeletedRecords")]
+		public long CountOfDeletedRecords { get; set; }
+	}
+
+	public class RefBatchOps
+	{
+		public RefBatchOps()
+		{
+			Ops = Array.Empty<RefBatchOp>();
+		}
+
+		[System.Diagnostics.CodeAnalysis.SuppressMessage("Design", "CA1034:Nested types should not be visible", Justification = "For serialization only")]
+		public class RefBatchOp
+		{
+			public RefBatchOp()
+			{
+				Payload = null;
+				PayloadHash = null;
+			}
+
+			public enum RefOperation
+			{
+				INVALID,
+				GET,
+				PUT,
+				HEAD,
+			}
+
+			[Required]
+			[CbField("opId")]
+			public uint OpId { get; set; }
+
+			[CbField("op")]
+			[JsonIgnore]
+			[UsedImplicitly]
+			public string OpString
+			{
+				get => Op.ToString();
+				set => Op = Enum.Parse<RefOperation>(value);
+			}
+
+			[Required]
+			public RefOperation Op { get; set; } = RefOperation.INVALID;
+
+			[Required]
+			[CbField("bucket")]
+			public BucketId Bucket { get; set; }
+
+			[Required]
+			[CbField("key")]
+			public RefId Key { get; set; }
+
+			[CbField("resolveAttachments")]
+			public bool? ResolveAttachments { get; set; } = null;
+
+			[CbField("payload")]
+			public CbObject? Payload { get; set; } = null;
+
+			[CbField("payloadHash")]
+			public IoHash? PayloadHash { get; set; } = null;
+		}
+
+		[CbField("ops")]
+		public RefBatchOp[] Ops { get; set; }
+	}
+
+	public class BatchOpsResponse
+	{
+		public BatchOpsResponse()
+		{
+
+		}
+
+		[System.Diagnostics.CodeAnalysis.SuppressMessage("Design", "CA1034:Nested types should not be visible", Justification = "For serialization only")]
+		public class OpResponses
+		{
+			public OpResponses()
+			{
+
+			}
+			[CbField("opId")]
+			public uint OpId { get; set; }
+
+			[CbField("response")]
+			public CbObject Response { get; set; } = null!;
+
+			[CbField("statusCode")]
+			public int StatusCode { get; set; }
+		}
+
+		[CbField("results")]
+		[System.Diagnostics.CodeAnalysis.SuppressMessage("Usage", "CA2227:Collection properties should be read only", Justification = "Used by serialization")]
+		public List<OpResponses> Results { get; set; } = new List<OpResponses>();
+	}
+
+	public class RefMetadataResponse
+	{
+		public RefMetadataResponse()
+		{
+			InlinePayload = null!;
+		}
+
+		[JsonConstructor]
+		public RefMetadataResponse(NamespaceId ns, BucketId bucket, RefId name, BlobId payloadIdentifier, DateTime lastAccess, bool isFinalized, byte[]? inlinePayload)
+		{
+			Ns = ns;
+			Bucket = bucket;
+			Name = name;
+			PayloadIdentifier = payloadIdentifier;
+			LastAccess = lastAccess;
+			IsFinalized = isFinalized;
+			InlinePayload = inlinePayload;
+		}
+
+		public RefMetadataResponse(RefRecord objectRecord)
+		{
+			Ns = objectRecord.Namespace;
+			Bucket = objectRecord.Bucket;
+			Name = objectRecord.Name;
+			PayloadIdentifier = objectRecord.BlobIdentifier;
+			LastAccess = objectRecord.LastAccess;
+			IsFinalized = objectRecord.IsFinalized;
+			InlinePayload = objectRecord.InlinePayload;
+		}
+
+		[CbField("ns")]
+		public NamespaceId Ns { get; set; }
+
+		[CbField("bucket")]
+		public BucketId Bucket { get; set; }
+
+		[CbField("name")]
+		public RefId Name { get; set; }
+
+		[CbField("payloadIdentifier")]
+		public BlobId PayloadIdentifier { get; set; }
+
+		[CbField("lastAccess")]
+		public DateTime LastAccess { get; set; }
+
+		[CbField("isFinalized")]
+		public bool IsFinalized { get; set; }
+
+		[CbField("inlinePayload")]
+		public byte[]? InlinePayload { get; set; }
+	}
+
+	public class PutObjectResponse
+	{
+		public PutObjectResponse()
+		{
+			Needs = null!;
+		}
+
+		public PutObjectResponse(IoHash[] missingReferences)
+		{
+			Needs = missingReferences;
+		}
 
 		public PutObjectResponse(IEnumerable<ContentId> missingContentIds, IEnumerable<BlobId> missingBlobIds)
 			: this(Enumerable.Concat(missingContentIds.Select(x => x.Hash), missingBlobIds.Select(x => x.Hash)).ToArray())
 		{
 		}
 
-        [CbField("needs")]
-        public IoHash[] Needs { get; set; }
-    }
-
-    public class ExistCheckMultipleRefsResponse
-    {
-        public ExistCheckMultipleRefsResponse(List<(BucketId, RefId)> missing)
-        {
-            Missing = missing.Select(pair =>
-            {
-                (BucketId bucketId, RefId ioHashKey) = pair;
-                return new MissingReference()
-                {
-                    Bucket = bucketId,
-                    Key = ioHashKey,
-                };
-            }).ToList();
-        }
-
-        [JsonConstructor]
-        public ExistCheckMultipleRefsResponse(List<MissingReference> missing)
-        {
-            Missing = missing;
-        }
-
-        [CbField("missing")]
-        [System.Diagnostics.CodeAnalysis.SuppressMessage("Usage", "CA2227:Collection properties should be read only", Justification = "Used by serialization")]
-        public List<MissingReference> Missing { get; set; }
-
-        [System.Diagnostics.CodeAnalysis.SuppressMessage("Design", "CA1034:Nested types should not be visible", Justification = "For serialization only")]
-        public class MissingReference
-        {
-            [CbField("bucket")]
-            public BucketId Bucket { get; set; }
-
-            [CbField("key")]
-            public RefId Key { get; set; }
-        }
-    }
-
-    public class GetNamespacesResponse
-    {
-        [JsonConstructor]
-        public GetNamespacesResponse(NamespaceId[] namespaces)
-        {
-            Namespaces = namespaces;
-        }
-
-        public NamespaceId[] Namespaces { get; set; }
-    }
-
-    public class HashMismatchException : Exception
-    {
-        public IoHash SuppliedHash { get; }
-        public IoHash ContentHash { get; }
-
-        public HashMismatchException(IoHash suppliedHash, IoHash contentHash) : base($"ID was not a hash of the content uploaded. Supplied hash was: {suppliedHash} but hash of content was {contentHash}")
-        {
-            SuppliedHash = suppliedHash;
-            ContentHash = contentHash;
-        }
-    }
+		[CbField("needs")]
+		public IoHash[] Needs { get; set; }
+	}
+
+	public class ExistCheckMultipleRefsResponse
+	{
+		public ExistCheckMultipleRefsResponse(List<(BucketId, RefId)> missing)
+		{
+			Missing = missing.Select(pair =>
+			{
+				(BucketId bucketId, RefId ioHashKey) = pair;
+				return new MissingReference()
+				{
+					Bucket = bucketId,
+					Key = ioHashKey,
+				};
+			}).ToList();
+		}
+
+		[JsonConstructor]
+		public ExistCheckMultipleRefsResponse(List<MissingReference> missing)
+		{
+			Missing = missing;
+		}
+
+		[CbField("missing")]
+		[System.Diagnostics.CodeAnalysis.SuppressMessage("Usage", "CA2227:Collection properties should be read only", Justification = "Used by serialization")]
+		public List<MissingReference> Missing { get; set; }
+
+		[System.Diagnostics.CodeAnalysis.SuppressMessage("Design", "CA1034:Nested types should not be visible", Justification = "For serialization only")]
+		public class MissingReference
+		{
+			[CbField("bucket")]
+			public BucketId Bucket { get; set; }
+
+			[CbField("key")]
+			public RefId Key { get; set; }
+		}
+	}
+
+	public class GetNamespacesResponse
+	{
+		[JsonConstructor]
+		public GetNamespacesResponse(NamespaceId[] namespaces)
+		{
+			Namespaces = namespaces;
+		}
+
+		public NamespaceId[] Namespaces { get; set; }
+	}
+
+	public class HashMismatchException : Exception
+	{
+		public IoHash SuppliedHash { get; }
+		public IoHash ContentHash { get; }
+
+		public HashMismatchException(IoHash suppliedHash, IoHash contentHash) : base($"ID was not a hash of the content uploaded. Supplied hash was: {suppliedHash} but hash of content was {contentHash}")
+		{
+			SuppliedHash = suppliedHash;
+			ContentHash = contentHash;
+		}
+	}
 }
