Horde: Adding images to help artifact migration
#rnx
#skipci

diff --git a/Engine/Source/Programs/Horde/HordeDashboard/src/components/jobDetailsV2/JobDetailArtifactsV2.tsx b/Engine/Source/Programs/Horde/HordeDashboard/src/components/jobDetailsV2/JobDetailArtifactsV2.tsx
--- a/Engine/Source/Programs/Horde/HordeDashboard/src/components/jobDetailsV2/JobDetailArtifactsV2.tsx
+++ b/Engine/Source/Programs/Horde/HordeDashboard/src/components/jobDetailsV2/JobDetailArtifactsV2.tsx
@@ -9,6 +9,7 @@ import { observable, action, makeObservable } from 'mobx';
 import backend from '../../backend';
 import { JobDataView, JobDetailsV2 } from './JobDetailsViewCommon';
 import { ISideRailLink } from '../../base/components/SideRail';
+import dashboard from '../../backend/Dashboard';
 
 const sideRail: ISideRailLink = { text: "Artifacts", url: "rail_artifacts" };
 
@@ -226,18 +227,26 @@ export const JobDetailArtifactsV2: React.FC<{ jobDetails: JobDetailsV2; stepId:
    dataView.set(stepId);
 
    if (jobDetails.jobData?.useArtifactsV2) {
+
+      const imgSrc = dashboard.darktheme ? "/images/notifications/artifacts_moved_dark.png" : "/images/notifications/artifacts_moved_light.png";
+
       return <Stack id={sideRail.url} styles={{ root: { paddingTop: 18, paddingRight: 12 } }}>
          <Stack className={hordeClasses.raised}>
             <Stack tokens={{ childrenGap: 12 }}>
                <Stack horizontalAlign="space-between" styles={{ root: { minHeight: 32 } }}>
                   <Text variant="mediumPlus" styles={{ root: { fontFamily: "Horde Open Sans SemiBold" } }}>Artifacts</Text>
                   <Stack tokens={{ childrenGap: 24 }}>
-                     <Stack horizontal verticalAlign='center' style={{ paddingTop: 24, paddingLeft: 12, paddingBottom: 52 }}>
+                     <Stack horizontal verticalAlign='center' style={{ paddingTop: 24, paddingLeft: 12, paddingBottom: 12 }}>
                         <FontIcon style={{ paddingRight: 12 }} iconName="Error" />
                         <Stack>
                            <Text variant='medium'>This job uses the new Horde artifact storage backend.  These artifacts are now accessible in the upper right, underneath the breadcrumb area.</Text>
                         </Stack>
                      </Stack>
+
+                     <Stack style={{paddingLeft: 38}}>
+                        <img style={{width: "fit-content"}} src={imgSrc} alt="" />
+                     </Stack>
+
                   </Stack>
                </Stack>
             </Stack>
