Horde: Fix typo in issue confirmation dialog
#rnx
#skipci

diff --git a/Engine/Source/Programs/Horde/HordeDashboard/src/components/IssueViewV2.tsx b/Engine/Source/Programs/Horde/HordeDashboard/src/components/IssueViewV2.tsx
--- a/Engine/Source/Programs/Horde/HordeDashboard/src/components/IssueViewV2.tsx
+++ b/Engine/Source/Programs/Horde/HordeDashboard/src/components/IssueViewV2.tsx
@@ -2149,7 +2149,7 @@ const AckIssueModal: React.FC<{ onClose: () => void }> = ({ onClose }) => {
       <Stack styles={{ root: { padding: 8 } }}>
          <Stack horizontal tokens={{ childrenGap: 16 }} styles={{ root: { paddingTop: 12, paddingLeft: 8, paddingBottom: 8 } }}>
             <Stack grow />
-            <PrimaryButton text="Acknoweldge" disabled={state.submitting ?? false} onClick={() => { onAck(); }} />
+            <PrimaryButton text="Acknowledge" disabled={state.submitting ?? false} onClick={() => { onAck(); }} />
             <DefaultButton text="Cancel" disabled={false} onClick={() => { close(); }} />
          </Stack>
       </Stack>
