Horde: Add stream information to pool step/batch panels and fix some layout issues
#skipci
#rnx

diff --git a/Engine/Source/Programs/Horde/HordeDashboard/src/components/PoolView.tsx b/Engine/Source/Programs/Horde/HordeDashboard/src/components/PoolView.tsx
--- a/Engine/Source/Programs/Horde/HordeDashboard/src/components/PoolView.tsx
+++ b/Engine/Source/Programs/Horde/HordeDashboard/src/components/PoolView.tsx
@@ -84,8 +84,6 @@ class PoolHandler extends PollBase {
 
       try {
 
-         const doTimings = !!this.pool;
-
          const pool = this.pool = await backend.getPool(this.poolId);
 
          // @todo: optimize using modifiedAfter
@@ -267,41 +265,6 @@ class PoolHandler extends PollBase {
             });
          }
 
-         if (doTimings) {
-
-            let jobTiming = activeJobs.filter(job => !this.jobTiming.get(job.id));
-
-            while (jobTiming.length) {
-
-               const batch = jobTiming.slice(0, 5);
-
-               const requests = batch.map(job => backend.getJobTiming(job.id));
-
-               await Promise.all(requests).then((responses) => {
-
-                  responses.forEach(r => {
-                     this.jobTiming.set(r.jobResponse.id, r);
-                  });
-
-               }).catch((errors) => {
-                  console.log(errors);
-                  // eslint-disable-next-line
-               }).finally(() => {
-
-                  jobTiming = jobTiming.slice(5);
-               });
-            }
-
-         }
-
-         activeJobs.forEach((job) => {
-            const timing = this.jobTiming.get(job.id);
-            if (timing) {
-               const steps = (job.batches?.map(b => b.steps).flat() ?? []) as StepData[];
-               steps.forEach(step => step.timing = timing.steps[step.id] ? timing.steps[step.id] : undefined)
-            }
-         });
-
          this.setUpdated();
 
       } catch (err) {
@@ -342,7 +305,7 @@ enum StepState {
    Previous = "Previous"
 }
 
-const PoolAgentPanel: React.FC<{poolId:string}> = ({poolId}) => {
+const PoolAgentPanel: React.FC<{ poolId: string }> = ({ poolId }) => {
 
    return <Stack styles={{ root: { paddingTop: 18, paddingLeft: 12, paddingRight: 12 } }} >
       <Stack tokens={{ childrenGap: 12 }}>
@@ -410,12 +373,11 @@ const StepPanel: React.FC<{ stepState: StepState }> = ({ stepState }) => {
    let columns: IColumn[] = [];
 
    columns = [
+      { key: 'column1', name: 'Job', minWidth: 790, maxWidth: 790 },
       { key: 'column2', name: 'Agent', minWidth: 100, maxWidth: 100 },
-      { key: 'column1', name: 'Job', minWidth: 600, maxWidth: 600 },
       { key: 'column3', name: 'Time Active', minWidth: 100, maxWidth: 100 },
       { key: 'column4', name: 'Start Time', minWidth: 100, maxWidth: 100 },
-      { key: 'column5', name: 'Estimated Finish', minWidth: 100, maxWidth: 100 },
-      { key: 'column6', name: 'View Log', minWidth: 80, maxWidth: 80 }
+      { key: 'column5', name: 'View Log', minWidth: 80, maxWidth: 80 }
    ];
 
    const agentItems: AgentItem[] = agents.map(a => {
@@ -464,47 +426,12 @@ const StepPanel: React.FC<{ stepState: StepState }> = ({ stepState }) => {
          }
       }
 
-      if (column.name === 'Estimated Finish') {
-
-         let eta = {
-            display: "",
-            server: ""
-         };
-
-         let finished = { display: "", server: "" };
-
-         eta = getStepETA(step, job);
-
-         finished = getStepFinishTime(step);
-
-         if (finished.display) {
-            eta.display = finished.display;
-            eta.server = finished.server;
-         }
-
-         let time = eta.display;
-
-         const etaColor = dashboard.darktheme ? "#A9A9A9" : "#999999";
-
-         const color = !step.finishTime ? etaColor : undefined;
-
-         // Open Sans tilde rendering issue at 13px, and not rendering at all at other px: [REDACTED_PROFILE_URL]/fonts/issues/399, do not change from 12px
-         return <Stack horizontalAlign={"end"}>
-            <Stack horizontal tokens={{ childrenGap: 2 }}>
-               {!!time && !step.finishTime && <Text style={{ color: color, fontSize: "11px", paddingTop: 2 }}>~</Text>}
-               <Text style={{ color: color, fontSize: "13px" }}>
-                  {time}
-               </Text>
-            </Stack>
-         </Stack>;
-
-      }
 
       if (column.name === "View Log") {
          if (stepState === StepState.Pending) {
             return null;
          }
-         return <Stack style={{ paddingRight: 32 }}>
+         return <Stack style={{ paddingRight: 8 }}>
             <Link to={`/log/${step.logId}`} target="_blank">
                <Stack horizontal horizontalAlign={"end"} verticalAlign="center" tokens={{ childrenGap: 0, padding: 0 }} style={{ width: "100%", height: "100%" }}>
                   <Text styles={{ root: { margin: '0px', padding: '0px', paddingRight: '8px' } }} className={"view-log-link"}>View Log</Text>
@@ -519,22 +446,28 @@ const StepPanel: React.FC<{ stepState: StepState }> = ({ stepState }) => {
          if (jobName.indexOf("- Kicked By") !== -1) {
             jobName = jobName.split("- Kicked By")[0];
          }
-         const stepName = `${jobName} - ${node.name}`;
+         let stepName = `${jobName} - ${node.name}`;
          const stepUrl = `/job/${job.id}?step=${step.id}`;
 
+         const stream = projectStore.streamById(job.streamId);
+         if (stream) {
+            stepName = `${stream.fullname ?? stream.id} - ${stepName}`
+         }
+
          return <Stack>
-            <Link target="_blank" to={stepUrl}><Text>{stepName}</Text></Link>
+            <Link target="_blank" to={stepUrl}>
+               <Stack horizontal>
+                  <StepStatusIcon step={step} />
+                  <Text>{stepName}</Text>
+               </Stack>
+            </Link>
          </Stack>;
 
       }
 
       if (column.name === "Agent") {
          return <div style={{ cursor: "pointer" }} onClick={() => { setLastSelectedAgent(agent.id) }}>
-            <Stack horizontal>
-               <StepStatusIcon step={step} />
-               <Text>{agent.name}</Text>
-            </Stack>
-
+            <Text>{agent.name}</Text>
          </div>
 
       }
@@ -611,9 +544,9 @@ const BatchPanel: React.FC = () => {
    let columns: IColumn[] = [];
 
    columns = [
-      { key: 'column1', name: 'Job', minWidth: 600, maxWidth: 600 },
-      { key: 'column4', name: 'Job Created', minWidth: 300, maxWidth: 300 },
-      { key: 'column2', name: 'Status', minWidth: 200, maxWidth: 200 },
+      { key: 'column1', name: 'Job', minWidth: 800, maxWidth: 800 },
+      { key: 'column2', name: 'Job Created', minWidth: 200, maxWidth: 200 },
+      { key: 'column3', name: 'Status', minWidth: 200, maxWidth: 200 },
    ];
 
    const batchItems: BatchItem[] = batches.map(b => {
@@ -628,6 +561,9 @@ const BatchPanel: React.FC = () => {
       const job = item.pending.job;
       const batch = item.pending.batch;
 
+      const stream = projectStore.streamById(job.streamId);
+
+
       if (column.name === "Status") {
          let statusText = "";
 
@@ -644,7 +580,7 @@ const BatchPanel: React.FC = () => {
 
       if (column.name === "Job Created") {
 
-         return <Stack ><Text style={{ fontSize: "13px", paddingRight: 12 }}>{getNiceTime(job.createTime, false)}</Text></Stack>;
+         return <Stack ><Text style={{ fontSize: "13px", paddingRight: 12 }}>{getShortNiceTime(job.createTime, true, true)}</Text></Stack>;
       }
 
 
@@ -659,13 +595,18 @@ const BatchPanel: React.FC = () => {
             const group = job.graphRef!.groups![batch.groupIdx];
             const node = group.nodes[step.nodeIdx];
             jobName += ` - ${node.name}`;
+            if (stream) {
+               jobName = `${stream.fullname ?? stream.id} - ` + jobName;
+            }
          }
 
          const jobUrl = `/job/${job.id}`;
 
          return <Stack horizontal>
             {!!step && <StepStatusIcon step={step} />}
-            <Link target="_blank" to={jobUrl}><Text>{jobName}</Text></Link>
+            <Link target="_blank" to={jobUrl}>
+               <Text>{jobName}</Text>
+            </Link>
          </Stack>;
 
       }
