Horde: Enable static analysis warnings about having async methods without an 'Async' suffix.

diff --git a/Engine/Source/Programs/Horde/.editorconfig b/Engine/Source/Programs/Horde/.editorconfig
--- a/Engine/Source/Programs/Horde/.editorconfig
+++ b/Engine/Source/Programs/Horde/.editorconfig
@@ -256,6 +256,19 @@ csharp_prefer_braces = true:silent
 csharp_preserve_single_line_blocks = true
 csharp_preserve_single_line_statements = true
 
+# Async methods should have "Async" suffix
+dotnet_naming_rule.async_methods_end_in_async.symbols = any_async_methods
+dotnet_naming_rule.async_methods_end_in_async.style = end_in_async
+dotnet_naming_rule.async_methods_end_in_async.severity = warning
+
+dotnet_naming_symbols.any_async_methods.applicable_kinds = method
+dotnet_naming_symbols.any_async_methods.applicable_accessibilities = *
+dotnet_naming_symbols.any_async_methods.required_modifiers = async
+
+dotnet_naming_style.end_in_async.required_prefix = 
+dotnet_naming_style.end_in_async.required_suffix = Async
+dotnet_naming_style.end_in_async.capitalization = pascal_case
+dotnet_naming_style.end_in_async.word_separator = 
 
 # All Rules Enabled as build warnings
 # Description: All Rules are enabled as build warnings. Rules with IsEnabledByDefault = false are force enabled as build warnings.
diff --git a/Engine/Source/Programs/Horde/Horde.Agent.TrayApp/Program.cs b/Engine/Source/Programs/Horde/Horde.Agent.TrayApp/Program.cs
--- a/Engine/Source/Programs/Horde/Horde.Agent.TrayApp/Program.cs
+++ b/Engine/Source/Programs/Horde/Horde.Agent.TrayApp/Program.cs
@@ -145,7 +145,7 @@ namespace Horde.Agent.TrayApp
 			};
 			_trayIcon.Click += TrayIcon_Click;
 
-			_clientTask = BackgroundTask.StartNew(StatusTask);
+			_clientTask = BackgroundTask.StartNew(StatusTaskAsync);
 			_tickPauseStateTask = BackgroundTask.StartNew(ctx => TickPauseStateAsync(ctx));
 			_waitForExitTask = BackgroundTask.StartNew(ctx => WaitForExitAsync(eventHandle, ctx));
 		}
@@ -272,7 +272,7 @@ namespace Horde.Agent.TrayApp
 			}
 		}
 
-		async Task StatusTask(CancellationToken cancellationToken)
+		async Task StatusTaskAsync(CancellationToken cancellationToken)
 		{
 			SetStatus(AgentStatusMessage.Starting);
 			for (; ; )
diff --git a/Engine/Source/Programs/Horde/Horde.Server/Ddc/BlobService.cs b/Engine/Source/Programs/Horde/Horde.Server/Ddc/BlobService.cs
--- a/Engine/Source/Programs/Horde/Horde.Server/Ddc/BlobService.cs
+++ b/Engine/Source/Programs/Horde/Horde.Server/Ddc/BlobService.cs
@@ -74,7 +74,7 @@ namespace Horde.Server.Ddc
 			List<BlobId> unknownBlobIds = new List<BlobId>();
 			foreach (BlobId blobId in blobIds)
 			{
-				if (!await storageClient.FindAliasAsync(GetAlias(blobId), cancellationToken).AnyAsync())
+				if (!await storageClient.FindAliasAsync(GetAlias(blobId), cancellationToken).AnyAsync(cancellationToken))
 				{
 					unknownBlobIds.Add(blobId);
 				}
diff --git a/Engine/Source/Programs/Horde/Horde.Server/Ddc/BufferedPayload.cs b/Engine/Source/Programs/Horde/Horde.Server/Ddc/BufferedPayload.cs
--- a/Engine/Source/Programs/Horde/Horde.Server/Ddc/BufferedPayload.cs
+++ b/Engine/Source/Programs/Horde/Horde.Server/Ddc/BufferedPayload.cs
@@ -67,12 +67,12 @@ namespace Horde.Server.Ddc
 		/// <summary>
 		/// Create a buffered payload from a stream
 		/// </summary>
-		public static async Task<MemoryBufferedPayload> Create(Tracer tracer, Stream s)
+		public static async Task<MemoryBufferedPayload> CreateAsync(Tracer tracer, Stream s)
 		{
 			using TelemetrySpan scope = tracer.StartActiveSpan("payload.buffer")
 				.SetAttribute("operation.name", "payload.buffer")
 				.SetAttribute("bufferType", "Memory");
-			MemoryBufferedPayload payload = new MemoryBufferedPayload(await s.ToByteArray());
+			MemoryBufferedPayload payload = new MemoryBufferedPayload(await s.ToByteArrayAsync());
 			return payload;
 		}
 
@@ -96,7 +96,7 @@ namespace Horde.Server.Ddc
 		/// <summary>
 		/// Create a new payload instance backed by the filesystem
 		/// </summary>
-		public static async Task<FilesystemBufferedPayload> Create(Tracer tracer, Stream s)
+		public static async Task<FilesystemBufferedPayload> CreateAsync(Tracer tracer, Stream s)
 		{
 			FileInfo tempFile = new FileInfo(Path.GetTempFileName());
 
@@ -168,21 +168,21 @@ namespace Horde.Server.Ddc
 				throw new Exception("Expected content-length on all requests");
 			}
 
-			return CreateFromStream(request.Body, contentLength.Value);
+			return CreateFromStreamAsync(request.Body, contentLength.Value);
 		}
 
 		/// <summary>
 		/// Create a new buffered payload instance from a stream
 		/// </summary>
-		public async Task<BufferedPayload> CreateFromStream(Stream s, long contentLength)
+		public async Task<BufferedPayload> CreateFromStreamAsync(Stream s, long contentLength)
 		{
 			// blob is small enough to fit into memory we just read it as is
 			if (contentLength < _options.CurrentValue.MemoryBufferSize)
 			{
-				return await MemoryBufferedPayload.Create(_tracer, s);
+				return await MemoryBufferedPayload.CreateAsync(_tracer, s);
 			}
 
-			return await FilesystemBufferedPayload.Create(_tracer, s);
+			return await FilesystemBufferedPayload.CreateAsync(_tracer, s);
 		}
 	}
 }
diff --git a/Engine/Source/Programs/Horde/Horde.Server/Ddc/IBlobService.cs b/Engine/Source/Programs/Horde/Horde.Server/Ddc/IBlobService.cs
--- a/Engine/Source/Programs/Horde/Horde.Server/Ddc/IBlobService.cs
+++ b/Engine/Source/Programs/Horde/Horde.Server/Ddc/IBlobService.cs
@@ -106,7 +106,7 @@ namespace Horde.Server.Ddc
 			// decompress the content and generate a identifier from it to verify the identifier we got
 			await using Stream decompressStream = payload.GetStream();
 			// TODO: we should add a overload for decompress content that can work on streams, otherwise we are still limited to 2GB compressed blobs
-			byte[] decompressedContent = compressedBufferUtils.DecompressContent(await decompressStream.ToByteArray());
+			byte[] decompressedContent = compressedBufferUtils.DecompressContent(await decompressStream.ToByteArrayAsync());
 
 			IoHash decompressedHash;
 			using (TelemetrySpan _ = tracer.StartActiveSpan("web.hash").SetAttribute("operation.name", "web.hash"))
diff --git a/Engine/Source/Programs/Horde/Horde.Server/Ddc/IReferenceResolver.cs b/Engine/Source/Programs/Horde/Horde.Server/Ddc/IReferenceResolver.cs
--- a/Engine/Source/Programs/Horde/Horde.Server/Ddc/IReferenceResolver.cs
+++ b/Engine/Source/Programs/Horde/Horde.Server/Ddc/IReferenceResolver.cs
@@ -124,13 +124,13 @@ namespace Horde.Server.Ddc
 
 						if (field.IsBinaryAttachment())
 						{
-							Task<(ContentId, BlobId[]?)> resolveContentId = ResolveContentId(ns, contentId);
+							Task<(ContentId, BlobId[]?)> resolveContentId = ResolveContentIdAsync(ns, contentId);
 							pendingContentIdResolves.Add(resolveContentId);
 						}
 						else if (field.IsObjectAttachment())
 						{
 							attachments.Add(new ObjectAttachment(blobIdentifier));
-							pendingCompactBinaryAttachments.Add(ParseCompactBinaryAttachment(ns, blobIdentifier));
+							pendingCompactBinaryAttachments.Add(ParseCompactBinaryAttachmentAsync(ns, blobIdentifier));
 						}
 						else
 						{
@@ -246,17 +246,17 @@ namespace Horde.Server.Ddc
 			{
 				if (attachment is BlobAttachment blobAttachment)
 				{
-					pendingBlobExistsChecks.Add(CheckBlobExists(ns, blobAttachment.Identifier));
+					pendingBlobExistsChecks.Add(CheckBlobExistsAsync(ns, blobAttachment.Identifier));
 				}
 				else if (attachment is ContentIdAttachment contentIdAttachment)
 				{
 					// If we find a content id we resolve that into the actual blobs it references
-					pendingContentIdChecks.Add(CheckContentIdExists(ns, contentIdAttachment));
+					pendingContentIdChecks.Add(CheckContentIdExistsAsync(ns, contentIdAttachment));
 				}
 				else if (attachment is ObjectAttachment objectAttachment)
 				{
 					// a object just references the same blob, traversing the object attachment is done in GetAttachments
-					pendingBlobExistsChecks.Add(CheckBlobExists(ns, objectAttachment.Identifier));
+					pendingBlobExistsChecks.Add(CheckBlobExistsAsync(ns, objectAttachment.Identifier));
 				}
 				else
 				{
@@ -306,12 +306,12 @@ namespace Horde.Server.Ddc
 			}
 		}
 
-		private async Task<(ContentIdAttachment, bool)> CheckContentIdExists(NamespaceId ns, ContentIdAttachment contentIdAttachment)
+		private async Task<(ContentIdAttachment, bool)> CheckContentIdExistsAsync(NamespaceId ns, ContentIdAttachment contentIdAttachment)
 		{
 			bool allBlobsExist = true;
 			foreach (BlobId b in contentIdAttachment.ReferencedBlobs)
 			{
-				(BlobId _, bool exists) = await CheckBlobExists(ns, b);
+				(BlobId _, bool exists) = await CheckBlobExistsAsync(ns, b);
 
 				if (!exists)
 				{
@@ -323,16 +323,16 @@ namespace Horde.Server.Ddc
 			return (contentIdAttachment, allBlobsExist);
 		}
 
-		private async Task<CbObject> ParseCompactBinaryAttachment(NamespaceId ns, BlobId blobIdentifier)
+		private async Task<CbObject> ParseCompactBinaryAttachmentAsync(NamespaceId ns, BlobId blobIdentifier)
 		{
 			BlobContents contents = await _blobStore.GetObjectAsync(ns, blobIdentifier);
-			byte[] data = await contents.Stream.ToByteArray();
+			byte[] data = await contents.Stream.ToByteArrayAsync();
 			CbObject childBinaryObject = new CbObject(data);
 
 			return childBinaryObject;
 		}
 
-		private async Task<(ContentId, BlobId[]?)> ResolveContentId(NamespaceId ns, ContentId contentId)
+		private async Task<(ContentId, BlobId[]?)> ResolveContentIdAsync(NamespaceId ns, ContentId contentId)
 		{
 			using TelemetrySpan scope = _tracer.StartActiveSpan("ReferenceResolver.ResolveContentId")
 				.SetAttribute("operation.name", "ReferenceResolver.ResolveContentId")
@@ -341,7 +341,7 @@ namespace Horde.Server.Ddc
 			return (contentId, resolvedBlobs);
 		}
 
-		private async Task<(BlobId, bool)> CheckBlobExists(NamespaceId ns, BlobId blob)
+		private async Task<(BlobId, bool)> CheckBlobExistsAsync(NamespaceId ns, BlobId blob)
 		{
 			return (blob, await _blobStore.ExistsAsync(ns, blob));
 		}
@@ -351,7 +351,7 @@ namespace Horde.Server.Ddc
 	{
 		public List<ContentId> UnresolvedReferences { get; }
 
-		public PartialReferenceResolveException(List<ContentId> unresolvedReferences) : base($"References missing: {string.Join(',', unresolvedReferences)}")
+		public PartialReferenceResolveException(List<ContentId> unresolvedReferences) : base($"References missing: {String.Join(',', unresolvedReferences)}")
 		{
 			UnresolvedReferences = unresolvedReferences;
 		}
@@ -361,7 +361,7 @@ namespace Horde.Server.Ddc
 	{
 		public List<BlobId> MissingBlobs { get; }
 
-		public ReferenceIsMissingBlobsException(List<BlobId> missingBlobs) : base($"References is missing these blobs: {string.Join(',', missingBlobs)}")
+		public ReferenceIsMissingBlobsException(List<BlobId> missingBlobs) : base($"References is missing these blobs: {String.Join(',', missingBlobs)}")
 		{
 			MissingBlobs = missingBlobs;
 		}
diff --git a/Engine/Source/Programs/Horde/Horde.Server/Ddc/ObjectsController.cs b/Engine/Source/Programs/Horde/Horde.Server/Ddc/ObjectsController.cs
--- a/Engine/Source/Programs/Horde/Horde.Server/Ddc/ObjectsController.cs
+++ b/Engine/Source/Programs/Horde/Horde.Server/Ddc/ObjectsController.cs
@@ -51,7 +51,7 @@ namespace Horde.Server.Ddc
 
 		[HttpGet("{ns}/{id}")]
 		[ProducesDefaultResponseType]
-		public async Task<IActionResult> Get(
+		public async Task<IActionResult> GetAsync(
 			[Required] NamespaceId ns,
 			[Required] BlobId id)
 		{
@@ -75,7 +75,7 @@ namespace Horde.Server.Ddc
 
 		[HttpHead("{ns}/{id}")]
 		[ProducesDefaultResponseType]
-		public async Task<IActionResult> Head(
+		public async Task<IActionResult> HeadAsync(
 			[Required] NamespaceId ns,
 			[Required] BlobId id)
 		{
@@ -97,7 +97,7 @@ namespace Horde.Server.Ddc
 
 		[HttpPost("{ns}/exists")]
 		[ProducesDefaultResponseType]
-		public async Task<IActionResult> ExistsMultiple(
+		public async Task<IActionResult> ExistsMultipleAsync(
 			[Required] NamespaceId ns,
 			[Required][FromQuery] List<BlobId> id)
 		{
@@ -123,7 +123,7 @@ namespace Horde.Server.Ddc
 
 		[HttpPost("{ns}/exist")]
 		[ProducesDefaultResponseType]
-		public async Task<IActionResult> ExistsBody(
+		public async Task<IActionResult> ExistsBodyAsync(
 			[Required] NamespaceId ns,
 			[FromBody] BlobId[] bodyIds)
 		{
@@ -149,7 +149,7 @@ namespace Horde.Server.Ddc
 
 		[HttpPut("{ns}/{id}")]
 		[RequiredContentType(CustomMediaTypeNames.UnrealCompactBinary)]
-		public async Task<IActionResult> Put(
+		public async Task<IActionResult> PutAsync(
 			[Required] NamespaceId ns,
 			[Required] BlobId id,
 			CancellationToken cancellationToken)
@@ -175,7 +175,7 @@ namespace Horde.Server.Ddc
 		}
 
 		[HttpGet("{ns}/{id}/references")]
-		public async Task<IActionResult> ResolveReferences(
+		public async Task<IActionResult> ResolveReferencesAsync(
 			[Required] NamespaceId ns,
 			[Required] BlobId id)
 		{
@@ -195,7 +195,7 @@ namespace Horde.Server.Ddc
 				return NotFound(new ValidationProblemDetails { Title = $"Object {e.Blob} not found" });
 			}
 
-			byte[] blobContents = await blob.Stream.ToByteArray();
+			byte[] blobContents = await blob.Stream.ToByteArrayAsync();
 			if (blobContents.Length == 0)
 			{
 				_logger.LogWarning("0 byte object found for {Id} {Namespace}", id, ns);
@@ -218,11 +218,11 @@ namespace Horde.Server.Ddc
 			}
 			catch (PartialReferenceResolveException e)
 			{
-				return BadRequest(new ValidationProblemDetails { Title = $"Object {id} is missing content ids", Detail = $"Following content ids are invalid: {string.Join(",", e.UnresolvedReferences)}" });
+				return BadRequest(new ValidationProblemDetails { Title = $"Object {id} is missing content ids", Detail = $"Following content ids are invalid: {String.Join(",", e.UnresolvedReferences)}" });
 			}
 			catch (ReferenceIsMissingBlobsException e)
 			{
-				return BadRequest(new ValidationProblemDetails { Title = $"Object {id} is missing blobs", Detail = $"Following blobs are missing: {string.Join(",", e.MissingBlobs)}" });
+				return BadRequest(new ValidationProblemDetails { Title = $"Object {id} is missing blobs", Detail = $"Following blobs are missing: {String.Join(",", e.MissingBlobs)}" });
 			}
 		}
 	}
diff --git a/Engine/Source/Programs/Horde/Horde.Server/Ddc/RefsController.cs b/Engine/Source/Programs/Horde/Horde.Server/Ddc/RefsController.cs
--- a/Engine/Source/Programs/Horde/Horde.Server/Ddc/RefsController.cs
+++ b/Engine/Source/Programs/Horde/Horde.Server/Ddc/RefsController.cs
@@ -76,7 +76,7 @@ namespace Horde.Server.Ddc
 		/// <param name="format">Optional specifier to set which output format is used json/raw/cb</param>
 		[HttpGet("{ns}/{bucket}/{key}.{format?}", Order = 500)]
 		[Produces(MediaTypeNames.Application.Json, MediaTypeNames.Application.Octet, CustomMediaTypeNames.UnrealCompactBinary, CustomMediaTypeNames.JupiterInlinedPayload, CustomMediaTypeNames.UnrealCompactBinaryPackage)]
-		public async Task<IActionResult> Get(
+		public async Task<IActionResult> GetAsync(
 			[FromRoute][Required] NamespaceId ns,
 			[FromRoute][Required] BucketId bucket,
 			[FromRoute][Required] RefId key,
@@ -145,7 +145,7 @@ namespace Horde.Server.Ddc
 						}
 					case MediaTypeNames.Application.Octet:
 						{
-							byte[] blobMemory = await blob.Stream.ToByteArray();
+							byte[] blobMemory = await blob.Stream.ToByteArrayAsync();
 							CbObject cb = new CbObject(blobMemory);
 
 							(int, CbField?) CountFields(CbObject o)
@@ -194,7 +194,7 @@ namespace Horde.Server.Ddc
 							byte[] blobMemory;
 							{
 								using TelemetrySpan scope = _tracer.StartActiveSpan("json.readblob").SetAttribute("operation.name", "json.readblob");
-								blobMemory = await blob.Stream.ToByteArray();
+								blobMemory = await blob.Stream.ToByteArrayAsync();
 							}
 							CbObject cb = new CbObject(blobMemory);
 							string s = cb.ToJson();
@@ -206,7 +206,7 @@ namespace Horde.Server.Ddc
 					case CustomMediaTypeNames.UnrealCompactBinaryPackage:
 						{
 							using TelemetrySpan packageScope = _tracer.StartActiveSpan("cbpackage.fetch").SetAttribute("operation.name", "cbpackage.fetch");
-							byte[] blobMemory = await blob.Stream.ToByteArray();
+							byte[] blobMemory = await blob.Stream.ToByteArrayAsync();
 							CbObject cb = new CbObject(blobMemory);
 
 							IAsyncEnumerable<Attachment> attachments = _referenceResolver.GetAttachmentsAsync(ns, cb);
@@ -225,13 +225,13 @@ namespace Horde.Server.Ddc
 									if (attachment is BlobAttachment blobAttachment)
 									{
 										BlobId referencedBlob = blobAttachment.Identifier;
-										attachmentContents = await _blobStore.GetObjectAsync(ns, referencedBlob);
+										attachmentContents = await _blobStore.GetObjectAsync(ns, referencedBlob, cancellationToken: token);
 									}
 									else if (attachment is ObjectAttachment objectAttachment)
 									{
 										flags |= CbPackageAttachmentFlags.IsObject;
 										BlobId referencedBlob = objectAttachment.Identifier;
-										attachmentContents = await _blobStore.GetObjectAsync(ns, referencedBlob);
+										attachmentContents = await _blobStore.GetObjectAsync(ns, referencedBlob, cancellationToken: token);
 									}
 									else if (attachment is ContentIdAttachment contentIdAttachment)
 									{
@@ -276,7 +276,7 @@ namespace Horde.Server.Ddc
 						}
 					case CustomMediaTypeNames.JupiterInlinedPayload:
 						{
-							byte[] blobMemory = await blob.Stream.ToByteArray();
+							byte[] blobMemory = await blob.Stream.ToByteArrayAsync();
 							CbObject cb = new CbObject(blobMemory);
 
 							static (int, int) CountFields(CbObject o)
@@ -397,8 +397,6 @@ namespace Horde.Server.Ddc
 			}
 		}
 
-
-
 		/// <summary>
 		/// Returns the metadata about a ref key
 		/// </summary>
@@ -407,7 +405,7 @@ namespace Horde.Server.Ddc
 		/// <param name="key">The unique name of this particular key. `iAmAVeryValidKey`</param>
 		/// <param name="fields">The fields to include in the response, omit this to include everything.</param>
 		[HttpGet("{ns}/{bucket}/{key}/metadata", Order = 500)]
-		public async Task<IActionResult> GetMetadata(
+		public async Task<IActionResult> GetMetadataAsync(
 			[FromRoute][Required] NamespaceId ns,
 			[FromRoute][Required] BucketId bucket,
 			[FromRoute][Required] RefId key,
@@ -449,7 +447,7 @@ namespace Horde.Server.Ddc
 		[HttpHead("{ns}/{bucket}/{key}", Order = 500)]
 		[ProducesResponseType(type: typeof(OkResult), 200)]
 		[ProducesResponseType(type: typeof(ValidationProblemDetails), 400)]
-		public async Task<IActionResult> Head(
+		public async Task<IActionResult> HeadAsync(
 			[FromRoute][Required] NamespaceId ns,
 			[FromRoute][Required] BucketId bucket,
 			[FromRoute][Required] RefId key)
@@ -474,7 +472,7 @@ namespace Horde.Server.Ddc
 
 				// we have to verify the blobs are available locally, as the record of the key is replicated a head of the content
 				// TODO: Once we support inline replication this step is not needed as at least one region as this blob, just maybe not this current one
-				byte[] blobContents = await blob.Stream.ToByteArray();
+				byte[] blobContents = await blob.Stream.ToByteArrayAsync();
 				CbObject compactBinaryObject = new CbObject(blobContents);
 				// the reference resolver will throw if any blob is missing, so no need to do anything other then process each reference
 				IAsyncEnumerable<BlobId> references = _referenceResolver.GetReferencedBlobsAsync(ns, compactBinaryObject);
@@ -514,7 +512,7 @@ namespace Horde.Server.Ddc
 
 		[HttpGet("{ns}/exists")]
 		[ProducesDefaultResponseType]
-		public async Task<IActionResult> ExistsMultiple(
+		public async Task<IActionResult> ExistsMultipleAsync(
 			[FromRoute][Required] NamespaceId ns,
 			[FromQuery][Required] List<string> names)
 		{
@@ -552,7 +550,7 @@ namespace Horde.Server.Ddc
 
 					// we have to verify the blobs are available locally, as the record of the key is replicated a head of the content
 					// TODO: Once we support inline replication this step is not needed as at least one region as this blob, just maybe not this current one
-					byte[] blobContents = await blob.Stream.ToByteArray();
+					byte[] blobContents = await blob.Stream.ToByteArrayAsync();
 					CbObject cb = new CbObject(blobContents);
 					// the reference resolver will throw if any blob is missing, so no need to do anything other then process each reference
 					IAsyncEnumerable<BlobId> references = _referenceResolver.GetReferencedBlobsAsync(ns, cb);
@@ -582,7 +580,7 @@ namespace Horde.Server.Ddc
 
 		[HttpPut("{ns}/{bucket}/{key}.{format?}", Order = 500)]
 		[DisableRequestSizeLimit]
-		public async Task<IActionResult> PutObject(
+		public async Task<IActionResult> PutObjectAsync(
 			[FromRoute][Required] NamespaceId ns,
 			[FromRoute][Required] BucketId bucket,
 			[FromRoute][Required] RefId key,
@@ -676,7 +674,7 @@ namespace Horde.Server.Ddc
 				return Problem(e.Message, null, (int)HttpStatusCode.RequestTimeout);
 			}
 
-			(ContentId[] missingReferences, BlobId[] missingBlobs) = await _refService.PutAsync(ns, bucket, key, blobHeader, payloadObject);
+			(ContentId[] missingReferences, BlobId[] missingBlobs) = await _refService.PutAsync(ns, bucket, key, blobHeader, payloadObject, cancellationToken);
 
 			List<IoHash> missingHashes = new List<IoHash>();
 			missingHashes.AddRange(missingReferences.Select(x => x.Hash));
@@ -687,7 +685,7 @@ namespace Horde.Server.Ddc
 		[HttpPut("{ns}/{bucket}/{key}", Order = 300)]
 		[DisableRequestSizeLimit]
 		[RequiredContentType(CustomMediaTypeNames.UnrealCompactBinaryPackage)]
-		public async Task<IActionResult> PutPackage(
+		public async Task<IActionResult> PutPackageAsync(
 			[FromRoute][Required] NamespaceId ns,
 			[FromRoute][Required] BucketId bucket,
 			[FromRoute][Required] RefId key,
@@ -737,7 +735,7 @@ namespace Horde.Server.Ddc
 			CbObject rootObject = packageReader.RootObject;
 			BlobId rootObjectHash = new BlobId(packageReader.RootHash);
 
-			(ContentId[] missingReferences, BlobId[] missingBlobs) = await _refService.PutAsync(ns, bucket, key, rootObjectHash, rootObject);
+			(ContentId[] missingReferences, BlobId[] missingBlobs) = await _refService.PutAsync(ns, bucket, key, rootObjectHash, rootObject, cancellationToken);
 			return Ok(new PutObjectResponse(missingReferences, missingBlobs));
 		}
 
@@ -768,7 +766,7 @@ namespace Horde.Server.Ddc
 		}
 
 		[HttpPost("{ns}/{bucket}/{key}/finalize/{hash}.{format?}")]
-		public async Task<IActionResult> FinalizeObject(
+		public async Task<IActionResult> FinalizeObjectAsync(
 			[FromRoute][Required] NamespaceId ns,
 			[FromRoute][Required] BucketId bucket,
 			[FromRoute][Required] RefId key,
@@ -798,7 +796,7 @@ namespace Horde.Server.Ddc
 		[HttpPost("{ns}")]
 		[Consumes(CustomMediaTypeNames.UnrealCompactBinary)]
 		[Produces(CustomMediaTypeNames.UnrealCompactBinary)]
-		public async Task<IActionResult> Batch(
+		public async Task<IActionResult> BatchAsync(
 			[FromRoute][Required] NamespaceId ns,
 			[FromBody][Required] RefBatchOps ops)
 		{
@@ -852,7 +850,7 @@ namespace Horde.Server.Ddc
 						throw new Exception();
 					}
 
-					CbObject cb = new CbObject(await blob.Stream.ToByteArray());
+					CbObject cb = new CbObject(await blob.Stream.ToByteArrayAsync());
 
 					if (op.ResolveAttachments ?? false)
 					{
@@ -887,7 +885,7 @@ namespace Horde.Server.Ddc
 
 					if (op.ResolveAttachments ?? false)
 					{
-						byte[] blobContents = await blob.Stream.ToByteArray();
+						byte[] blobContents = await blob.Stream.ToByteArrayAsync();
 						CbObject cb = new CbObject(blobContents);
 						// the reference resolver will throw if any blob is missing, so no need to do anything other then process each reference
 						IAsyncEnumerable<BlobId> references = _referenceResolver.GetReferencedBlobsAsync(ns, cb);
@@ -1001,7 +999,7 @@ namespace Horde.Server.Ddc
 		[HttpDelete("{ns}/{bucket}/{key}", Order = 500)]
 		[ProducesResponseType(200)]
 		[ProducesResponseType(404)]
-		public async Task<IActionResult> Delete(
+		public async Task<IActionResult> DeleteAsync(
 			[FromRoute][Required] NamespaceId ns,
 			[FromRoute][Required] BucketId bucket,
 			[FromRoute][Required] RefId key)
diff --git a/Engine/Source/Programs/Horde/Horde.Server/Ddc/StreamUtils.cs b/Engine/Source/Programs/Horde/Horde.Server/Ddc/StreamUtils.cs
--- a/Engine/Source/Programs/Horde/Horde.Server/Ddc/StreamUtils.cs
+++ b/Engine/Source/Programs/Horde/Horde.Server/Ddc/StreamUtils.cs
@@ -11,7 +11,7 @@ namespace Horde.Server.Ddc
 {
 	public static class StreamUtils
 	{
-		public static async Task<byte[]> ToByteArray(this Stream s)
+		public static async Task<byte[]> ToByteArrayAsync(this Stream s)
 		{
 			try
 			{
diff --git a/Engine/Source/Programs/Horde/Horde/Commands/Vcs/VcsBase.cs b/Engine/Source/Programs/Horde/Horde/Commands/Vcs/VcsBase.cs
--- a/Engine/Source/Programs/Horde/Horde/Commands/Vcs/VcsBase.cs
+++ b/Engine/Source/Programs/Horde/Horde/Commands/Vcs/VcsBase.cs
@@ -212,9 +212,9 @@ namespace Horde.Commands.Vcs
 			throw new InvalidDataException($"No root directory found under {startDir}");
 		}
 
-		protected Task<DirectoryState> GetCurrentDirectoryState(DirectoryReference rootDir, DirectoryState? oldState) => GetCurrentDirectoryState(rootDir.ToDirectoryInfo(), oldState);
+		protected Task<DirectoryState> GetCurrentDirectoryState(DirectoryReference rootDir, DirectoryState? oldState) => GetCurrentDirectoryStateAsync(rootDir.ToDirectoryInfo(), oldState);
 
-		protected async Task<DirectoryState> GetCurrentDirectoryState(DirectoryInfo directoryInfo, DirectoryState? oldState)
+		protected async Task<DirectoryState> GetCurrentDirectoryStateAsync(DirectoryInfo directoryInfo, DirectoryState? oldState)
 		{
 			List<DirectoryInfo> subDirectoryInfos = new List<DirectoryInfo>();
 			foreach (DirectoryInfo subDirectoryInfo in directoryInfo.EnumerateDirectories().Where(x => !x.Name.StartsWith(".", StringComparison.Ordinal)))
@@ -233,7 +233,7 @@ namespace Horde.Commands.Vcs
 					oldState.Directories.TryGetValue(subDirectoryInfo.Name, out prevSubDirectoryState);
 				}
 
-				tasks[idx] = Task.Run(() => GetCurrentDirectoryState(subDirectoryInfo, prevSubDirectoryState));
+				tasks[idx] = Task.Run(() => GetCurrentDirectoryStateAsync(subDirectoryInfo, prevSubDirectoryState));
 			}
 
 			DirectoryState newState = new DirectoryState();
