Horde: Include compute tasks for lease utilization calculations

diff --git a/Engine/Source/Programs/Horde/Horde.Server.Tests/Fleet/LeaseUtilizationStrategyTest.cs b/Engine/Source/Programs/Horde/Horde.Server.Tests/Fleet/LeaseUtilizationStrategyTest.cs
--- a/Engine/Source/Programs/Horde/Horde.Server.Tests/Fleet/LeaseUtilizationStrategyTest.cs
+++ b/Engine/Source/Programs/Horde/Horde.Server.Tests/Fleet/LeaseUtilizationStrategyTest.cs
@@ -44,7 +44,7 @@ namespace Horde.Server.Tests.Fleet
 			await AddPlaceholderLease(_agent1, pool, Clock.UtcNow - TimeSpan.FromMinutes(120), TimeSpan.FromMinutes(120));
 			await AddPlaceholderLease(_agent2, pool, Clock.UtcNow - TimeSpan.FromMinutes(120), TimeSpan.FromMinutes(120));
 			await AddPlaceholderLease(_agent3, pool, Clock.UtcNow - TimeSpan.FromMinutes(120), TimeSpan.FromMinutes(120));
-			await AddPlaceholderLease(_agent4, pool, Clock.UtcNow - TimeSpan.FromMinutes(120), TimeSpan.FromMinutes(120));
+			await AddPlaceholderLease(_agent4, pool, Clock.UtcNow - TimeSpan.FromMinutes(120), TimeSpan.FromMinutes(120), LeaseType.Compute);
 			await AssertPoolSizeAsync(pool, 4);
 		}
 		
@@ -99,14 +99,29 @@ namespace Horde.Server.Tests.Fleet
 			PoolSizeResult output = await strategy.CalculatePoolSizeAsync(pool, _poolAgents);
 			Assert.AreEqual(expectedNumAgents, output.DesiredAgentCount);
 		}
+
+		enum LeaseType
+		{
+			ExecuteJob,
+			Compute
+		}
 		
-		private async Task<ILease> AddPlaceholderLease(IAgent agent, IPool pool, DateTime startTime, TimeSpan duration)
+		private async Task<ILease> AddPlaceholderLease(IAgent agent, IPool pool, DateTime startTime, TimeSpan duration, LeaseType leaseType = LeaseType.ExecuteJob)
 		{
 			Assert.IsNotNull(agent.SessionId);
+
+			byte[] payload;
+			ExecuteJobTask executeJobTask = new();
+			executeJobTask.JobName = "placeholderJobName";
+			payload = Any.Pack(executeJobTask).ToByteArray();
 			
-			ExecuteJobTask placeholderJobTask = new();
-			placeholderJobTask.JobName = "placeholderJobName";
-			byte[] payload = Any.Pack(placeholderJobTask).ToByteArray();
+			if (leaseType == LeaseType.Compute)
+			{
+				ComputeTask computeTask = new();
+				computeTask.Nonce = ByteString.CopyFromUtf8("test-nonce");
+				computeTask.Key = ByteString.CopyFromUtf8("test-key");
+				payload = Any.Pack(computeTask).ToByteArray();
+			}
 
 			ILease lease = await LeaseCollection.AddAsync(LeaseId.GenerateNewId(), null, "placeholderLease", agent.Id, agent.SessionId!.Value, new StreamId("placeholderStream"), pool.Id, null, startTime, payload);
 			bool wasModified = await LeaseCollection.TrySetOutcomeAsync(lease.Id, startTime + duration, LeaseOutcome.Success, null);
diff --git a/Engine/Source/Programs/Horde/Horde.Server/Agents/Fleet/LeaseUtilizationStrategy.cs b/Engine/Source/Programs/Horde/Horde.Server/Agents/Fleet/LeaseUtilizationStrategy.cs
--- a/Engine/Source/Programs/Horde/Horde.Server/Agents/Fleet/LeaseUtilizationStrategy.cs
+++ b/Engine/Source/Programs/Horde/Horde.Server/Agents/Fleet/LeaseUtilizationStrategy.cs
@@ -183,7 +183,7 @@ namespace Horde.Server.Agents.Fleet
 						: ((lease.FinishTime.Value - minTime).TotalSeconds / Settings.SampleTimeSec);
 
 					Any payload = Any.Parser.ParseFrom(lease.Payload.ToArray());
-					if (payload.Is(ExecuteJobTask.Descriptor))
+					if (payload.Is(ExecuteJobTask.Descriptor) || payload.Is(ComputeTask.Descriptor))
 					{
 						agentData.Add(minT, maxT, 1.0, 0.0);
 					}
