Horde: Support links to artifact folders
#skipci
#rnx

diff --git a/Engine/Source/Programs/Horde/HordeDashboard/src/components/artifacts/ArtifactsModal.tsx b/Engine/Source/Programs/Horde/HordeDashboard/src/components/artifacts/ArtifactsModal.tsx
--- a/Engine/Source/Programs/Horde/HordeDashboard/src/components/artifacts/ArtifactsModal.tsx
+++ b/Engine/Source/Programs/Horde/HordeDashboard/src/components/artifacts/ArtifactsModal.tsx
@@ -4,6 +4,7 @@ import { DetailsList, DetailsListLayoutMode, FontIcon, IColumn, IconButton, Moda
 import { action, makeObservable, observable } from "mobx";
 import { observer } from "mobx-react-lite";
 import { useEffect, useState } from "react";
+import { NavigateFunction, useNavigate } from "react-router-dom";
 import backend from "../../backend";
 import { ArtifactContextType, GetArtifactDirectoryEntryResponse, GetArtifactDirectoryResponse, GetArtifactFileEntryResponse, GetArtifactResponseV2 } from "../../backend/Api";
 import dashboard from "../../backend/Dashboard";
@@ -27,14 +28,21 @@ type BrowserItem = {
 
 class ArtifactsHandler {
 
-   constructor(jobId: string, stepId: string, contextType: ArtifactContextType, artifacts?:GetArtifactResponseV2[]) {
+   constructor(jobId: string, stepId: string, contextType: ArtifactContextType, artifactPath?: string, artifacts?: GetArtifactResponseV2[]) {
       makeObservable(this);
       this.jobId = jobId;
       this.stepId = stepId;
       this.context = contextType;
-      this.artifacts = artifacts;
+      this.artifacts = artifacts;      
+      this.set(artifactPath);
+
+      const params = new URLSearchParams(window.location.search);
+      params.delete("artifactPath");
+      this.baseSearch = `?${params.toString()}`;
+
+      ArtifactsHandler.current = this;
 
-      this.set();
+      
    }
 
    @observable
@@ -45,27 +53,27 @@ class ArtifactsHandler {
       this.updated++;
    }
 
-   private async set() {
+   private async set(artifactPath?: string) {
 
       let artifacts: GetArtifactResponseV2[] | undefined = this.artifacts;
 
-      if (!artifacts) {         
+      if (!artifacts) {
 
          const key = `job:${this.jobId}/step:${this.stepId}`;
          try {
             const v = await backend.getJobArtifactsV2(undefined, [key]);
-            artifacts = v.artifacts;               
+            artifacts = v.artifacts;
          } catch (err) {
             console.error(err);
-         } 
-   
+         }
+
          if (!artifacts) {
             console.error(`Missing artifacts for job: ${this.jobId} step: ${this.stepId}`);
             return;
          }
-   
-         this.artifacts = artifacts;   
-      }      
+
+         this.artifacts = artifacts;
+      }
 
       if (!artifacts) {
          console.error(`Missing artifacts for job: ${this.jobId} step: ${this.stepId}`);
@@ -108,12 +116,24 @@ class ArtifactsHandler {
       this.artifact = a;
 
       this.loading = true;
-      this.updateReady();
 
-      this.browse = await backend.getBrowseArtifacts(a.id);
+      if (!artifactPath) {
+         this.updateReady();
+      }
 
-      this.loading = false;
-      this.updateReady();
+      this.browse = await backend.getBrowseArtifacts(a.id);
+      
+      if (!artifactPath) {
+         this.loading = false;
+         this.updateReady();
+      }
+      
+      if (artifactPath) {
+         this.browse = await backend.getBrowseArtifacts(this.artifact.id, artifactPath);
+         this.path = artifactPath;   
+         this.loading = false;
+         this.updateReady();
+      }
 
    }
 
@@ -123,7 +143,7 @@ class ArtifactsHandler {
 
    }
 
-   async browseTo(path: string, push = true) {
+   async browseTo(path: string, navigate: NavigateFunction, push = true) {
 
       if (!this.artifact) {
          return;
@@ -139,6 +159,15 @@ class ArtifactsHandler {
          this.history.push(path);
       }
 
+      if (this.baseSearch) {
+         let url = `${window.location.pathname}${this.baseSearch}`;
+         if (path?.length) {
+            url += `&artifactPath=${encodeURI(path)}`;
+         }
+         navigate(url, { replace: true });
+         console.log(this.artifact.id, path);
+      } 
+
       this.loading = false;
       this.updateReady();
 
@@ -207,8 +236,10 @@ class ArtifactsHandler {
       this.history = [];
       this.stepId = "";
       this.loading = false;
+      this.baseSearch = undefined;
+      ArtifactsHandler.current = undefined;
    }
-   
+
 
    selectionCallback?: () => void;
 
@@ -231,6 +262,10 @@ class ArtifactsHandler {
    history: string[] = [];
 
    loading = false;
+
+   baseSearch?: string;
+
+   static current?: ArtifactsHandler;
 }
 
 const styles = mergeStyleSets({
@@ -260,6 +295,8 @@ const styles = mergeStyleSets({
 
 const BrowseHistory: React.FC<{ handler: ArtifactsHandler }> = observer(({ handler }) => {
 
+   const navigate = useNavigate();
+
    // subscribe
    if (handler.updated) { };
 
@@ -271,10 +308,10 @@ const BrowseHistory: React.FC<{ handler: ArtifactsHandler }> = observer(({ handl
             <IconButton disabled={backDisabled} style={{ fontSize: 14, paddingTop: 1 }} iconProps={{ iconName: 'ArrowLeft' }} onClick={() => {
                if (handler.history.length === 1) {
                   handler.history = [];
-                  handler.browseTo("", false);
+                  handler.browseTo("", navigate, false);
                } else {
                   handler.history.pop();
-                  handler.browseTo(handler.history[handler.history.length - 1], false);
+                  handler.browseTo(handler.history[handler.history.length - 1], navigate, false);
                }
             }} />
          </Stack>
@@ -285,6 +322,8 @@ const BrowseHistory: React.FC<{ handler: ArtifactsHandler }> = observer(({ handl
 
 const BrowseBreadCrumbs: React.FC<{ handler: ArtifactsHandler }> = observer(({ handler }) => {
 
+   const navigate = useNavigate();
+
    // subscribe
    if (handler.updated) { }
 
@@ -320,7 +359,7 @@ const BrowseBreadCrumbs: React.FC<{ handler: ArtifactsHandler }> = observer(({ h
       const color = last ? undefined : (dashboard.darktheme ? "#55B7FF" : "#0078D4");
       const cursor = last ? undefined : "pointer";
 
-      return <Stack horizontal onClick={() => handler.browseTo(path)} style={{ cursor: cursor }}>
+      return <Stack horizontal onClick={() => handler.browseTo(path, navigate)} style={{ cursor: cursor }}>
          <Stack>
             <Text style={{ fontSize: fontSize, color: color, fontWeight: last ? 600 : undefined }}>{e}</Text>
          </Stack>
@@ -328,7 +367,7 @@ const BrowseBreadCrumbs: React.FC<{ handler: ArtifactsHandler }> = observer(({ h
       </Stack>
    });
 
-   elements.unshift(<Stack style={{ cursor: "pointer" }} onClick={() => handler.browseTo("")}>
+   elements.unshift(<Stack style={{ cursor: "pointer" }} onClick={() => handler.browseTo("", navigate)}>
       <Stack style={{ paddingLeft: 2, paddingRight: 4 }}>
          <Text style={{ fontSize: fontSize, color: (dashboard.darktheme ? "#55B7FF" : "#0078D4"), fontWeight: 600 }}>{rootName} /</Text>
       </Stack>
@@ -469,7 +508,17 @@ const DownloadButton: React.FC<{ handler: ArtifactsHandler }> = observer(({ hand
 
 let idcounter = 0;
 
-const JobDetailArtifactsInner: React.FC<{ handler: ArtifactsHandler }> = observer(({ handler }) => {
+const JobDetailArtifactsInner: React.FC<{ jobId: string; stepId: string, artifacts?: GetArtifactResponseV2[], contextType: ArtifactContextType, artifactPath?: string }> = observer(({ jobId, stepId, artifacts, contextType, artifactPath }) => {
+
+   const handler = ArtifactsHandler.current ?? new ArtifactsHandler(jobId, stepId, contextType, artifactPath, artifacts);
+
+   useEffect(() => {
+      return () => {
+         handler?.clear();
+      };
+   }, [handler]);
+
+   const navigate = useNavigate();
 
    // subscribe
    if (handler.updated) { }
@@ -477,7 +526,7 @@ const JobDetailArtifactsInner: React.FC<{ handler: ArtifactsHandler }> = observe
    const browse = handler.browse;
 
    if (!browse) {
-      return <Stack><Spinner size={ SpinnerSize.large}/></Stack>;
+      return <Stack><Spinner size={SpinnerSize.large} /></Stack>;
    }
 
    const items: BrowserItem[] = [];
@@ -563,12 +612,12 @@ const JobDetailArtifactsInner: React.FC<{ handler: ArtifactsHandler }> = observe
 
             if (item.type === BrowserType.Directory) {
                const nbrowse = handler.path ? `${handler.path}/${item.text}` : item.text;
-               handler.browseTo(nbrowse);
+               handler.browseTo(nbrowse, navigate);
             }
             if (item.type === BrowserType.NavigateUp && handler.path) {
                const nbrowse = handler.path.split("/")
                nbrowse.pop();
-               handler.browseTo(nbrowse.join("/"));
+               handler.browseTo(nbrowse.join("/"), navigate);
             }
 
          }}>
@@ -636,15 +685,7 @@ const JobDetailArtifactsInner: React.FC<{ handler: ArtifactsHandler }> = observe
 
 
 
-export const JobArtifactsModal: React.FC<{ jobId: string; stepId: string, artifacts?:GetArtifactResponseV2[], contextType: ArtifactContextType, onClose: () => void }> = ({ jobId, stepId, artifacts, contextType, onClose }) => {
-
-   const [handler] = useState(new ArtifactsHandler(jobId, stepId, contextType, artifacts));
-
-   useEffect(() => {
-      return () => {
-         handler?.clear();
-      };
-   }, [handler]);
+export const JobArtifactsModal: React.FC<{ jobId: string; stepId: string, artifacts?: GetArtifactResponseV2[], contextType: ArtifactContextType, artifactPath?: string, onClose: () => void }> = ({ jobId, stepId, artifacts, contextType, artifactPath, onClose }) => {
 
    return <Stack>
       <Modal isOpen={true} isBlocking={true} topOffsetFixed={true} styles={{ main: { padding: 8, width: 1180, height: 820, hasBeenOpened: false, top: "80px", position: "absolute" } }} onDismiss={() => onClose()} className={hordeClasses.modal}>
@@ -662,10 +703,9 @@ export const JobArtifactsModal: React.FC<{ jobId: string; stepId: string, artifa
                            onClick={() => { onClose() }}
                         />
                      </Stack>
-
                   </Stack>
                   <Stack styles={{ root: { paddingLeft: 4, paddingRight: 0, paddingTop: 8, paddingBottom: 4 } }}>
-                     <JobDetailArtifactsInner handler={handler} />
+                     <JobDetailArtifactsInner stepId={stepId} jobId={jobId} contextType={contextType} artifactPath={artifactPath} artifacts={artifacts} />
                   </Stack>
                </Stack>
             </Stack>
diff --git a/Engine/Source/Programs/Horde/HordeDashboard/src/components/jobDetailsV2/JobOperationsBar.tsx b/Engine/Source/Programs/Horde/HordeDashboard/src/components/jobDetailsV2/JobOperationsBar.tsx
--- a/Engine/Source/Programs/Horde/HordeDashboard/src/components/jobDetailsV2/JobOperationsBar.tsx
+++ b/Engine/Source/Programs/Horde/HordeDashboard/src/components/jobDetailsV2/JobOperationsBar.tsx
@@ -1,9 +1,9 @@
 // Copyright Epic Games, Inc. All Rights Reserved.
 
-import { CommandBar, CommandBarButton, ICommandBarItemProps, IContextualMenuItem, Stack } from '@fluentui/react';
+import { CommandBar, CommandBarButton, ICommandBarItemProps, IContextualMenuItem, NavBase, Stack } from '@fluentui/react';
 import { observer } from 'mobx-react-lite';
 import React, { useState } from 'react';
-import { Link, useNavigate } from 'react-router-dom';
+import { Link, useLocation, useNavigate } from 'react-router-dom';
 import { ArtifactContextType, JobState, JobStepOutcome, JobStepState, StepData } from '../../backend/Api';
 import dashboard from '../../backend/Dashboard';
 import { hordeClasses } from '../../styles/Styles';
@@ -208,7 +208,12 @@ export const JobOperations: React.FC<{ jobDetails: JobDetailsV2 }> = observer(({
 
 const StepArtifactsOperations: React.FC<{ jobDetails: JobDetailsV2, stepId: string }> = observer(({ jobDetails, stepId }) => {
 
-   const [artifactsShown, setArtifactsShown] = useState<ArtifactContextType | undefined>(undefined);
+   
+   const navigate = useNavigate();
+
+   const query = useQuery();
+   const artifactContext = !!query.get("artifactContext") ? query.get("artifactContext")! as ArtifactContextType : undefined;
+   const artifactPath = !!query.get("artifactPath") ? query.get("artifactPath")! : undefined;
 
    const jobData = jobDetails.jobData;
 
@@ -237,12 +242,14 @@ const StepArtifactsOperations: React.FC<{ jobDetails: JobDetailsV2, stepId: stri
 
    const opsList: IContextualMenuItem[] = [];
 
+   const baseUrl = window.location.pathname + window.location.search;
+
    opsList.push({
       key: 'stepops_artifacts_step',
       text: "Logs",
       iconProps: { iconName: "Folder" },
       disabled: !atypes.get("step-saved"),
-      onClick: () => { setArtifactsShown("step-saved") }
+      onClick: () => { navigate(`${baseUrl}&artifactContext=step-saved`, { replace: true }) }
    });
 
    opsList.push({
@@ -250,7 +257,7 @@ const StepArtifactsOperations: React.FC<{ jobDetails: JobDetailsV2, stepId: stri
       text: "Temp Storage",
       iconProps: { iconName: "MenuOpen" },
       disabled: !atypes.get("step-output"),
-      onClick: () => { setArtifactsShown("step-output") }
+      onClick: () => { navigate(`${baseUrl}&artifactContext=step-output`, { replace: true }) }
    });
 
    opsList.push({
@@ -258,7 +265,7 @@ const StepArtifactsOperations: React.FC<{ jobDetails: JobDetailsV2, stepId: stri
       text: "Traces",
       iconProps: { iconName: "SearchTemplate" },
       disabled: !atypes.get("step-trace"),
-      onClick: () => { setArtifactsShown("step-trace") }
+      onClick: () => { navigate(`${baseUrl}&artifactContext=step-trace`, { replace: true }) }
    });
 
    const opsItems: ICommandBarItemProps[] = [
@@ -274,7 +281,7 @@ const StepArtifactsOperations: React.FC<{ jobDetails: JobDetailsV2, stepId: stri
    ];
 
    return <Stack>
-      {!!artifactsShown && <JobArtifactsModal stepId={step.id} jobId={jobDetails.jobId!} contextType={artifactsShown} artifacts={stepArtifacts} onClose={() => { setArtifactsShown(undefined); }} />}
+      {!!artifactContext && <JobArtifactsModal stepId={step.id} jobId={jobDetails.jobId!} contextType={artifactContext} artifactPath={artifactPath}  artifacts={stepArtifacts} onClose={() => { navigate(window.location.pathname + `?step=${stepId}`, { replace: true }) }} />}
       <Stack horizontal styles={{ root: { paddingLeft: 0 } }}>
          <Stack grow />
          <Stack horizontal>
