Horde: Add handler for when temp storage link goes missing
#skipci
#rnx

diff --git a/Engine/Source/Programs/Horde/HordeDashboard/src/components/artifacts/ArtifactsModal.tsx b/Engine/Source/Programs/Horde/HordeDashboard/src/components/artifacts/ArtifactsModal.tsx
--- a/Engine/Source/Programs/Horde/HordeDashboard/src/components/artifacts/ArtifactsModal.tsx
+++ b/Engine/Source/Programs/Horde/HordeDashboard/src/components/artifacts/ArtifactsModal.tsx
@@ -7,7 +7,7 @@ import { useEffect, useState } from "react";
 import { NavigateFunction, useNavigate } from "react-router-dom";
 import backend from "../../backend";
 import { ArtifactContextType, GetArtifactDirectoryEntryResponse, GetArtifactDirectoryResponse, GetArtifactFileEntryResponse, GetArtifactResponseV2 } from "../../backend/Api";
-import dashboard from "../../backend/Dashboard";
+import dashboard, { StatusColor } from "../../backend/Dashboard";
 import { hordeClasses } from "../../styles/Styles";
 
 
@@ -33,7 +33,7 @@ class ArtifactsHandler {
       this.jobId = jobId;
       this.stepId = stepId;
       this.context = contextType;
-      this.artifacts = artifacts;      
+      this.artifacts = artifacts;
       this.set(artifactPath);
 
       const params = new URLSearchParams(window.location.search);
@@ -42,7 +42,7 @@ class ArtifactsHandler {
 
       ArtifactsHandler.current = this;
 
-      
+
    }
 
    @observable
@@ -109,6 +109,7 @@ class ArtifactsHandler {
 
       if (!a) {
          console.error("Unable to find artifact for context", this.context, artifacts);
+         this.artifactMissing = true;
          this.updateReady();
          return;
       }
@@ -122,15 +123,15 @@ class ArtifactsHandler {
       }
 
       this.browse = await backend.getBrowseArtifacts(a.id);
-      
+
       if (!artifactPath) {
          this.loading = false;
          this.updateReady();
       }
-      
+
       if (artifactPath) {
          this.browse = await backend.getBrowseArtifacts(this.artifact.id, artifactPath);
-         this.path = artifactPath;   
+         this.path = artifactPath;
          this.loading = false;
          this.updateReady();
       }
@@ -143,6 +144,19 @@ class ArtifactsHandler {
 
    }
 
+   get contextName(): string {
+      switch (this.context) {
+         case "step-output":
+            return "Temp Storage";
+         case "step-saved":
+            return "Log";
+         case "step-trace":
+            return "Trace";
+         default:
+            return "Unknown";
+      }
+   }
+
    async browseTo(path: string, navigate: NavigateFunction, push = true) {
 
       if (!this.artifact) {
@@ -166,7 +180,7 @@ class ArtifactsHandler {
          }
          navigate(url, { replace: true });
          console.log(this.artifact.id, path);
-      } 
+      }
 
       this.loading = false;
       this.updateReady();
@@ -237,6 +251,7 @@ class ArtifactsHandler {
       this.stepId = "";
       this.loading = false;
       this.baseSearch = undefined;
+      this.artifactMissing = undefined;
       ArtifactsHandler.current = undefined;
    }
 
@@ -265,6 +280,8 @@ class ArtifactsHandler {
 
    baseSearch?: string;
 
+   artifactMissing?: boolean;
+
    static current?: ArtifactsHandler;
 }
 
@@ -523,6 +540,18 @@ const JobDetailArtifactsInner: React.FC<{ jobId: string; stepId: string, artifac
    // subscribe
    if (handler.updated) { }
 
+   if (handler.artifactMissing) {
+
+      const text = `${handler.contextName} artifact was not found on server.`;
+
+      return <Stack>
+         <Stack horizontal verticalAlign="center" verticalFill tokens={{ childrenGap: 12 }}>
+            <FontIcon style={{ paddingTop: 1, fontSize: 17, color: dashboard.getStatusColors().get(StatusColor.Failure) }} iconName="Error" />
+            <Text variant="mediumPlus">{text}</Text>
+         </Stack>
+      </Stack>;
+   }
+
    const browse = handler.browse;
 
    if (!browse) {
