Horde: Add a debug controller endpoint for debugging batch and dependency updates
#rnx

diff --git a/Engine/Source/Programs/Horde/Horde.Server/Server/DebugController.cs b/Engine/Source/Programs/Horde/Horde.Server/Server/DebugController.cs
--- a/Engine/Source/Programs/Horde/Horde.Server/Server/DebugController.cs
+++ b/Engine/Source/Programs/Horde/Horde.Server/Server/DebugController.cs
@@ -249,6 +249,7 @@ namespace Horde.Server.Server
 		
 		private readonly MongoService _mongoService;
 		private readonly ConfigService _configService;
+		private readonly JobService _jobService;
 		private readonly JobTaskSource _jobTaskSource;
 		private readonly IGraphCollection _graphCollection;
 		private readonly ILogFileCollection _logFileCollection;
@@ -259,12 +260,13 @@ namespace Horde.Server.Server
 		/// <summary>
 		/// Constructor
 		/// </summary>
-		public SecureDebugController(MongoService mongoService, ConfigService configService, JobTaskSource jobTaskSource,
+		public SecureDebugController(MongoService mongoService, ConfigService configService, JobService jobService, JobTaskSource jobTaskSource,
 			IGraphCollection graphCollection,
 			ILogFileCollection logFileCollection, IOptions<ServerSettings> settings, IOptionsSnapshot<GlobalConfig> globalConfig, ILogger<SecureDebugController> logger)
 		{
 			_mongoService = mongoService;
 			_configService = configService;
+			_jobService = jobService;
 			_jobTaskSource = jobTaskSource;
 			_graphCollection = graphCollection;
 			_logFileCollection = logFileCollection;
@@ -495,7 +497,7 @@ namespace Horde.Server.Server
 		/// <returns>Text message</returns>
 		[HttpGet]
 		[Route("/api/v1/debug/profiler/start")]
-		public async Task<ActionResult> StartProfiler()
+		public async Task<ActionResult> StartProfilerAsync()
 		{
 			if (!_globalConfig.Value.Authorize(ServerAclAction.Debug, User))
 			{
@@ -575,6 +577,32 @@ namespace Horde.Server.Server
 			string stringArg = "hello";
 			throw new Exception($"Message: numberArg:{numberArg}, stringArg:{stringArg}");
 		}
+
+		/// <summary>
+		/// Forces an update of a job's batches to debug issues such as updating dependencies
+		/// </summary>
+		/// <returns></returns>
+		[HttpGet]
+		[Route("/api/v1/debug/batchupdate/{JobId}/{BatchId}")]
+		public async Task<ActionResult> DebugBatchUpdateAsync(string jobId, string batchId)
+		{
+			if (!_globalConfig.Value.Authorize(ServerAclAction.Debug, User))
+			{
+				return Forbid(ServerAclAction.Debug);
+			}
+
+			IJob? job = await _jobService.GetJobAsync(JobId.Parse(jobId));
+
+			if (job == null) 
+			{
+				return NotFound();
+			}
+
+			await _jobService.TryUpdateBatchAsync(job, SubResourceId.Parse(batchId), newError: JobStepBatchError.None);
+
+			return Ok();
+
+		}
 	}
 }
 
