Horde: Automatically navigate to folder when there is only one entry in root
#rnx
#skipci

diff --git a/Engine/Source/Programs/Horde/HordeDashboard/src/components/artifacts/ArtifactsModal.tsx b/Engine/Source/Programs/Horde/HordeDashboard/src/components/artifacts/ArtifactsModal.tsx
--- a/Engine/Source/Programs/Horde/HordeDashboard/src/components/artifacts/ArtifactsModal.tsx
+++ b/Engine/Source/Programs/Horde/HordeDashboard/src/components/artifacts/ArtifactsModal.tsx
@@ -117,25 +117,30 @@ class ArtifactsHandler {
       this.artifact = a;
 
       this.loading = true;
+      this.updateReady();
 
       if (!artifactPath) {
-         this.updateReady();
-      }
-
-      this.browse = await backend.getBrowseArtifacts(a.id);
-
-      if (!artifactPath) {
+         this.browse = await backend.getBrowseArtifacts(a.id);
+         if (this.browse.directories?.length === 1 && !this.browse.files?.length) {
+            let d = this.browse.directories[0] as GetArtifactDirectoryEntryResponse | undefined;
+            const path: string[] = [];
+            while (d) {
+               path.push(d.name);
+               d = d.directories?.length ? d.directories[0] : undefined;
+            }
+            const dpath = path.join("/");
+            this.browse = await backend.getBrowseArtifacts(a.id, dpath);
+            this.path = dpath;
+         }
          this.loading = false;
-         this.updateReady();
-      }
-
-      if (artifactPath) {
+      } else {
          this.browse = await backend.getBrowseArtifacts(this.artifact.id, artifactPath);
          this.path = artifactPath;
          this.loading = false;
-         this.updateReady();
       }
 
+      this.updateReady();
+
    }
 
    hasContext(c: ArtifactContextType) {
