Horde: Support linking to artifacts from the log view
#skipci
#rnx

diff --git a/Engine/Source/Programs/Horde/HordeDashboard/src/components/LogView.tsx b/Engine/Source/Programs/Horde/HordeDashboard/src/components/LogView.tsx
--- a/Engine/Source/Programs/Horde/HordeDashboard/src/components/LogView.tsx
+++ b/Engine/Source/Programs/Horde/HordeDashboard/src/components/LogView.tsx
@@ -463,6 +463,10 @@ export const LogList: React.FC<{ logId: string }> = observer(({ logId }) => {
       historyAgentId = query.get("agentId")!;
    }
 
+   const artifactContext = !!query.get("artifactContext") ? query.get("artifactContext")! as ArtifactContextType : undefined;
+   const artifactPath = !!query.get("artifactPath") ? query.get("artifactPath")! : undefined;
+
+
 
    globalHandler = handler;
    globalSearchState = searchState;
@@ -833,7 +837,7 @@ export const LogList: React.FC<{ logId: string }> = observer(({ logId }) => {
 
       }
 
-   }   
+   }
 
    // we need job details here, though need to make a better accessor
    const fixme = (logSource as any).jobDetails as JobDetails | undefined;
@@ -850,12 +854,6 @@ export const LogList: React.FC<{ logId: string }> = observer(({ logId }) => {
          }
       }
 
-      menuProps.items.push({
-         key: 'jobstep_history',
-         text: 'Step History',
-         onClick: () => setLogHistory(true)
-      })
-
       if (!fixme.jobdata?.useArtifactsV2) {
          menuProps.items.push({
             key: 'jobstep_artifacts',
@@ -877,12 +875,14 @@ export const LogList: React.FC<{ logId: string }> = observer(({ logId }) => {
 
          const opsList: IContextualMenuItem[] = [];
 
+         const baseUrl = window.location.pathname + window.location.search;
+
          opsList.push({
             key: 'stepops_artifacts_step',
             text: "Logs",
             iconProps: { iconName: "Folder" },
             disabled: !atypes.get("step-saved"),
-            onClick: () => { setLogArtifacts("step-saved") }
+            onClick: () => { navigate(`${baseUrl}?artifactContext=step-saved`, { replace: true }) }
          });
 
          opsList.push({
@@ -890,7 +890,7 @@ export const LogList: React.FC<{ logId: string }> = observer(({ logId }) => {
             text: "Temp Storage",
             iconProps: { iconName: "MenuOpen" },
             disabled: !atypes.get("step-output"),
-            onClick: () => { setLogArtifacts("step-output") }
+            onClick: () => { navigate(`${baseUrl}?artifactContext=step-output`, { replace: true }) }
          });
 
          opsList.push({
@@ -898,7 +898,7 @@ export const LogList: React.FC<{ logId: string }> = observer(({ logId }) => {
             text: "Traces",
             iconProps: { iconName: "SearchTemplate" },
             disabled: !atypes.get("step-trace"),
-            onClick: () => { setLogArtifacts("step-trace") }
+            onClick: () => { navigate(`${baseUrl}?artifactContext=step-trace`, { replace: true }) }
          });
 
          menuProps.items.push({
@@ -910,6 +910,13 @@ export const LogList: React.FC<{ logId: string }> = observer(({ logId }) => {
          })
       }
 
+      menuProps.items.push({
+         key: 'jobstep_history',
+         text: 'Step History',
+         onClick: () => setLogHistory(true)
+      })
+
+
    }
 
    function updateError() {
@@ -1018,7 +1025,7 @@ export const LogList: React.FC<{ logId: string }> = observer(({ logId }) => {
       {!!fixme && <IssueModalV2 issueId={query.get("issue")} popHistoryOnClose={issueHistory} />}
       {!!fixme && logHistory && <StepHistoryModal jobDetails={fixme!} stepId={fixme!.stepByLogId(logId)?.id} onClose={() => setLogHistory(false)} />}
       {!!fixme && logArtifacts === "legacy" && <StepArtifactsModal jobDetails={fixme!} stepId={fixme!.stepByLogId(logId)?.id} onClose={() => setLogArtifacts("")} />}
-      {!!fixme && logArtifacts && logArtifacts !== "legacy" && <JobArtifactsModal jobId={fixme!.jobdata!.id} stepId={fixme!.stepByLogId(logId)?.id!} artifacts={(logSource as JobLogSource).artifactsV2} contextType={logArtifacts as ArtifactContextType} onClose={() => setLogArtifacts("")} />}
+      {!!fixme && !!artifactContext && logArtifacts !== "legacy" && <JobArtifactsModal jobId={fixme!.jobdata!.id} stepId={fixme!.stepByLogId(logId)?.id!} artifacts={(logSource as JobLogSource).artifactsV2} contextType={artifactContext} artifactPath={artifactPath} onClose={() => { navigate(window.location.pathname, { replace: true }) }} />}
       {!!historyAgentId && <HistoryModal agentId={historyAgentId} onDismiss={() => { navigate(baseUrl, { replace: true }); setHistoryAgentId(undefined) }} />}
       <Breadcrumbs items={logSource?.crumbs ?? []} title={logSource?.crumbTitle} />
       <Stack tokens={{ childrenGap: 0 }} style={{ backgroundColor: "#FFFFFF", paddingTop: 12 }}>
