Horde: Fix auth error in Perforce user impersonation

diff --git a/Engine/Source/Programs/Horde/Horde.Server/Perforce/PerforceService.cs b/Engine/Source/Programs/Horde/Horde.Server/Perforce/PerforceService.cs
--- a/Engine/Source/Programs/Horde/Horde.Server/Perforce/PerforceService.cs
+++ b/Engine/Source/Programs/Horde/Horde.Server/Perforce/PerforceService.cs
@@ -568,7 +568,7 @@ namespace Horde.Server.Perforce
 				}
 
 				DateTime expiresAt = DateTime.UtcNow + new TimeSpan(response.Data.TicketExpiration * TimeSpan.TicksPerSecond) - TimeSpan.FromMinutes(15.0);
-				ticketInfo = new Credentials(userName, null, response.Data.Ticket, expiresAt);
+				ticketInfo = new Credentials(userName, response.Data.Ticket, response.Data.Ticket, expiresAt);
 
 				lock (_userCredentialsByCluster)
 				{
