Horde: Fix test failures not halting Docker build. Semicolon needs to be quoted.

diff --git a/Engine/Source/Programs/Horde/Horde.Server/Dockerfile b/Engine/Source/Programs/Horde/Horde.Server/Dockerfile
--- a/Engine/Source/Programs/Horde/Horde.Server/Dockerfile
+++ b/Engine/Source/Programs/Horde/Horde.Server/Dockerfile
@@ -51,11 +51,11 @@ COPY Source/Programs/Horde/Horde.Server.Tests/ Source/Programs/Horde/Horde.Serve
 COPY Source/Programs/Shared/EpicGames.Serialization.Tests/ Source/Programs/Shared/EpicGames.Serialization.Tests/
 RUN redis-server --save "" --appendonly no --daemonize yes && \
 	mongod --noauth --quiet --fork --dbpath /tmp/mongodb --logpath /tmp/mongod.log && \
-	dotnet test Source/Programs/Shared/EpicGames.BuildGraph.Tests/EpicGames.BuildGraph.Tests.csproj --blame-hang-timeout 2m --logger console;verbosity=normal && \
-	dotnet test Source/Programs/Shared/EpicGames.Core.Tests/EpicGames.Core.Tests.csproj --blame-hang-timeout 2m --logger console;verbosity=normal && \
-	dotnet test Source/Programs/Shared/EpicGames.Horde.Tests/EpicGames.Horde.Tests.csproj --blame-hang-timeout 2m --logger console;verbosity=normal && \
-	dotnet test Source/Programs/Shared/EpicGames.Serialization.Tests/EpicGames.Serialization.Tests.csproj --blame-hang-timeout 2m --logger console;verbosity=normal && \
-	dotnet test Source/Programs/Horde/Horde.Server.Tests/Horde.Server.Tests.csproj --blame-hang-timeout 2m --logger console;verbosity=normal
+	dotnet test Source/Programs/Shared/EpicGames.BuildGraph.Tests/EpicGames.BuildGraph.Tests.csproj --blame-hang-timeout 2m --logger 'console;verbosity=normal' && \
+	dotnet test Source/Programs/Shared/EpicGames.Core.Tests/EpicGames.Core.Tests.csproj --blame-hang-timeout 2m --logger 'console;verbosity=normal' && \
+	dotnet test Source/Programs/Shared/EpicGames.Horde.Tests/EpicGames.Horde.Tests.csproj --blame-hang-timeout 2m --logger 'console;verbosity=normal' && \
+	dotnet test Source/Programs/Shared/EpicGames.Serialization.Tests/EpicGames.Serialization.Tests.csproj --blame-hang-timeout 2m --logger 'console;verbosity=normal' && \
+	dotnet test Source/Programs/Horde/Horde.Server.Tests/Horde.Server.Tests.csproj --blame-hang-timeout 2m --logger 'console;verbosity=normal'
 
 # Remove native libs not used on Linux x86_64
 RUN rm -rf /app/out/runtimes/osx* && \
