Horde: Add schema generation for docs to build script.

diff --git a/Engine/Source/Programs/Horde/Horde.Server/Commands/Config/DocsCommand.cs b/Engine/Source/Programs/Horde/Horde.Server/Commands/Config/DocsCommand.cs
--- a/Engine/Source/Programs/Horde/Horde.Server/Commands/Config/DocsCommand.cs
+++ b/Engine/Source/Programs/Horde/Horde.Server/Commands/Config/DocsCommand.cs
@@ -31,7 +31,7 @@ namespace Horde.Server.Commands.Config
 			public ServerSettings Horde { get; set; } = new ServerSettings();
 		}
 
-		record class PageInfo(string Title, string FileName);
+		record class PageInfo(string Title, string LinkRail, string FileName);
 
 		public override async Task<int> ExecuteAsync(ILogger logger)
 		{
@@ -44,10 +44,10 @@ namespace Horde.Server.Commands.Config
 
 			Dictionary<JsonSchemaType, PageInfo> typeToPageInfo = new Dictionary<JsonSchemaType, PageInfo>
 			{
-				[serverSchema.RootType] = new PageInfo("appsettings.json (Server)", "Config-Schema-Server.md"),
-				[globalSchema.RootType] = new PageInfo("Globals.json", "Config-Schema-Globals.md"),
-				[projectSchema.RootType] = new PageInfo("*.project.json", "Config-Schema-Projects.md"),
-				[streamSchema.RootType] = new PageInfo("*.stream.json", "Config-Schema-Streams.md"),
+				[serverSchema.RootType] = new PageInfo("appsettings.json (Server)", "Horde (../Home.md) > Deployment (../Deployment.md) > Server (Server.md)", "Deployment/ServerSettings.md"),
+				[globalSchema.RootType] = new PageInfo("Globals.json", "Horde (../Home.md) > Configuration (../Config.md)", "Config/Schema/Globals.md"),
+				[projectSchema.RootType] = new PageInfo("*.project.json", "Horde (../Home.md) > Configuration (../Config.md)", "Config/Schema/Projects.md"),
+				[streamSchema.RootType] = new PageInfo("*.stream.json", "Horde (../Home.md) > Configuration (../Config.md)", "Config/Schema/Streams.md"),
 			};
 
 			if (Agent == null)
@@ -60,13 +60,13 @@ namespace Horde.Server.Commands.Config
 				Type agentSettingsType = agentAssembly.GetType("Horde.Agent.AgentSettings")!;
 
 				JsonSchema agentSchema = Schemas.CreateSchema(agentSettingsType);
-				await WriteDocAsync(agentSchema.RootType, "appsettings.json (Agent)", "Config-Schema-Agent.md", new Dictionary<string, string>());
+				await WriteDocAsync(agentSchema.RootType, "appsettings.json (Agent)", "Deployment/AgentSettings.md", "Horde (../Home.md) > Deployment (../Deployment.md) > Agent (Agent.md)", new Dictionary<string, string>(), logger);
 			}
 
 			Dictionary<string, string> typeNameToPageName = typeToPageInfo.ToDictionary(x => x.Key.Name!, x => x.Value.FileName, StringComparer.Ordinal);
 			foreach ((JsonSchemaType type, PageInfo pageInfo) in typeToPageInfo)
 			{
-				await WriteDocAsync(type, pageInfo.Title, pageInfo.FileName, typeNameToPageName);
+				await WriteDocAsync(type, pageInfo.Title, pageInfo.FileName, pageInfo.LinkRail, typeNameToPageName, logger);
 			}
 
 			return 0;
@@ -96,14 +96,16 @@ namespace Horde.Server.Commands.Config
 			public bool TryPop([NotNullWhen(true)] out JsonSchemaType? obj) => _stack.TryPop(out obj);
 		}
 
-		async Task WriteDocAsync(JsonSchemaType rootType, string title, string fileName, Dictionary<string, string> typeNameToLink)
+		async Task WriteDocAsync(JsonSchemaType rootType, string title, string fileName, string linkRail, Dictionary<string, string> typeNameToLink, ILogger logger)
 		{
 			FileReference file = FileReference.Combine(OutputDir, fileName);
+			DirectoryReference.CreateDirectory(file.Directory);
+
 			using (FileStream stream = FileReference.Open(file, FileMode.Create, FileAccess.Write))
 			{
 				using (StreamWriter writer = new StreamWriter(stream))
 				{
-					await writer.WriteLineAsync($"[Horde](../README.md) > [Configuration](Config.md) > {title}");
+					await writer.WriteLineAsync($"{linkRail} > {title}");
 					await writer.WriteLineAsync();
 					await writer.WriteLineAsync($"# {title}");
 
@@ -162,6 +164,7 @@ namespace Horde.Server.Commands.Config
 					}
 				}
 			}
+			logger.LogInformation("Written {File}", file);
 		}
 
 		static void FindCustomTypes(JsonSchemaType type, List<JsonSchemaType> types, HashSet<string> visitedTypeNames)
diff --git a/Engine/Source/Programs/Shared/EpicGames.Horde/Storage/RefName.cs b/Engine/Source/Programs/Shared/EpicGames.Horde/Storage/RefName.cs
--- a/Engine/Source/Programs/Shared/EpicGames.Horde/Storage/RefName.cs
+++ b/Engine/Source/Programs/Shared/EpicGames.Horde/Storage/RefName.cs
@@ -5,7 +5,6 @@ using EpicGames.Serialization;
 using System;
 using System.ComponentModel;
 using System.Globalization;
-using System.Text;
 using System.Text.Json;
 using System.Text.Json.Serialization;
 
@@ -14,6 +13,7 @@ namespace EpicGames.Horde.Storage
 	/// <summary>
 	/// Identifier for a storage namespace
 	/// </summary>
+	[JsonSchemaString]
 	[JsonConverter(typeof(RefNameJsonConverter))]
 	[TypeConverter(typeof(RefNameTypeConverter))]
 	[CbConverter(typeof(RefNameCbConverter))]
