Horde: IDE warning fixes.

diff --git a/Engine/Source/Programs/Horde/Horde.Agent/Services/TelemetryService.cs b/Engine/Source/Programs/Horde/Horde.Agent/Services/TelemetryService.cs
--- a/Engine/Source/Programs/Horde/Horde.Agent/Services/TelemetryService.cs
+++ b/Engine/Source/Programs/Horde/Horde.Agent/Services/TelemetryService.cs
@@ -428,9 +428,18 @@ class TelemetryService : BackgroundService
 		
 		foreach (string line in lines)
 		{
-			if (line.Length < 5) continue;
-			if (line.StartsWith("---", StringComparison.Ordinal)) continue;
-			if (line.StartsWith("Filter", StringComparison.Ordinal)) continue;
+			if (line.Length < 5)
+			{
+				continue;
+			}
+			if (line.StartsWith("---", StringComparison.Ordinal))
+			{
+				continue;
+			}
+			if (line.StartsWith("Filter", StringComparison.Ordinal))
+			{
+				continue;
+			}
 				
 			string[] parts = line.Split("   ", StringSplitOptions.RemoveEmptyEntries);
 			string filterName = parts[0];
diff --git a/Engine/Source/Programs/Horde/Horde.Agent/Utility/XgeMetadataExtractor.cs b/Engine/Source/Programs/Horde/Horde.Agent/Utility/XgeMetadataExtractor.cs
--- a/Engine/Source/Programs/Horde/Horde.Agent/Utility/XgeMetadataExtractor.cs
+++ b/Engine/Source/Programs/Horde/Horde.Agent/Utility/XgeMetadataExtractor.cs
@@ -275,7 +275,10 @@ public class XgeMetadataExtractor
 		Dictionary<string, XgeTaskMetadataSummary> result = new();
 		foreach (XgeTaskMetadataSummary? s in GetLocalIbMonFilePaths().Select(GetSummaryFromIbMonFile))
 		{
-			if (s.LocalTaskCount == 0 && s.RemoteTaskCount == 0) continue;
+			if (s.LocalTaskCount == 0 && s.RemoteTaskCount == 0)
+			{
+				continue;
+			}
 			
 			if (!result.TryGetValue(s.Title, out XgeTaskMetadataSummary? aggregatedSummary))
 			{
@@ -298,7 +301,10 @@ public class XgeMetadataExtractor
 		Dictionary<string, XgeTaskMetadataSummary> result = new();
 		foreach (XgeTaskMetadataSummary? s in summaries)
 		{
-			if (s.LocalTaskCount == 0 && s.RemoteTaskCount == 0) continue;
+			if (s.LocalTaskCount == 0 && s.RemoteTaskCount == 0)
+			{
+				continue;
+			}
 			
 			if (!result.TryGetValue(s.Title, out XgeTaskMetadataSummary? aggregatedSummary))
 			{
@@ -442,7 +448,10 @@ public class XgeMetadataExtractor
 
 		public override bool Read(ref Utf8JsonReader reader, Type typeToConvert, JsonSerializerOptions options)
 		{
-			if (Boolean.TryParse(reader.GetString(), out bool value)) { return value; }
+			if (Boolean.TryParse(reader.GetString(), out bool value))
+			{
+				return value; 
+			}
 			throw new Exception("Cannot unmarshal type bool");
 		}
 
diff --git a/Engine/Source/Programs/Horde/Horde.Server/Agents/Fleet/FleetManagerFactory.cs b/Engine/Source/Programs/Horde/Horde.Server/Agents/Fleet/FleetManagerFactory.cs
--- a/Engine/Source/Programs/Horde/Horde.Server/Agents/Fleet/FleetManagerFactory.cs
+++ b/Engine/Source/Programs/Horde/Horde.Server/Agents/Fleet/FleetManagerFactory.cs
@@ -82,11 +82,18 @@ public sealed class FleetManagerFactory : IFleetManagerFactory
 	
 	private static T DeserializeSettings<T>(string config)
 	{
-		if (String.IsNullOrEmpty(config)) { config = "{}"; }
+		if (String.IsNullOrEmpty(config))
+		{ 
+			config = "{}"; 
+		}
+
 		try
 		{
 			T? settings = JsonSerializer.Deserialize<T>(config);
-			if (settings == null) throw new NullReferenceException($"Unable to deserialize");
+			if (settings == null)
+			{
+				throw new NullReferenceException($"Unable to deserialize");
+			}
 			return settings;
 		}
 		catch (ArgumentException e)
@@ -97,7 +104,10 @@ public sealed class FleetManagerFactory : IFleetManagerFactory
 
 	private IAmazonEC2 GetAwsEc2(FleetManagerType type)
 	{
-		if (_awsEc2 != null) return _awsEc2;
+		if (_awsEc2 != null)
+		{
+			return _awsEc2;
+		}
 		
 		_awsEc2 = _provider.GetService<IAmazonEC2>();
 		if (_settings.CurrentValue.WithAws == false || _awsEc2 == null)
@@ -110,7 +120,10 @@ public sealed class FleetManagerFactory : IFleetManagerFactory
 	
 	private IAmazonAutoScaling GetAwsAutoScaling(FleetManagerType type)
 	{
-		if (_awsAutoScaling != null) return _awsAutoScaling;
+		if (_awsAutoScaling != null)
+		{
+			return _awsAutoScaling;
+		}
 		
 		_awsAutoScaling = _provider.GetService<IAmazonAutoScaling>();
 		if (_settings.CurrentValue.WithAws == false || _awsAutoScaling == null)
diff --git a/Engine/Source/Programs/Horde/Horde.Server/Agents/Fleet/FleetService.cs b/Engine/Source/Programs/Horde/Horde.Server/Agents/Fleet/FleetService.cs
--- a/Engine/Source/Programs/Horde/Horde.Server/Agents/Fleet/FleetService.cs
+++ b/Engine/Source/Programs/Horde/Horde.Server/Agents/Fleet/FleetService.cs
@@ -412,8 +412,14 @@ namespace Horde.Server.Agents.Fleet
 					return new JobQueueStrategy(_jobCollection, _graphCollection, _streamCollection, _clock, _cache, _globalConfig, pool.JobQueueSettings);
 				case PoolSizeStrategy.LeaseUtilization:
 					LeaseUtilizationSettings luSettings = new();
-					if (pool.MinAgents != null) luSettings.MinAgents = pool.MinAgents.Value;
-					if (pool.NumReserveAgents != null) luSettings.NumReserveAgents = pool.NumReserveAgents.Value;
+					if (pool.MinAgents != null)
+					{
+						luSettings.MinAgents = pool.MinAgents.Value;
+					}
+					if (pool.NumReserveAgents != null)
+					{
+						luSettings.NumReserveAgents = pool.NumReserveAgents.Value;
+					}
 					return new LeaseUtilizationStrategy(_agentCollection, _poolCollection, _leaseCollection, _clock, _cache, luSettings);
 				case PoolSizeStrategy.NoOp:
 					return new NoOpPoolSizeStrategy();
@@ -471,7 +477,10 @@ namespace Horde.Server.Agents.Fleet
 		{
 			json = String.IsNullOrEmpty(json) ? "{}" : json;
 			T? config = JsonSerializer.Deserialize<T>(json, new JsonSerializerOptions { PropertyNameCaseInsensitive = true });
-			if (config == null) throw new ArgumentException("Unable to deserialize config: " + json);
+			if (config == null)
+			{
+				throw new ArgumentException("Unable to deserialize config: " + json);
+			}
 			return config;
 		}
 	}
diff --git a/Engine/Source/Programs/Horde/Horde.Server/Agents/Fleet/Providers/AwsRecyclingFleetManager.cs b/Engine/Source/Programs/Horde/Horde.Server/Agents/Fleet/Providers/AwsRecyclingFleetManager.cs
--- a/Engine/Source/Programs/Horde/Horde.Server/Agents/Fleet/Providers/AwsRecyclingFleetManager.cs
+++ b/Engine/Source/Programs/Horde/Horde.Server/Agents/Fleet/Providers/AwsRecyclingFleetManager.cs
@@ -244,8 +244,11 @@ public sealed class AwsRecyclingFleetManager : IFleetManager
 		using TelemetrySpan span = _tracer.StartActiveSpan($"{nameof(AwsRecyclingFleetManager)}.{nameof(StartInstancesWithRetriesAsync)}");
 		span.SetAttribute("instances", String.Join(',', instanceIds));
 		span.SetAttribute("instanceTypePriority", instanceTypePriorityString);
-		
-		if (instances.Count == 0) return;
+
+		if (instances.Count == 0)
+		{
+			return;
+		}
 		
 		List<InstanceType?> instanceTypes = new();
 		if (instanceTypePriority == null || instanceTypePriority.Count == 0)
@@ -304,8 +307,11 @@ public sealed class AwsRecyclingFleetManager : IFleetManager
 		using TelemetrySpan span = _tracer.StartActiveSpan($"{nameof(AwsRecyclingFleetManager)}.{nameof(StartInstancesAsync)}");
 		span.SetAttribute("instances", String.Join(',', instanceIds));
 		span.SetAttribute("instanceType", instanceType);
-		
-		if (instances.Count == 0) return;
+
+		if (instances.Count == 0)
+		{
+			return;
+		}
 		if (instanceType != null)
 		{
 			await ChangeInstanceTypeAsync(instances, instanceType, cancellationToken);
@@ -365,7 +371,10 @@ public sealed class AwsRecyclingFleetManager : IFleetManager
 			span.SetAttribute("fromInstanceType", instance.InstanceType);
 			span.SetAttribute("toInstanceType", newInstanceType);
 			
-			if (instance.InstanceType == newInstanceType) { continue; }
+			if (instance.InstanceType == newInstanceType)
+			{
+				continue; 
+			}
 
 			ModifyInstanceAttributeRequest request = new () { InstanceId = instance.InstanceId, InstanceType = newInstanceType };
 			ModifyInstanceAttributeResponse response = await _ec2.ModifyInstanceAttributeAsync(request, cancellationToken);
diff --git a/Engine/Source/Programs/Horde/Horde.Server/Logs/Storage/PersistentLogStorage.cs b/Engine/Source/Programs/Horde/Horde.Server/Logs/Storage/PersistentLogStorage.cs
--- a/Engine/Source/Programs/Horde/Horde.Server/Logs/Storage/PersistentLogStorage.cs
+++ b/Engine/Source/Programs/Horde/Horde.Server/Logs/Storage/PersistentLogStorage.cs
@@ -47,12 +47,8 @@ namespace Horde.Server.Logs.Storage
 			_logger.LogDebug("Reading log {LogId} index length {Length} from persistent storage", logId, length);
 
 			string path = $"{logId}/index_{length}";
-			ReadOnlyMemory<byte>? data = await _storageProvider.ReadBytesAsync(path);
-			if (data == null)
-			{
-				return null;
-			}
-			return LogIndexData.FromMemory(data.Value);
+			ReadOnlyMemory<byte> data = await _storageProvider.ReadBytesAsync(path);
+			return LogIndexData.FromMemory(data);
 		}
 
 		/// <inheritdoc/>
@@ -71,13 +67,9 @@ namespace Horde.Server.Logs.Storage
 			_logger.LogDebug("Reading log {LogId} chunk offset {Offset} from persistent storage", logId, offset);
 
 			string path = $"{logId}/offset_{offset}";
-			ReadOnlyMemory<byte>? data = await _storageProvider.ReadBytesAsync(path);
-			if(data == null)
-			{
-				return null;
-			}
+			ReadOnlyMemory<byte> data = await _storageProvider.ReadBytesAsync(path);
 
-			MemoryReader reader = new MemoryReader(data.Value);
+			MemoryReader reader = new MemoryReader(data);
 			LogChunkData chunkData = reader.ReadLogChunkData(offset, lineIndex);
 
 			if (reader.RemainingMemory.Length > 0)
diff --git a/Engine/Source/Programs/Shared/EpicGames.Core/AssemblyUtils.cs b/Engine/Source/Programs/Shared/EpicGames.Core/AssemblyUtils.cs
--- a/Engine/Source/Programs/Shared/EpicGames.Core/AssemblyUtils.cs
+++ b/Engine/Source/Programs/Shared/EpicGames.Core/AssemblyUtils.cs
@@ -1,7 +1,6 @@
 // Copyright Epic Games, Inc. All Rights Reserved.
 
 using System;
-using System.Collections.Concurrent;
 using System.Collections.Generic;
 using System.Diagnostics;
 using System.IO;
@@ -94,7 +93,7 @@ namespace EpicGames.Core
 		{
 			// Ignore any reference assemblies
 			string? directory = Path.GetFileName(Path.GetDirectoryName(assemblyPath));
-			if (!string.IsNullOrEmpty(directory) && (directory == "ref" | directory == "refint"))
+			if (!String.IsNullOrEmpty(directory) && (directory == "ref" | directory == "refint"))
 			{
 				return;
 			}
diff --git a/Engine/Source/Programs/Shared/EpicGames.Core/FileTransaction.cs b/Engine/Source/Programs/Shared/EpicGames.Core/FileTransaction.cs
--- a/Engine/Source/Programs/Shared/EpicGames.Core/FileTransaction.cs
+++ b/Engine/Source/Programs/Shared/EpicGames.Core/FileTransaction.cs
@@ -2,9 +2,7 @@
 
 using System;
 using System.Buffers;
-using System.Collections.Generic;
 using System.IO;
-using System.Text;
 using System.Threading;
 using System.Threading.Tasks;
 
diff --git a/Engine/Source/Programs/Shared/EpicGames.Core/JsonObject.cs b/Engine/Source/Programs/Shared/EpicGames.Core/JsonObject.cs
--- a/Engine/Source/Programs/Shared/EpicGames.Core/JsonObject.cs
+++ b/Engine/Source/Programs/Shared/EpicGames.Core/JsonObject.cs
@@ -9,16 +9,15 @@ using System.Collections.Specialized;
 using System.Collections;
 using System.IO;
 using System.Text;
-
-using SystemJsonObject = System.Text.Json.Nodes.JsonObject;
-using SystemJsonNode = System.Text.Json.Nodes.JsonNode;
-using SystemJsonValue = System.Text.Json.Nodes.JsonValue;
-using SystemJsonArray = System.Text.Json.Nodes.JsonArray;
 using System.Text.Encodings.Web;
-using System.Text.Unicode;
 
 namespace EpicGames.Core
 {
+	using SystemJsonObject = System.Text.Json.Nodes.JsonObject;
+	using SystemJsonNode = System.Text.Json.Nodes.JsonNode;
+	using SystemJsonValue = System.Text.Json.Nodes.JsonValue;
+	using SystemJsonArray = System.Text.Json.Nodes.JsonArray;
+
 	/// <summary>
 	/// Extension methods for OrderedDictionary to have functionality more similar to Dictionary.
 	/// </summary>
diff --git a/Engine/Source/Programs/Shared/EpicGames.Core/LogEvent.cs b/Engine/Source/Programs/Shared/EpicGames.Core/LogEvent.cs
--- a/Engine/Source/Programs/Shared/EpicGames.Core/LogEvent.cs
+++ b/Engine/Source/Programs/Shared/EpicGames.Core/LogEvent.cs
@@ -6,7 +6,6 @@ using System.Collections.Generic;
 using System.Diagnostics.CodeAnalysis;
 using System.Globalization;
 using System.Linq;
-using System.Reflection;
 using System.Text;
 using System.Text.Json;
 using System.Text.Json.Serialization;
diff --git a/Engine/Source/Programs/Shared/EpicGames.Core/LogEventParser.cs b/Engine/Source/Programs/Shared/EpicGames.Core/LogEventParser.cs
--- a/Engine/Source/Programs/Shared/EpicGames.Core/LogEventParser.cs
+++ b/Engine/Source/Programs/Shared/EpicGames.Core/LogEventParser.cs
@@ -3,7 +3,6 @@
 using System;
 using System.Collections.Generic;
 using System.Diagnostics;
-using System.IO;
 using System.Linq;
 using System.Reflection;
 using System.Text;
diff --git a/Engine/Source/Programs/Shared/EpicGames.Core/MemoryWriter.cs b/Engine/Source/Programs/Shared/EpicGames.Core/MemoryWriter.cs
--- a/Engine/Source/Programs/Shared/EpicGames.Core/MemoryWriter.cs
+++ b/Engine/Source/Programs/Shared/EpicGames.Core/MemoryWriter.cs
@@ -149,7 +149,7 @@ namespace EpicGames.Core
 	/// </summary>
 	public class PooledMemoryWriter : IMemoryWriter, IDisposable
 	{
-		MemoryPool<byte> _pool;
+		readonly MemoryPool<byte> _pool;
 		IMemoryOwner<byte> _owner;
 		Memory<byte> _memory;
 		int _length;
diff --git a/Engine/Source/Programs/Shared/EpicGames.Core/MessageTemplate.cs b/Engine/Source/Programs/Shared/EpicGames.Core/MessageTemplate.cs
--- a/Engine/Source/Programs/Shared/EpicGames.Core/MessageTemplate.cs
+++ b/Engine/Source/Programs/Shared/EpicGames.Core/MessageTemplate.cs
@@ -2,7 +2,6 @@
 
 using System;
 using System.Collections.Generic;
-using System.Linq;
 using System.Text;
 
 namespace EpicGames.Core
diff --git a/Engine/Source/Programs/Shared/EpicGames.Core/ProcessUtils.cs b/Engine/Source/Programs/Shared/EpicGames.Core/ProcessUtils.cs
--- a/Engine/Source/Programs/Shared/EpicGames.Core/ProcessUtils.cs
+++ b/Engine/Source/Programs/Shared/EpicGames.Core/ProcessUtils.cs
@@ -2,7 +2,6 @@
 
 using System;
 using System.Collections.Generic;
-using System.ComponentModel;
 using System.Diagnostics;
 using System.Linq;
 using System.Runtime.InteropServices;
diff --git a/Engine/Source/Programs/Shared/EpicGames.Core/VarInt.cs b/Engine/Source/Programs/Shared/EpicGames.Core/VarInt.cs
--- a/Engine/Source/Programs/Shared/EpicGames.Core/VarInt.cs
+++ b/Engine/Source/Programs/Shared/EpicGames.Core/VarInt.cs
@@ -2,7 +2,6 @@
 
 using System;
 using System.Numerics;
-using EpicGames.Core;
 
 namespace EpicGames.Core
 {
diff --git a/Engine/Source/Programs/Shared/EpicGames.Core/ZipArchiveExtensions.cs b/Engine/Source/Programs/Shared/EpicGames.Core/ZipArchiveExtensions.cs
--- a/Engine/Source/Programs/Shared/EpicGames.Core/ZipArchiveExtensions.cs
+++ b/Engine/Source/Programs/Shared/EpicGames.Core/ZipArchiveExtensions.cs
@@ -1,9 +1,7 @@
 // Copyright Epic Games, Inc. All Rights Reserved.
 
-using System;
 using System.IO;
 using System.IO.Compression;
-using System.Reflection;
 using System.Runtime.InteropServices;
 
 #pragma warning disable CA1707 // Identifiers should not contain underscores
diff --git a/Engine/Source/Programs/Shared/EpicGames.Perforce.Managed.Tests/ManagedWorkspaceTest.cs b/Engine/Source/Programs/Shared/EpicGames.Perforce.Managed.Tests/ManagedWorkspaceTest.cs
--- a/Engine/Source/Programs/Shared/EpicGames.Perforce.Managed.Tests/ManagedWorkspaceTest.cs
+++ b/Engine/Source/Programs/Shared/EpicGames.Perforce.Managed.Tests/ManagedWorkspaceTest.cs
@@ -250,7 +250,10 @@ public class ManagedWorkspaceTest : BasePerforceFixtureTest
 			
 			HashSet<(string clientFile, int rev, long size, string digest)> fixtureFiles = new();
 			ChangelistFixture changelist = Stream.GetChangelist(cr.Number);
-			if (changelist.IsShelved) continue;
+			if (changelist.IsShelved)
+			{
+				continue;
+			}
 			
 			foreach (DepotFileFixture depotFile in changelist.StreamFiles)
 			{
