Horde: Merging latest DDC code from UnrealCloudDDC.

diff --git a/Engine/Source/Programs/Horde/Horde.Server.Tests/Ddc/UnitTests/CompressedBufferTests.cs b/Engine/Source/Programs/Horde/Horde.Server.Tests/Ddc/UnitTests/CompressedBufferTests.cs
--- a/Engine/Source/Programs/Horde/Horde.Server.Tests/Ddc/UnitTests/CompressedBufferTests.cs
+++ b/Engine/Source/Programs/Horde/Horde.Server.Tests/Ddc/UnitTests/CompressedBufferTests.cs
@@ -1,34 +1,35 @@
 // Copyright Epic Games, Inc. All Rights Reserved.
 
-using System.Runtime.InteropServices;
+using System.IO;
 using System.Text;
+using System.Threading;
+using System.Threading.Tasks;
+using EpicGames.Core;
 using Horde.Server.Ddc;
 using Microsoft.VisualStudio.TestTools.UnitTesting;
 using OpenTelemetry.Trace;
 
 namespace Horde.Server.Tests.Ddc.UnitTests
 {
-	[TestClass]
-	public class CompressedBufferTests
-	{
+    [TestClass]
+    public class CompressedBufferTests
+    {
 
-		[TestMethod]
-		public void CompressAndDecompress()
-		{
-			if (RuntimeInformation.IsOSPlatform(OSPlatform.Windows) && RuntimeInformation.OSArchitecture == Architecture.Arm64)
-			{
-				Assert.Inconclusive("No oodle libs for Windows-Arm64");
-			}
+        [TestMethod]
+        public async Task CompressAndDecompressAsync()
+        {
+            byte[] bytes = Encoding.UTF8.GetBytes("this is a test string");
 
-			byte[] bytes = Encoding.UTF8.GetBytes("this is a test string");
+            CompressedBufferUtils bufferUtils = new(TracerProvider.Default.GetTracer("TestTracer"));
 
-			CompressedBufferUtils bufferUtils = new(TracerProvider.Default.GetTracer("TestTracer"));
+            using MemoryStream ms = new MemoryStream(); 
+            IoHash uncompressedHash = bufferUtils.CompressContent(ms, OoodleCompressorMethod.Mermaid, OoodleCompressionLevel.VeryFast, bytes);
+            ms.Position = 0;
 
-			byte[] compressedBytes = bufferUtils.CompressContent(OoodleCompressorMethod.Mermaid, OoodleCompressionLevel.VeryFast, bytes);
+            BufferedPayload bufferedPayload = await bufferUtils.DecompressContentAsync(ms, (ulong)ms.Length, CancellationToken.None);
 
-			byte[] roundTrippedBytes = bufferUtils.DecompressContent(compressedBytes);
-
-			CollectionAssert.AreEqual(bytes, roundTrippedBytes);
-		}
-	}
+            byte[] roundTrippedBytes = await bufferedPayload.GetStream().ReadAllBytesAsync();
+            CollectionAssert.AreEqual(bytes, roundTrippedBytes);
+        }
+    }
 }
diff --git a/Engine/Source/Programs/Horde/Horde.Server/Ddc/BlobId.cs b/Engine/Source/Programs/Horde/Horde.Server/Ddc/BlobId.cs
--- a/Engine/Source/Programs/Horde/Horde.Server/Ddc/BlobId.cs
+++ b/Engine/Source/Programs/Horde/Horde.Server/Ddc/BlobId.cs
@@ -6,6 +6,8 @@ using System.Globalization;
 using System.IO;
 using System.Text.Json;
 using System.Text.Json.Serialization;
+using System.Threading;
+using System.Threading.Tasks;
 using EpicGames.Core;
 using EpicGames.Serialization;
 
@@ -56,6 +58,21 @@ namespace Horde.Server.Ddc
 
 		/// <inheritdoc cref="IoHash.op_Inequality"/>
 		public static bool operator !=(BlobId left, BlobId right) => !(left == right);
+
+		/// <summary>
+		/// Constructs a BlobId from a hash
+		/// </summary>
+		public static BlobId FromIoHash(IoHash hash) => new BlobId(hash);
+
+		/// <summary>
+		/// Creates a BlobId from a data stream
+		/// </summary>
+		public static async Task<BlobId> FromStreamAsync(Stream stream, CancellationToken cancellationToken) => FromIoHash(await IoHash.ComputeAsync(stream, cancellationToken));
+
+		/// <summary>
+		/// Converts a BlobId to an IoHash
+		/// </summary>
+		public IoHash AsIoHash() => Hash;
 	}
 
 	/// <summary>
diff --git a/Engine/Source/Programs/Horde/Horde.Server/Ddc/BlobsController.cs b/Engine/Source/Programs/Horde/Horde.Server/Ddc/BlobsController.cs
--- a/Engine/Source/Programs/Horde/Horde.Server/Ddc/BlobsController.cs
+++ b/Engine/Source/Programs/Horde/Horde.Server/Ddc/BlobsController.cs
@@ -216,7 +216,8 @@ namespace Horde.Server.Ddc
 		[RequiredContentType(MediaTypeNames.Application.Octet)]
 		[DisableRequestSizeLimit]
 		public async Task<IActionResult> PostAsync(
-			[Required] NamespaceId ns)
+			[Required] NamespaceId ns,
+			CancellationToken cancellationToken)
 		{
 			ActionResult? result = await _requestHelper.HasAccessToNamespaceAsync(User, Request, ns, new[] { StorageAclAction.WriteBlobs });
 			if (result != null)
@@ -231,7 +232,7 @@ namespace Horde.Server.Ddc
 
 				await using Stream stream = payload.GetStream();
 
-				BlobId id = new BlobId(await IoHash.ComputeAsync(stream));
+				BlobId id = await BlobId.FromStreamAsync(stream, cancellationToken);
 				await _storage.PutObjectKnownHashAsync(ns, payload, id);
 
 				return Ok(new
diff --git a/Engine/Source/Programs/Horde/Horde.Server/Ddc/BufferedPayload.cs b/Engine/Source/Programs/Horde/Horde.Server/Ddc/BufferedPayload.cs
--- a/Engine/Source/Programs/Horde/Horde.Server/Ddc/BufferedPayload.cs
+++ b/Engine/Source/Programs/Horde/Horde.Server/Ddc/BufferedPayload.cs
@@ -80,6 +80,50 @@ namespace Horde.Server.Ddc
 		public override Stream GetStream() => new ReadOnlyMemoryStream(_buffer);
 	}
 
+	/// <summary>
+	/// Helper to generate a filesystem buffered payload from a stream 
+	/// </summary>
+	sealed class FilesystemBufferedPayloadWriter : IDisposable
+	{
+		private FileInfo? _tempFile;
+
+		public FilesystemBufferedPayloadWriter()
+		{
+			_tempFile = new FileInfo(Path.GetTempFileName());
+		}
+
+		public void Dispose()
+		{
+			if (_tempFile is { Exists: true })
+			{
+				_tempFile.Delete();
+			}
+		}
+
+		public FilesystemBufferedPayload Done()
+		{
+			if (_tempFile == null)
+			{
+				throw new Exception("Writable buffer already closed once");
+			}
+
+			FilesystemBufferedPayload payload = new FilesystemBufferedPayload(_tempFile);
+			// transfer ownership of the temp file to the filesystem buffered payload
+			_tempFile = null;
+			return payload;
+		}
+
+		public Stream GetWritableStream()
+		{
+			if (_tempFile == null)
+			{
+				throw new Exception("Writable buffer was closed when fetching writable stream");
+			}
+
+			return _tempFile.OpenWrite();
+		}
+	}
+
 	/// <summary>
 	/// A streaming request backed by a temporary file on disk
 	/// </summary>
@@ -87,7 +131,7 @@ namespace Horde.Server.Ddc
 	{
 		private readonly FileInfo _tempFile;
 
-		private FilesystemBufferedPayload(FileInfo tempFile)
+		internal FilesystemBufferedPayload(FileInfo tempFile)
 			: base(tempFile.Length)
 		{
 			_tempFile = tempFile;
diff --git a/Engine/Source/Programs/Horde/Horde.Server/Ddc/CompressedBufferUtils.cs b/Engine/Source/Programs/Horde/Horde.Server/Ddc/CompressedBufferUtils.cs
--- a/Engine/Source/Programs/Horde/Horde.Server/Ddc/CompressedBufferUtils.cs
+++ b/Engine/Source/Programs/Horde/Horde.Server/Ddc/CompressedBufferUtils.cs
@@ -5,94 +5,96 @@ using System.Buffers.Binary;
 using System.Collections.Generic;
 using System.IO;
 using System.Linq;
+using System.Text;
+using System.Threading.Tasks;
+using Blake3;
 using EpicGames.Core;
 using EpicGames.Compression;
 using Force.Crc32;
 using K4os.Compression.LZ4;
 using OpenTelemetry.Trace;
+using System.Threading;
 
-#pragma warning disable CS1591
+#pragma warning disable CS1591 // Missing XML comment for public type
 
 namespace Horde.Server.Ddc
 {
-	public class CompressedBufferUtils
+	class CompressedBufferHeader
 	{
-		private readonly Tracer _tracer;
-
-		public CompressedBufferUtils(Tracer tracer)
+		[System.Diagnostics.CodeAnalysis.SuppressMessage("Design", "CA1028:Enum Storage should be Int32", Justification = "Interop requires byte")]
+		public enum CompressionMethod : byte
 		{
-			_tracer = tracer;
+			// Header is followed by one uncompressed block. 
+			None = 0,
+			// Header is followed by an array of compressed block sizes then the compressed blocks. 
+			Oodle = 3,
+			LZ4 = 4,
 		}
 
-		private class Header
-		{
-			[System.Diagnostics.CodeAnalysis.SuppressMessage("Design", "CA1028:Enum Storage should be Int32", Justification = "Interop requires byte")]
-			public enum CompressionMethod : byte
-			{
-				// Header is followed by one uncompressed block. 
-				None = 0,
-				// Header is followed by an array of compressed block sizes then the compressed blocks. 
-				Oodle = 3,
-				LZ4 = 4,
-			}
-
-			public const uint ExpectedMagic = 0xb7756362; // <dot>ucb
-			public const uint HeaderLength = 64;
+		public const uint ExpectedMagic = 0xb7756362; // <dot>ucb
+		public const uint HeaderLength = 64;
 
-			// A magic number to identify a compressed buffer. Always 0xb7756362.
+		// A magic number to identify a compressed buffer. Always 0xb7756362.
 
-			public uint Magic { get; set; }
-			// A CRC-32 used to check integrity of the buffer. Uses the polynomial 0x04c11db7.
+		public uint Magic { get; set; }
+		// A CRC-32 used to check integrity of the buffer. Uses the polynomial 0x04c11db7.
 
-			public uint Crc32 { get; set; }
+		public uint Crc32 { get; set; }
 
-			// The method used to compress the buffer. Affects layout of data following the header. 
-			public CompressionMethod Method { get; set; }
-			public byte CompressionLevel { get; set; }
-			public byte CompressionMethodUsed { get; set; }
+		// The method used to compress the buffer. Affects layout of data following the header. 
+		public CompressionMethod Method { get; set; }
+		public byte CompressionLevel { get; set; }
+		public byte CompressionMethodUsed { get; set; }
 
-			// The power of two size of every uncompressed block except the last. Size is 1 << BlockSizeExponent. 
-			public byte BlockSizeExponent { get; set; }
+		// The power of two size of every uncompressed block except the last. Size is 1 << BlockSizeExponent. 
+		public byte BlockSizeExponent { get; set; }
 
-			// The number of blocks that follow the header. 
-			public uint BlockCount { get; set; }
+		// The number of blocks that follow the header. 
+		public uint BlockCount { get; set; }
 
-			// The total size of the uncompressed data. 
-			public ulong TotalRawSize { get; set; }
+		// The total size of the uncompressed data. 
+		public ulong TotalRawSize { get; set; }
 
-			// The total size of the compressed data including the header. 
-			public ulong TotalCompressedSize { get; set; }
+		// The total size of the compressed data including the header. 
+		public ulong TotalCompressedSize { get; set; }
 
-			/** The hash of the uncompressed data. */
-			public byte[] RawHash { get; set; } = Array.Empty<byte>();
+		/** The hash of the uncompressed data. */
+		public byte[] RawHash { get; set; } = Array.Empty<byte>();
 
-			public void ByteSwap()
-			{
-				Magic = BinaryPrimitives.ReverseEndianness(Magic);
-				Crc32 = BinaryPrimitives.ReverseEndianness(Crc32);
-				BlockCount = BinaryPrimitives.ReverseEndianness(BlockCount);
-				TotalRawSize = BinaryPrimitives.ReverseEndianness(TotalRawSize);
-				TotalCompressedSize = BinaryPrimitives.ReverseEndianness(TotalCompressedSize);
-			}
+		public void ByteSwap()
+		{
+			Magic = BinaryPrimitives.ReverseEndianness(Magic);
+			Crc32 = BinaryPrimitives.ReverseEndianness(Crc32);
+			BlockCount = BinaryPrimitives.ReverseEndianness(BlockCount);
+			TotalRawSize = BinaryPrimitives.ReverseEndianness(TotalRawSize);
+			TotalCompressedSize = BinaryPrimitives.ReverseEndianness(TotalCompressedSize);
 		}
+	}
 
-		private static Header ExtractHeader(byte[] content)
+	public class CompressedBufferUtils
+	{
+		private readonly Tracer _tracer;
+
+		public CompressedBufferUtils(Tracer tracer)
 		{
-			// the header is always stored big endian
-			bool needsByteSwap = BitConverter.IsLittleEndian;
-			if (content.Length < Header.HeaderLength)
-			{
-				throw new ArgumentOutOfRangeException(nameof(content), $"Content was less then {Header.HeaderLength} bytes and thus is not a compressed buffer");
-			}
+			_tracer = tracer;
+		}
+		
+		private static (CompressedBufferHeader, uint[]) ExtractHeader(BinaryReader br)
+		{
+			byte[] headerData = br.ReadBytes((int)CompressedBufferHeader.HeaderLength);
 
-			using MemoryStream ms = new MemoryStream(content);
+			using MemoryStream ms = new MemoryStream(headerData);
 			using BinaryReader reader = new BinaryReader(ms);
 
-			Header header = new Header
+			// the header is always stored big endian
+			bool needsByteSwap = BitConverter.IsLittleEndian;
+
+			CompressedBufferHeader header = new CompressedBufferHeader
 			{
 				Magic = reader.ReadUInt32(),
 				Crc32 = reader.ReadUInt32(),
-				Method = (Header.CompressionMethod)reader.ReadByte(),
+				Method = (CompressedBufferHeader.CompressionMethod)reader.ReadByte(),
 				CompressionLevel = reader.ReadByte(),
 				CompressionMethodUsed = reader.ReadByte(),
 				BlockSizeExponent = reader.ReadByte(),
@@ -108,32 +110,50 @@ namespace Horde.Server.Ddc
 				header.ByteSwap();
 			}
 
-			if (ms.Position != Header.HeaderLength)
+			if (header.Magic != CompressedBufferHeader.ExpectedMagic)
 			{
-				throw new Exception($"Read {ms.Position} bytes but expected to read {Header.HeaderLength}");
-			}
-
-			if (header.Magic != Header.ExpectedMagic)
-			{
-				throw new InvalidMagicException(header.Magic, Header.ExpectedMagic);
+				throw new InvalidMagicException(header.Magic, CompressedBufferHeader.ExpectedMagic);
 			}
 
 			// calculate the crc from the start of the method field (skipping magic which is a constant and the crc field itself)
 			const int MethodOffset = sizeof(uint) + sizeof(uint);
 
 			// none compressed objects have no extra blocks
-			uint blocksByteUsed = header.Method != Header.CompressionMethod.None ? header.BlockCount * (uint)sizeof(uint) : 0;
-			uint calculatedCrc = Crc32Algorithm.Compute(content, MethodOffset, (int)(Header.HeaderLength - MethodOffset + blocksByteUsed));
+			uint blocksByteUsed = header.Method != CompressedBufferHeader.CompressionMethod.None ? header.BlockCount * (uint)sizeof(uint) : 0;
+
+			byte[] crcData = new byte[blocksByteUsed + headerData.Length];
+			Array.Copy(headerData, crcData, headerData.Length);
+
+			uint[] blocks = Array.Empty<uint>();
+
+			if (blocksByteUsed != 0)
+			{
+				byte[] blocksData = br.ReadBytes((int)blocksByteUsed);
+
+				Array.Copy(blocksData, 0, crcData, headerData.Length, blocksData.Length);
+
+				blocks = new uint[header.BlockCount];
+				
+				for (int i = 0; i < header.BlockCount; i++)
+				{
+					ReadOnlySpan<byte> memory = new ReadOnlySpan<byte>(blocksData, i * sizeof(uint), sizeof(uint));
+					uint compressedBlockSize = BinaryPrimitives.ReadUInt32BigEndian(memory);
+					blocks[i] = compressedBlockSize;
+				}
+			}
+
+			uint calculatedCrc = Crc32Algorithm.Compute(crcData, MethodOffset, (int)(CompressedBufferHeader.HeaderLength - MethodOffset + blocksByteUsed));
 
 			if (header.Crc32 != calculatedCrc)
 			{
 				throw new InvalidHashException(header.Crc32, calculatedCrc);
 			}
 
-			return header;
+			
+			return (header, blocks);
 		}
 
-		static void WriteHeader(Header header, BinaryWriter writer)
+		static void WriteHeader(CompressedBufferHeader header, BinaryWriter writer)
 		{
 			// the header is always stored big endian
 			bool needsByteSwap = BitConverter.IsLittleEndian;
@@ -164,66 +184,57 @@ namespace Horde.Server.Ddc
 			}
 		}
 
-		public byte[] DecompressContent(byte[] content)
+		public async Task<BufferedPayload> DecompressContentAsync(Stream sourceStream, ulong streamSize, CancellationToken cancellationToken)
 		{
-			Header header = ExtractHeader(content);
+			using BinaryReader br = new BinaryReader(sourceStream);
+			(CompressedBufferHeader header, uint[] compressedBlockSizes) = ExtractHeader(br);
 
-			if (content.LongLength < (long)header.TotalCompressedSize)
+			if (streamSize < header.TotalCompressedSize)
 			{
-				throw new Exception($"Expected buffer to be {header.TotalCompressedSize} but it was {content.LongLength}");
+				throw new Exception($"Expected stream to be {header.TotalCompressedSize} but it was {streamSize}");
 			}
 
-			ulong decompressedPayloadOffset = 0;
-			byte[] decompressedPayload = new byte[header.TotalRawSize];
-
-			ReadOnlySpan<byte> memory = new ReadOnlySpan<byte>(content);
-			memory = memory.Slice((int)Header.HeaderLength);
+			using FilesystemBufferedPayloadWriter bufferedPayloadWriter = new FilesystemBufferedPayloadWriter();
 
-			bool willHaveBlocks = header.Method != Header.CompressionMethod.None;
-			if (willHaveBlocks)
 			{
-				uint[] compressedBlockSizes = new uint[header.BlockCount];
-				for (int i = 0; i < header.BlockCount; i++)
-				{
-					uint compressedBlockSize = BinaryPrimitives.ReadUInt32BigEndian(memory);
-					compressedBlockSizes[i] = compressedBlockSize;
-					memory = memory.Slice(sizeof(uint));
-				}
+				await using Stream targetStream = bufferedPayloadWriter.GetWritableStream();
+				ulong decompressedPayloadOffset = 0;
 
-				ulong blockSize = 1ul << header.BlockSizeExponent;
-				ulong compressedOffset = 0;
-
-				foreach (uint compressedBlockSize in compressedBlockSizes)
+				bool willHaveBlocks = header.Method != CompressedBufferHeader.CompressionMethod.None;
+				if (willHaveBlocks)
 				{
-					ulong rawBlockSize = Math.Min(header.TotalRawSize - decompressedPayloadOffset, blockSize);
-					ReadOnlySpan<byte> compressedPayload = memory.Slice((int)compressedOffset, (int)compressedBlockSize);
-					Span<byte> targetSpan = new Span<byte>(decompressedPayload, (int)decompressedPayloadOffset, (int)rawBlockSize);
+					ulong blockSize = 1ul << header.BlockSizeExponent;
 
-					int writtenBytes;
-					// if a block has the same raw and compressed size its uncompressed and we should not attempt to decompress it
-					if (rawBlockSize == compressedBlockSize)
-					{
-						writtenBytes = (int)rawBlockSize;
-						compressedPayload.CopyTo(targetSpan);
-					}
-					else
+					foreach (uint compressedBlockSize in compressedBlockSizes)
 					{
-						writtenBytes = DecompressPayload(compressedPayload, header, rawBlockSize, targetSpan);
-					}
+						ulong rawBlockSize = Math.Min(header.TotalRawSize - decompressedPayloadOffset, blockSize);
+						byte[] compressedPayload = br.ReadBytes((int)compressedBlockSize);
 
-					decompressedPayloadOffset += (uint)writtenBytes;
-					compressedOffset += compressedBlockSize;
+						int writtenBytes;
+						// if a block has the same raw and compressed size its uncompressed and we should not attempt to decompress it
+						if (rawBlockSize == compressedBlockSize)
+						{
+							writtenBytes = (int)rawBlockSize;
+							targetStream.Write(compressedPayload);
+						}
+						else
+						{
+							writtenBytes = DecompressPayload(compressedPayload, header, rawBlockSize, targetStream);
+						}
+
+						decompressedPayloadOffset += (uint)writtenBytes;
+					}
+				}
+				else
+				{
+					await sourceStream.CopyToAsync(targetStream, cancellationToken);
 				}
 			}
-			else
-			{
-				// if no compression is applied there are no extra blocks and just a single chunk that is uncompressed
-				Span<byte> targetSpan = new Span<byte>(decompressedPayload, 0, (int)header.TotalRawSize);
-				DecompressPayload(memory, header, header.TotalRawSize, targetSpan);
-				decompressedPayloadOffset = header.TotalRawSize;
-			}
 
-			if (header.TotalRawSize != decompressedPayloadOffset)
+			// not using the buffered payload as we transfer the ownership to the caller of this method
+			FilesystemBufferedPayload finalizedBufferedPayload = bufferedPayloadWriter.Done();
+
+			if (header.TotalRawSize != (ulong)finalizedBufferedPayload.Length)
 			{
 				throw new Exception("Did not decompress the full payload");
 			}
@@ -232,26 +243,30 @@ namespace Horde.Server.Ddc
 				using TelemetrySpan _ = _tracer.StartActiveSpan("web.hash").SetAttribute("operation.name", "web.hash");
 
 				// only read the first 20 bytes of the hash field as IoHashes are 20 bytes and not 32 bytes
-				IoHash slicedHash = new IoHash(header.RawHash.AsSpan(0, 20));
-				IoHash contentHash = IoHash.Compute(decompressedPayload);
+				byte[] slicedHash = new byte[20];
+				Array.Copy(header.RawHash, 0, slicedHash, 0, 20);
+
+				BlobId headerIdentifier = new BlobId(new IoHash(slicedHash));
+				await using Stream hashStream = finalizedBufferedPayload.GetStream();
+				BlobId contentHash = await BlobId.FromStreamAsync(hashStream, cancellationToken);
 
-				if (!slicedHash.Equals(contentHash))
+				if (!headerIdentifier.Equals(contentHash))
 				{
-					throw new Exception($"Payload was expected to be {slicedHash} but was {contentHash}");
+					throw new Exception($"Payload was expected to be {headerIdentifier} but was {contentHash}");
 				}
 			}
 
-			return decompressedPayload;
+			return finalizedBufferedPayload;
 		}
 
-		private static int DecompressPayload(ReadOnlySpan<byte> compressedPayload, Header header, ulong rawBlockSize, Span<byte> target)
+		private static int DecompressPayload(ReadOnlySpan<byte> compressedPayload, CompressedBufferHeader header, ulong rawBlockSize, Stream target)
 		{
 			switch (header.Method)
 			{
-				case Header.CompressionMethod.None:
-					compressedPayload.CopyTo(target);
+				case CompressedBufferHeader.CompressionMethod.None:
+					target.Write(compressedPayload);
 					return compressedPayload.Length;
-				case Header.CompressionMethod.Oodle:
+				case CompressedBufferHeader.CompressionMethod.Oodle:
 					{
 						byte[] result = new byte[rawBlockSize];
 						long writtenBytes = Oodle.Decompress(compressedPayload, result);
@@ -259,12 +274,14 @@ namespace Horde.Server.Ddc
 						{
 							throw new Exception("Failed to run oodle decompress");
 						}
-						result.CopyTo(target);
+						target.Write(result);
 						return (int)writtenBytes;
 					}
-				case Header.CompressionMethod.LZ4:
+				case CompressedBufferHeader.CompressionMethod.LZ4:
 					{
-						int writtenBytes = LZ4Codec.Decode(compressedPayload, target);
+						byte[] result = new byte[rawBlockSize];
+						int writtenBytes = LZ4Codec.Decode(compressedPayload, result);
+						target.Write(result);
 						return writtenBytes;
 					}
 				default:
@@ -272,56 +289,101 @@ namespace Horde.Server.Ddc
 			}
 		}
 
-		public byte[] CompressContent(OoodleCompressorMethod method, OoodleCompressionLevel compressionLevel, byte[] rawContents)
+		public IoHash CompressContent(Stream s, OoodleCompressorMethod method, OoodleCompressionLevel compressionLevel, byte[] rawContents)
 		{
-			OodleCompressorType oodleMethod = OodleUtils.ToOodleApiCompressor(method);
-			OodleCompressionLevel oodleLevel = OodleUtils.ToOodleApiCompressionLevel(compressionLevel);
-
 			const long DefaultBlockSize = 256 * 1024;
 			long blockSize = DefaultBlockSize;
 			long blockCount = (rawContents.LongLength + blockSize - 1) / blockSize;
-
-			byte blockSizeExponent = (byte)Math.Floor(Math.Log2(blockSize));
-
 			Span<byte> contentsSpan = new Span<byte>(rawContents);
-			List<byte[]> compressedBlocks = new List<byte[]>();
+			List<byte[]> blocks = new List<byte[]>();
 
 			for (int i = 0; i < blockCount; i++)
 			{
 				int rawBlockSize = Math.Min(rawContents.Length - (i * (int)blockSize), (int)blockSize);
 				Span<byte> bufferToCompress = contentsSpan.Slice((int)(i * blockSize), rawBlockSize);
 
-				int maxSize = Oodle.MaximumOutputSize(oodleMethod, bufferToCompress.Length);
+				blocks.Add(bufferToCompress.ToArray());
+			}
+			
+			return CompressContent(s, method, compressionLevel, blocks, blockSize);
+		}
+
+		public IoHash CompressContent(Stream s, OoodleCompressorMethod method, OoodleCompressionLevel compressionLevel, List<byte[]> blocks, long blockSize)
+		{
+			OodleCompressorType oodleMethod = OodleUtils.ToOodleApiCompressor(method);
+			OodleCompressionLevel oodleLevel = OodleUtils.ToOodleApiCompressionLevel(compressionLevel);
+
+			long blockCount = blocks.Count;
+
+			byte blockSizeExponent = (byte)Math.Floor(Math.Log2(blockSize));
+
+			//Span<byte> contentsSpan = new Span<byte>(rawContents);
+			List<byte[]> compressedBlocks = new List<byte[]>();
+			using Hasher hasher = Hasher.New();
+
+			ulong uncompressedContentLength = (ulong)blocks.Sum(b => b.LongLength);
+
+			ulong compressedContentLength = 0;
+			for (int i = 0; i < blockCount; i++)
+			{
+				int rawBlockSize = blocks[i].Length;
+				//int rawBlockSize = Math.Min(rawContents.Length - (i * (int)blockSize), (int)blockSize);
+				//Span<byte> bufferToCompress = contentsSpan.Slice((int)(i * blockSize), rawBlockSize);
+				byte[] bufferToCompress = blocks[i];
+				hasher.UpdateWithJoin(new ReadOnlySpan<byte>(bufferToCompress, 0, rawBlockSize));
+				int maxSize = Oodle.MaximumOutputSize(oodleMethod, rawBlockSize);
 				byte[] compressedBlock = new byte[maxSize];
-				long encodedSize = Oodle.Compress(oodleMethod, bufferToCompress.ToArray(), compressedBlock, oodleLevel);
+				long encodedSize = Oodle.Compress(oodleMethod, bufferToCompress, compressedBlock, oodleLevel);
 
 				if (encodedSize == 0)
 				{
 					throw new Exception("Failed to compress content");
 				}
 
-				compressedBlocks.Add(compressedBlock);
+				byte[] actualCompressedBlock = new byte[encodedSize];
+				Array.Copy(compressedBlock, actualCompressedBlock, encodedSize);
+				compressedBlocks.Add(actualCompressedBlock);
+				compressedContentLength += (ulong)encodedSize;
 			}
 
-			uint compressedContentLength = (uint)compressedBlocks.Sum(b => b.LongLength);
+			Hash blake3Hash = hasher.Finalize();
+			byte[] hashData = blake3Hash.AsSpanUnsafe().Slice(0, 20).ToArray();
+			IoHash hash = new IoHash(hashData);
 
-			Header header = new Header
+			CompressedBufferHeader header = new CompressedBufferHeader
 			{
-				Magic = Header.ExpectedMagic,
+				Magic = CompressedBufferHeader.ExpectedMagic,
 				Crc32 = 0,
-				Method = Header.CompressionMethod.Oodle,
+				Method = CompressedBufferHeader.CompressionMethod.Oodle,
 				CompressionLevel = (byte)compressionLevel,
 				CompressionMethodUsed = (byte)method,
 				BlockSizeExponent = blockSizeExponent,
 				BlockCount = (uint)blockCount,
-				TotalRawSize = (ulong)rawContents.LongLength,
+				TotalRawSize = (ulong)uncompressedContentLength,
 				TotalCompressedSize = (ulong)compressedContentLength,
-				RawHash = IoHash.Compute(rawContents).ToByteArray()
+				RawHash = hashData
 			};
 
-			uint blocksByteUsed = (uint)blockCount * sizeof(uint);
+			byte[] headerAndBlocks = WriteHeaderToBuffer(header, compressedBlocks.Select(b => (uint)b.Length).ToArray());
 
-			byte[] headerBuffer = new byte[Header.HeaderLength + blocksByteUsed + compressedContentLength];
+			using BinaryWriter writer = new BinaryWriter(s, Encoding.Default, leaveOpen: true);
+
+			writer.Write(headerAndBlocks);
+
+			for (int i = 0; i < blockCount; i++)
+			{
+				writer.Write(compressedBlocks[i]);
+			}
+
+			return hash;
+		}
+
+		static byte[] WriteHeaderToBuffer(CompressedBufferHeader header, uint[] compressedBlockLengths)
+		{
+			uint blockCount = header.BlockCount;
+			uint blocksByteUsed = blockCount * sizeof(uint);
+
+			byte[] headerBuffer = new byte[CompressedBufferHeader.HeaderLength + blocksByteUsed];
 
 			// write the compressed buffer, but with the wrong crc which we update and rewrite later
 			{
@@ -332,23 +394,18 @@ namespace Horde.Server.Ddc
 
 				for (int i = 0; i < blockCount; i++)
 				{
-					uint value = (uint)compressedBlocks[i].Length;
+					uint value = compressedBlockLengths[i];
 					if (BitConverter.IsLittleEndian)
 					{
 						value = BinaryPrimitives.ReverseEndianness(value);
 					}
 					writer.Write(value);
 				}
-
-				for (int i = 0; i < blockCount; i++)
-				{
-					writer.Write(compressedBlocks[i]);
-				}
 			}
 
 			// calculate the crc from the start of the method field (skipping magic which is a constant and the crc field itself)
 			const int MethodOffset = sizeof(uint) + sizeof(uint);
-			uint calculatedCrc = Crc32Algorithm.Compute(headerBuffer, MethodOffset, (int)(Header.HeaderLength - MethodOffset + blocksByteUsed));
+			uint calculatedCrc = Crc32Algorithm.Compute(headerBuffer, MethodOffset, (int)(CompressedBufferHeader.HeaderLength - MethodOffset + blocksByteUsed));
 			header.Crc32 = calculatedCrc;
 
 			// write the header again now that we have the crc
@@ -380,7 +437,6 @@ namespace Horde.Server.Ddc
 	// from OodleDataCompression.h , we define our own enums for oodle compressions used and convert to the ones expected in the oodle api
 #pragma warning disable CA1028 // Enum Storage should be Int32
 	public enum OoodleCompressorMethod : byte
-
 	{
 		NotSet = 0,
 		Selkie = 1,
@@ -406,7 +462,7 @@ namespace Horde.Server.Ddc
 		Optimal4 = 8,
 	}
 #pragma warning restore CA1028 // Enum Storage should be Int32
-	public static class OodleUtils
+	static class OodleUtils
 	{
 		public static OodleCompressionLevel ToOodleApiCompressionLevel(OoodleCompressionLevel compressionLevel)
 		{
diff --git a/Engine/Source/Programs/Horde/Horde.Server/Ddc/ContentId.cs b/Engine/Source/Programs/Horde/Horde.Server/Ddc/ContentId.cs
--- a/Engine/Source/Programs/Horde/Horde.Server/Ddc/ContentId.cs
+++ b/Engine/Source/Programs/Horde/Horde.Server/Ddc/ContentId.cs
@@ -53,6 +53,26 @@ namespace Horde.Server.Ddc
 
 		/// <inheritdoc cref="IoHash.op_Inequality"/>
 		public static bool operator !=(ContentId left, ContentId right) => !(left == right);
+
+		/// <summary>
+		/// Constructs a ContentId from an IoHash
+		/// </summary>
+		public static ContentId FromIoHash(IoHash hash) => new ContentId(hash);
+
+		/// <summary>
+		/// Constructs a ContentId from an IoHash
+		/// </summary>
+		public static ContentId FromBlobId(BlobId blobId) => FromIoHash(blobId.AsIoHash());
+
+		/// <summary>
+		/// Converts a ContentId to IoHash
+		/// </summary>
+		public IoHash AsIoHash() => Hash;
+
+		/// <summary>
+		/// Converts a ContentId to BlobId
+		/// </summary>
+		public BlobId AsBlobIdentifier() => BlobId.FromIoHash(Hash);
 	}
 
 	/// <summary>
diff --git a/Engine/Source/Programs/Horde/Horde.Server/Ddc/IBlobService.cs b/Engine/Source/Programs/Horde/Horde.Server/Ddc/IBlobService.cs
--- a/Engine/Source/Programs/Horde/Horde.Server/Ddc/IBlobService.cs
+++ b/Engine/Source/Programs/Horde/Horde.Server/Ddc/IBlobService.cs
@@ -97,7 +97,7 @@ namespace Horde.Server.Ddc
 
 	public static class BlobServiceExtensions
 	{
-		public static async Task<ContentId> PutCompressedObjectAsync(this IBlobService blobService, NamespaceId ns, BufferedPayload payload, ContentId? id, IServiceProvider provider)
+		public static async Task<ContentId> PutCompressedObjectAsync(this IBlobService blobService, NamespaceId ns, BufferedPayload payload, ContentId? id, IServiceProvider provider, CancellationToken cancellationToken)
 		{
 			IContentIdService contentIdStore = provider.GetService<IContentIdService>()!;
 			CompressedBufferUtils compressedBufferUtils = provider.GetService<CompressedBufferUtils>()!;
@@ -105,26 +105,32 @@ namespace Horde.Server.Ddc
 
 			// decompress the content and generate a identifier from it to verify the identifier we got
 			await using Stream decompressStream = payload.GetStream();
-			// TODO: we should add a overload for decompress content that can work on streams, otherwise we are still limited to 2GB compressed blobs
-			byte[] decompressedContent = compressedBufferUtils.DecompressContent(await decompressStream.ToByteArrayAsync());
 
-			IoHash decompressedHash;
-			using (TelemetrySpan _ = tracer.StartActiveSpan("web.hash").SetAttribute("operation.name", "web.hash"))
+			using BufferedPayload bufferedPayload = await compressedBufferUtils.DecompressContentAsync(decompressStream, (ulong)payload.Length, cancellationToken);
+			await using Stream decompressedStream = bufferedPayload.GetStream();
+
+			ContentId identifierDecompressedPayload;
+			if (id != null)
 			{
-				decompressedHash = IoHash.Compute(decompressedContent);
+				await blobService.VerifyContentMatchesHashAsync(decompressedStream, id.Value.AsIoHash(), cancellationToken);
+				identifierDecompressedPayload = id.Value;
 			}
-			if (id != null && id.Value.Hash != decompressedHash)
+			else
 			{
-				throw new HashMismatchException(decompressedHash, id.Value.Hash);
-			}
+				BlobId blobHash;
+				{
+					using TelemetrySpan _ = tracer.StartActiveSpan("web.hash").SetAttribute("operation.name", "web.hash");
+					blobHash = await BlobId.FromStreamAsync(decompressedStream, cancellationToken);
+				}
 
-			ContentId identifierDecompressedPayload = new ContentId(decompressedHash);
+				identifierDecompressedPayload = ContentId.FromBlobId(blobHash);
+			}
 
 			BlobId identifierCompressedPayload;
 			{
 				using TelemetrySpan _ = tracer.StartActiveSpan("web.hash").SetAttribute("operation.name", "web.hash");
 				await using Stream hashStream = payload.GetStream();
-				identifierCompressedPayload = new BlobId(await IoHash.ComputeAsync(hashStream));
+				identifierCompressedPayload = await BlobId.FromStreamAsync(hashStream, cancellationToken);
 			}
 
 			// commit the mapping from the decompressed hash to the compressed hash, we run this in parallel with the blob store submit
@@ -134,7 +140,7 @@ namespace Horde.Server.Ddc
 
 			// we still commit the compressed buffer to the object store using the hash of the compressed content
 			{
-				await blobService.PutObjectKnownHashAsync(ns, payload, identifierCompressedPayload);
+				await blobService.PutObjectKnownHashAsync(ns, payload, identifierCompressedPayload, cancellationToken);
 			}
 
 			await contentIdStoreTask;
diff --git a/Engine/Source/Programs/Horde/Horde.Server/Ddc/IReferenceResolver.cs b/Engine/Source/Programs/Horde/Horde.Server/Ddc/IReferenceResolver.cs
--- a/Engine/Source/Programs/Horde/Horde.Server/Ddc/IReferenceResolver.cs
+++ b/Engine/Source/Programs/Horde/Horde.Server/Ddc/IReferenceResolver.cs
@@ -29,7 +29,7 @@ namespace Horde.Server.Ddc
 
 		public override IoHash AsIoHash()
 		{
-			return Identifier.Hash;
+			return Identifier.AsIoHash();
 		}
 	}
 
@@ -44,7 +44,7 @@ namespace Horde.Server.Ddc
 
 		public override IoHash AsIoHash()
 		{
-			return Identifier.Hash;
+			return Identifier.AsIoHash();
 		}
 	}
 
@@ -61,7 +61,7 @@ namespace Horde.Server.Ddc
 
 		public override IoHash AsIoHash()
 		{
-			return Identifier.Hash;
+			return Identifier.AsIoHash();
 		}
 	}
 
@@ -119,8 +119,8 @@ namespace Horde.Server.Ddc
 					{
 						IoHash attachmentHash = field.AsAttachment();
 
-						BlobId blobIdentifier = new BlobId(attachmentHash);
-						ContentId contentId = new ContentId(attachmentHash);
+						BlobId blobIdentifier = BlobId.FromIoHash(attachmentHash);
+						ContentId contentId = ContentId.FromIoHash(attachmentHash);
 
 						if (field.IsBinaryAttachment())
 						{
@@ -152,7 +152,7 @@ namespace Horde.Server.Ddc
 						try
 						{
 							(contentId, resolvedBlobs) = await pendingContentIdResolve;
-							blobIdentifier = new BlobId(contentId.Hash);
+							blobIdentifier = contentId.AsBlobIdentifier();
 							wasContentId = !(resolvedBlobs is { Length: 1 } && resolvedBlobs[0].Equals(blobIdentifier));
 						}
 						catch (InvalidContentIdException)
diff --git a/Engine/Source/Programs/Horde/Horde.Server/Ddc/RefsController.cs b/Engine/Source/Programs/Horde/Horde.Server/Ddc/RefsController.cs
--- a/Engine/Source/Programs/Horde/Horde.Server/Ddc/RefsController.cs
+++ b/Engine/Source/Programs/Horde/Horde.Server/Ddc/RefsController.cs
@@ -174,7 +174,7 @@ namespace Horde.Server.Ddc
 
 								IoHash hash = binaryAttachmentField.AsBinaryAttachment();
 
-								BlobContents referencedBlobContents = await _blobStore.GetObjectAsync(ns, new BlobId(hash));
+								BlobContents referencedBlobContents = await _blobStore.GetObjectAsync(ns, BlobId.FromIoHash(hash));
 
 								if (_nginxRedirectHelper.CanRedirect(Request, referencedBlobContents))
 								{
@@ -212,7 +212,7 @@ namespace Horde.Server.Ddc
 							IAsyncEnumerable<Attachment> attachments = _referenceResolver.GetAttachmentsAsync(ns, cb);
 
 							using CbPackageBuilder writer = new CbPackageBuilder();
-							writer.AddAttachment(objectRecord.BlobIdentifier.Hash, CbPackageAttachmentFlags.IsObject, blobMemory);
+							writer.AddAttachment(objectRecord.BlobIdentifier.AsIoHash(), CbPackageAttachmentFlags.IsObject, blobMemory);
 
 							await Parallel.ForEachAsync(attachments, async (attachment, token) =>
 							{
@@ -627,12 +627,12 @@ namespace Horde.Server.Ddc
 							// TODO: convert the json object into a compact binary instead
 							CbWriter writer = new CbWriter();
 							writer.BeginObject();
-							writer.WriteBinaryAttachmentValue(blobHeader.Hash);
+							writer.WriteBinaryAttachmentValue(blobHeader.AsIoHash());
 							writer.EndObject();
 
 							byte[] blob = writer.ToByteArray();
 							payloadObject = new CbObject(blob);
-							blobHeader = new BlobId(IoHash.Compute(blob));
+							blobHeader = BlobId.FromBlob(blob);
 							break;
 						}
 					case CustomMediaTypeNames.UnrealCompactBinary:
@@ -649,13 +649,13 @@ namespace Horde.Server.Ddc
 
 							CbWriter writer = new CbWriter();
 							writer.BeginObject();
-							writer.WriteBinaryAttachment("RawHash", blobHeader.Hash);
+							writer.WriteBinaryAttachment("RawHash", blobHeader.AsIoHash());
 							writer.WriteInteger("RawSize", payload.Length);
 							writer.EndObject();
 
 							byte[] blob = writer.ToByteArray();
 							payloadObject = new CbObject(blob);
-							blobHeader = new BlobId(IoHash.Compute(blob));
+							blobHeader = BlobId.FromBlob(blob);
 							break;
 						}
 					default:
@@ -716,11 +716,11 @@ namespace Horde.Server.Ddc
 					if (entry.Flags.HasFlag(CbPackageAttachmentFlags.IsCompressed))
 					{
 						using MemoryBufferedPayload payload = new MemoryBufferedPayload(blob);
-						await _blobStore.PutCompressedObjectAsync(ns, payload, new ContentId(entry.AttachmentHash), HttpContext.RequestServices);
+						await _blobStore.PutCompressedObjectAsync(ns, payload, ContentId.FromIoHash(entry.AttachmentHash), HttpContext.RequestServices, cancellationToken);
 					}
 					else
 					{
-						await _blobStore.PutObjectAsync(ns, blob, new BlobId(entry.AttachmentHash), cancellationToken);
+						await _blobStore.PutObjectAsync(ns, blob, BlobId.FromIoHash(entry.AttachmentHash), cancellationToken);
 					}
 				}
 			}
@@ -733,7 +733,7 @@ namespace Horde.Server.Ddc
 			}
 
 			CbObject rootObject = packageReader.RootObject;
-			BlobId rootObjectHash = new BlobId(packageReader.RootHash);
+			BlobId rootObjectHash = BlobId.FromIoHash(packageReader.RootHash);
 
 			(ContentId[] missingReferences, BlobId[] missingBlobs) = await _refService.PutAsync(ns, bucket, key, rootObjectHash, rootObject, cancellationToken);
 			return Ok(new PutObjectResponse(missingReferences, missingBlobs));
