Horde: Fix sizes with artifact size in root block of temp storage
#skipci
#rnx

diff --git a/Engine/Source/Programs/Horde/HordeDashboard/src/components/artifacts/ArtifactsModal.tsx b/Engine/Source/Programs/Horde/HordeDashboard/src/components/artifacts/ArtifactsModal.tsx
--- a/Engine/Source/Programs/Horde/HordeDashboard/src/components/artifacts/ArtifactsModal.tsx
+++ b/Engine/Source/Programs/Horde/HordeDashboard/src/components/artifacts/ArtifactsModal.tsx
@@ -575,7 +575,7 @@ const JobDetailArtifactsInner: React.FC<{ jobId: string; stepId: string, artifac
       function recurseDirectories(dir: GetArtifactDirectoryEntryResponse, flattened: GetArtifactDirectoryEntryResponse[]) {
          if (!dir.directories?.length) {
             const name = flattened.length ? flattened.map(d => d.name).join("/") + "/" + dir.name : dir.name;
-            items.push({ key: d.hash, text: name, icon: "Folder", type: BrowserType.Directory, size: d.length });
+            items.push({ key: dir.hash, text: name, icon: "Folder", type: BrowserType.Directory, size: dir.length });
          } else {
             flattened.push(dir);
             dir.directories.forEach(d => recurseDirectories(d, [...flattened]));
