Horde: Remove ref links and header from jump links
#skipci
#rnx

diff --git a/Engine/Source/Programs/Horde/HordeDashboard/src/components/docs/DocView.tsx b/Engine/Source/Programs/Horde/HordeDashboard/src/components/docs/DocView.tsx
--- a/Engine/Source/Programs/Horde/HordeDashboard/src/components/docs/DocView.tsx
+++ b/Engine/Source/Programs/Horde/HordeDashboard/src/components/docs/DocView.tsx
@@ -93,9 +93,9 @@ class DocumentCache {
                lines.forEach(line => {
                   let anchor = "";
 
-                  if (line.startsWith("# ")) {
-                     anchor = line.split("# ")[1];
-                  }
+                  //if (line.startsWith("# ")) {
+                  //   anchor = line.split("# ")[1];
+                  //}
 
                   if (line.startsWith("## ")) {
                      anchor = line.split("## ")[1];
@@ -206,20 +206,7 @@ const DocRail = observer(() => {
 
    const state = linkState.state;
 
-   const refLinks: ISideRailLink[] = [];
-
-   if (dashboard.user?.dashboardFeatures?.showLandingPage === true) {
-      refLinks.push({ text: "Landing", url: "/docs" });
-   }
-
-   refLinks.push({ text: "Home", url: "/docs/Home.md" });
-   refLinks.push({ text: "User Guide", url: "/docs/Users.md" });
-   refLinks.push({ text: "Deployment", url: "/docs/Deployment.md" });
-   refLinks.push({ text: "Configuration", url: "/docs/Config.md" });
-   refLinks.push({ text: "Horde Internals", url: "/docs/Internals.md" });
-   refLinks.push({ text: "Release Notes", url: "/docs/ReleaseNotes.md" });
-
-   return <SideRail jumpLinks={state.jumpLinks} relatedLinks={refLinks} />
+   return <SideRail jumpLinks={state.jumpLinks} />
 
 })
 
