Horde: Merge reference tests from Jupiter.

diff --git a/Engine/Source/Programs/Horde/Horde.Server.Tests/ControllerIntegrationTest.cs b/Engine/Source/Programs/Horde/Horde.Server.Tests/ControllerIntegrationTest.cs
--- a/Engine/Source/Programs/Horde/Horde.Server.Tests/ControllerIntegrationTest.cs
+++ b/Engine/Source/Programs/Horde/Horde.Server.Tests/ControllerIntegrationTest.cs
@@ -15,16 +15,24 @@ using Microsoft.AspNetCore.TestHost;
 using Microsoft.Extensions.Configuration;
 using Microsoft.Extensions.DependencyInjection;
 using Horde.Server.Agents;
-using Horde.Server.Streams;
 using Horde.Server.Jobs.Templates;
 using Horde.Server.Configuration;
 using Horde.Server.Jobs.Graphs;
 using Horde.Server.Jobs.Artifacts;
 using Moq;
 using Microsoft.Extensions.Options;
+using Serilog;
 
 namespace Horde.Server.Tests;
 
+static class SerilogExtensions
+{
+	public static LoggerConfiguration Override<T>(this Serilog.Configuration.LoggerMinimumLevelConfiguration configuration, Serilog.Events.LogEventLevel minimumLevel) where T : class
+	{
+		return configuration.Override(typeof(T).FullName!, minimumLevel);
+	}
+}
+
 public class TestWebApplicationFactory<TStartup> : WebApplicationFactory<TStartup> where TStartup : class
 {
 	private readonly MongoDbInstance _mongoDbInstance;
@@ -32,6 +40,14 @@ public class TestWebApplicationFactory<TStartup> : WebApplicationFactory<TStartu
 	public TestWebApplicationFactory(MongoDbInstance mongoDbInstance)
 	{
 		_mongoDbInstance = mongoDbInstance;
+
+		Serilog.Log.Logger = new LoggerConfiguration()
+			.Enrich.FromLogContext()
+			.WriteTo.Console()
+			.MinimumLevel.Information()
+			.MinimumLevel.Override<MongoService>(Serilog.Events.LogEventLevel.Warning)
+			.MinimumLevel.Override("Redis", Serilog.Events.LogEventLevel.Warning)
+			.CreateLogger();
 	}
 
 	protected override void ConfigureWebHost(IWebHostBuilder builder)
@@ -66,13 +82,15 @@ public class ControllerIntegrationTest : IDisposable
 		Factory = new TestWebApplicationFactory<Startup>(MongoDbInstance);
 		Client = Factory.CreateClient();
 
-		_fixture = new Lazy<Task<Fixture>>(Task.Run(() => CreateFixture()));
+		_fixture = new Lazy<Task<Fixture>>(CreateFixtureTask);
 	}
 
 	protected MongoDbInstance MongoDbInstance { get; }
 	private TestWebApplicationFactory<Startup> Factory { get; }
 	protected HttpClient Client { get; }
 
+	protected IServiceProvider ServiceProvider => Factory.Services;
+
 	public void Dispose()
 	{
 		Dispose(true);
@@ -89,6 +107,11 @@ public class ControllerIntegrationTest : IDisposable
 		return _fixture.Value;
 	}
 
+	private Task<Fixture> CreateFixtureTask()
+	{
+		return Task.Run(() => CreateFixture());
+	}
+
 	private async Task<Fixture> CreateFixture()
 	{
 		IServiceProvider services = Factory.Services;
diff --git a/Engine/Source/Programs/Horde/Horde.Server.Tests/Ddc/Functional/Objects/ReferencesTests.cs b/Engine/Source/Programs/Horde/Horde.Server.Tests/Ddc/Functional/Objects/ReferencesTests.cs
new file mode 100644
--- /dev/null
+++ b/Engine/Source/Programs/Horde/Horde.Server.Tests/Ddc/Functional/Objects/ReferencesTests.cs
@@ -0,0 +1,2139 @@
+// Copyright Epic Games, Inc. All Rights Reserved.
+
+using System;
+using System.Collections.Generic;
+using System.IO;
+using System.Linq;
+using System.Net;
+using System.Net.Http;
+using System.Net.Http.Headers;
+using System.Net.Http.Json;
+using System.Net.Mime;
+using System.Text;
+using System.Text.Json;
+using System.Text.Json.Nodes;
+using System.Threading.Tasks;
+using EpicGames.Core;
+using Microsoft.AspNetCore.Mvc;
+using Microsoft.Extensions.DependencyInjection;
+using Microsoft.VisualStudio.TestTools.UnitTesting;
+using EpicGames.Horde.Storage;
+using EpicGames.Serialization;
+using EpicGames.AspNet;
+using Horde.Server.Ddc;
+using Horde.Server.Server;
+using Microsoft.Extensions.Options;
+using Horde.Server.Storage;
+using Horde.Server.Configuration;
+
+namespace Horde.Server.Tests.Ddc.FunctionalTests.References
+{
+	public interface IReferencesStore
+	{
+		Task<RefRecord> GetAsync(NamespaceId ns, BucketId bucket, RefId key, FieldFlags flags);
+
+		[Flags]
+		public enum FieldFlags
+		{
+			None = 0,
+			IncludePayload = 1,
+			All = IncludePayload
+		}
+
+		Task PutAsync(NamespaceId ns, BucketId bucket, RefId key, BlobId blobHash, byte[] blob, bool isFinalized);
+		Task FinalizeAsync(NamespaceId ns, BucketId bucket, RefId key, BlobId blobIdentifier);
+
+		Task UpdateLastAccessTimeAsync(NamespaceId ns, BucketId bucket, RefId key, DateTime newLastAccessTime);
+		IAsyncEnumerable<(NamespaceId, BucketId, RefId, DateTime)> GetRecordsAsync();
+
+		IAsyncEnumerable<NamespaceId> GetNamespacesAsync();
+		Task<bool> DeleteAsync(NamespaceId ns, BucketId bucket, RefId key);
+		Task<long> DropNamespaceAsync(NamespaceId ns);
+		Task<long> DeleteBucketAsync(NamespaceId ns, BucketId bucket);
+	}
+
+	static class CompileHacks
+	{
+		public static Task DeleteObjectAsync(this IBlobService blobService, NamespaceId namespaceId, BlobId blobId)
+		{
+			throw new NotImplementedException();
+		}
+	}
+
+	[Ignore]
+	[TestClass]
+    public class ReferencesTests : ControllerIntegrationTest
+    {
+		//		private static TestServer? _server;
+		private static HttpClient? _httpClient;
+
+		protected IBlobService Service { get; set; } = null!;
+		protected IReferencesStore ReferencesStore { get; set; } = null!;
+		protected NamespaceId TestNamespace { get; } = new NamespaceId("test-namespace");
+
+		public ReferencesTests()
+		{
+			_httpClient = Client;
+
+			ServerSettings serverSettings = ServiceProvider.GetRequiredService<IOptions<ServerSettings>>().Value;
+
+			GlobalConfig globalConfig = ServiceProvider.GetRequiredService<IOptions<GlobalConfig>>().Value;
+
+			globalConfig.Storage.Backends.Add(new BackendConfig { Id = new BackendId("default"), Type = StorageBackendType.Memory });
+			globalConfig.Storage.Namespaces.Add(new NamespaceConfig { Id = TestNamespace, Backend = globalConfig.Storage.Backends[^1].Id });
+			globalConfig.PostLoad(serverSettings);
+
+			ConfigService configService = ServiceProvider.GetRequiredService<ConfigService>();
+			configService.Set(IoHash.Zero, globalConfig);
+
+			Service = ServiceProvider.GetRequiredService<IBlobService>();
+//			ReferencesStore = ServiceProvider.GetRequiredService<IRefService>();
+		}
+
+		[TestMethod]
+        public async Task PutGetBlobAsync()
+        {
+            const string objectContents = "This is treated as a opaque blob";
+            byte[] data = Encoding.ASCII.GetBytes(objectContents);
+            BlobId objectHash = BlobId.FromBlob(data);
+            RefId key = RefId.FromName("newBlobObject");
+            using HttpContent requestContent = new ByteArrayContent(data);
+            requestContent.Headers.ContentType = new MediaTypeHeaderValue(MediaTypeNames.Application.Octet);
+            requestContent.Headers.Add(CommonHeaders.HashHeaderName, objectHash.ToString());
+
+            HttpResponseMessage result = await _httpClient!.PutAsync(new Uri($"api/v1/refs/{TestNamespace}/bucket/{key}", UriKind.Relative), requestContent);
+            result.EnsureSuccessStatusCode();
+
+            {
+                HttpResponseMessage getResponse = await _httpClient.GetAsync(new Uri($"api/v1/refs/{TestNamespace}/bucket/{key}.raw", UriKind.Relative));
+                getResponse.EnsureSuccessStatusCode();
+                await using MemoryStream ms = new MemoryStream();
+                await getResponse.Content.CopyToAsync(ms);
+
+                byte[] roundTrippedBuffer = ms.ToArray();
+                string roundTrippedPayload = Encoding.ASCII.GetString(roundTrippedBuffer);
+
+                Assert.AreEqual(objectContents, roundTrippedPayload);
+                CollectionAssert.AreEqual(data, roundTrippedBuffer);
+                Assert.AreEqual(objectHash, BlobId.FromBlob(roundTrippedBuffer));
+            }
+
+            {
+                BlobId attachment;
+                {
+                    HttpResponseMessage getResponse = await _httpClient.GetAsync(new Uri($"api/v1/refs/{TestNamespace}/bucket/{key}.uecb", UriKind.Relative));
+                    getResponse.EnsureSuccessStatusCode();
+                    await using MemoryStream ms = new MemoryStream();
+                    await getResponse.Content.CopyToAsync(ms);
+
+                    byte[] roundTrippedBuffer = ms.ToArray();
+                    CbObject cb = new CbObject(roundTrippedBuffer);
+                    List<CbField> fields = cb.ToList();
+
+                    Assert.AreEqual(2, fields.Count);
+                    CbField payloadField = fields[0];
+                    Assert.IsNotNull(payloadField);
+                    Assert.IsTrue(payloadField.IsBinaryAttachment());
+                    attachment = BlobId.FromIoHash(payloadField.AsBinaryAttachment());
+                }
+
+                {
+                    HttpResponseMessage getAttachment = await _httpClient.GetAsync(new Uri($"api/v1/blobs/{TestNamespace}/{attachment}", UriKind.Relative));
+                    getAttachment.EnsureSuccessStatusCode();
+                    await using MemoryStream ms = new MemoryStream();
+                    await getAttachment.Content.CopyToAsync(ms);
+                    byte[] roundTrippedBuffer = ms.ToArray();
+                    string roundTrippedString = Encoding.ASCII.GetString(roundTrippedBuffer);
+
+                    Assert.AreEqual(objectContents, roundTrippedString);
+                    Assert.AreEqual(objectHash, BlobId.FromBlob(roundTrippedBuffer));
+                }
+            }
+
+            {
+                HttpResponseMessage getResponse = await _httpClient.GetAsync(new Uri($"api/v1/refs/{TestNamespace}/bucket/{key}.json", UriKind.Relative));
+                getResponse.EnsureSuccessStatusCode();
+                await using MemoryStream ms = new MemoryStream();
+                await getResponse.Content.CopyToAsync(ms);
+
+                byte[] roundTrippedBuffer = ms.ToArray();
+                string s = Encoding.ASCII.GetString(roundTrippedBuffer);
+                JsonNode? jsonNode = JsonNode.Parse(s);
+                Assert.IsNotNull(jsonNode);
+                Assert.AreEqual(objectHash, BlobId.Parse(jsonNode["RawHash"]!.GetValue<string>()));
+            }
+
+            {
+                // request the object as a json response using accept instead of the format filter
+                using HttpRequestMessage request = new(HttpMethod.Get, new Uri($"api/v1/refs/{TestNamespace}/bucket/{key}", UriKind.Relative));
+                request.Headers.Accept.Add(new MediaTypeWithQualityHeaderValue(MediaTypeNames.Application.Json));
+                HttpResponseMessage getResponse = await _httpClient.SendAsync(request);
+                getResponse.EnsureSuccessStatusCode();
+                Assert.AreEqual(MediaTypeNames.Application.Json, getResponse.Content.Headers.ContentType?.MediaType);
+
+                await using MemoryStream ms = new MemoryStream();
+                await getResponse.Content.CopyToAsync(ms);
+
+                byte[] roundTrippedBuffer = ms.ToArray();
+                string s = Encoding.ASCII.GetString(roundTrippedBuffer);
+                JsonNode? node = JsonNode.Parse(s);
+                Assert.IsNotNull(node);
+                Assert.AreEqual(objectHash, BlobId.Parse(node["RawHash"]!.ToString()));
+            }
+
+            {
+                // request the object as a jupiter inlined payload
+                using HttpRequestMessage request = new (HttpMethod.Get, new Uri($"api/v1/refs/{TestNamespace}/bucket/{key}", UriKind.Relative));
+                request.Headers.Accept.Add(new MediaTypeWithQualityHeaderValue(CustomMediaTypeNames.JupiterInlinedPayload));
+                HttpResponseMessage getResponse = await _httpClient.SendAsync(request);
+                getResponse.EnsureSuccessStatusCode();
+                Assert.AreEqual(CustomMediaTypeNames.JupiterInlinedPayload, getResponse.Content.Headers.ContentType?.MediaType);
+
+                await using MemoryStream ms = new MemoryStream();
+                await getResponse.Content.CopyToAsync(ms);
+                byte[] roundTrippedBuffer = ms.ToArray();
+
+                string roundTrippedString = Encoding.ASCII.GetString(roundTrippedBuffer);
+
+                Assert.AreEqual(objectContents, roundTrippedString);
+                Assert.AreEqual(objectHash, BlobId.FromBlob(roundTrippedBuffer));
+            }
+        }
+
+        [TestMethod]
+        public async Task PutGetCompactBinaryAsync()
+        {
+            CbWriter writer = new CbWriter();
+            writer.BeginObject();
+            writer.WriteString("stringField", "thisIsAField");
+            writer.EndObject();
+
+            byte[] objectData = writer.ToByteArray();
+            BlobId objectHash = BlobId.FromBlob(objectData);
+            RefId key = RefId.FromName("newReferenceObject");
+
+            using HttpContent requestContent = new ByteArrayContent(objectData);
+            requestContent.Headers.ContentType = new MediaTypeHeaderValue(CustomMediaTypeNames.UnrealCompactBinary);
+            requestContent.Headers.Add(CommonHeaders.HashHeaderName, objectHash.ToString());
+
+            using HttpResponseMessage result = await _httpClient!.PutAsync(new Uri($"api/v1/refs/{TestNamespace}/bucket/{key}.uecb", UriKind.Relative), requestContent);
+            result.EnsureSuccessStatusCode();
+
+            {
+                Assert.AreEqual(CustomMediaTypeNames.UnrealCompactBinary, result!.Content.Headers.ContentType!.MediaType);
+                // check that no blobs are missing
+                await using MemoryStream ms = new MemoryStream();
+                await result.Content.CopyToAsync(ms);
+                byte[] roundTrippedBuffer = ms.ToArray();
+                CbObject cb = new CbObject(roundTrippedBuffer);
+                CbField needsField = cb["needs"];
+                Assert.AreNotEqual(CbField.Empty, needsField);
+                List<BlobId> missingBlobs = needsField.AsArray().Select(field => BlobId.FromIoHash(field.AsHash())).ToList();
+                Assert.AreEqual(0, missingBlobs.Count);
+            }
+
+            {
+                BucketId bucket = new BucketId("bucket");
+
+                RefRecord objectRecord = await ReferencesStore.GetAsync(TestNamespace, bucket, key, IReferencesStore.FieldFlags.IncludePayload);
+
+                Assert.IsTrue(objectRecord.IsFinalized);
+                Assert.AreEqual(key, objectRecord.Name);
+                Assert.AreEqual(objectHash, objectRecord.BlobIdentifier);
+                Assert.IsNotNull(objectRecord.InlinePayload);
+            }
+
+            {
+                HttpResponseMessage getResponse = await _httpClient.GetAsync(new Uri($"api/v1/refs/{TestNamespace}/bucket/{key}.raw", UriKind.Relative));
+                getResponse.EnsureSuccessStatusCode();
+                await using MemoryStream ms = new MemoryStream();
+                await getResponse.Content.CopyToAsync(ms);
+
+                byte[] roundTrippedBuffer = ms.ToArray();
+
+                CollectionAssert.AreEqual(objectData, roundTrippedBuffer);
+                Assert.AreEqual(objectHash, BlobId.FromBlob(roundTrippedBuffer));
+            }
+
+            {
+                HttpResponseMessage getResponse = await _httpClient.GetAsync(new Uri($"api/v1/refs/{TestNamespace}/bucket/{key}.uecb", UriKind.Relative));
+                getResponse.EnsureSuccessStatusCode();
+                await using MemoryStream ms = new MemoryStream();
+                await getResponse.Content.CopyToAsync(ms);
+
+                byte[] roundTrippedBuffer = ms.ToArray();
+                CbObject cb = new CbObject(roundTrippedBuffer);
+                List<CbField> fields = cb.ToList();
+
+                Assert.AreEqual(1, fields.Count);
+                CbField stringField = fields[0];
+                Assert.AreEqual("stringField", stringField.Name);
+                Assert.AreEqual("thisIsAField", stringField.AsString());
+            }
+
+            {
+                HttpResponseMessage getResponse = await _httpClient.GetAsync(new Uri($"api/v1/refs/{TestNamespace}/bucket/{key}.json", UriKind.Relative));
+                getResponse.EnsureSuccessStatusCode();
+                await using MemoryStream ms = new MemoryStream();
+                await getResponse.Content.CopyToAsync(ms);
+
+                byte[] roundTrippedBuffer = ms.ToArray();
+
+                string s = Encoding.ASCII.GetString(roundTrippedBuffer);
+                JsonNode? node = JsonNode.Parse(s);
+                Assert.IsNotNull(node);
+                Assert.AreEqual("thisIsAField", node["stringField"]!.GetValue<string>());
+            }
+        }
+
+        
+        [TestMethod]
+        public async Task PutGetCompactBinaryFilteringAsync()
+        {
+            CbWriter writer = new CbWriter();
+            writer.BeginObject();
+            writer.WriteString("stringField", "thisIsAField");
+            writer.EndObject();
+
+            byte[] objectData = writer.ToByteArray();
+            BlobId objectHash = BlobId.FromBlob(objectData);
+            RefId key = RefId.FromName("newReferenceObject");
+
+            using HttpContent requestContent = new ByteArrayContent(objectData);
+            requestContent.Headers.ContentType = new MediaTypeHeaderValue(CustomMediaTypeNames.UnrealCompactBinary);
+            requestContent.Headers.Add(CommonHeaders.HashHeaderName, objectHash.ToString());
+
+            HttpResponseMessage result = await _httpClient!.PutAsync(new Uri($"api/v1/refs/{TestNamespace}/bucket/{key}.uecb", UriKind.Relative), requestContent);
+            result.EnsureSuccessStatusCode();
+
+            {
+                Assert.AreEqual(CustomMediaTypeNames.UnrealCompactBinary, result!.Content.Headers.ContentType!.MediaType);
+                // check that no blobs are missing
+                await using MemoryStream ms = new MemoryStream();
+                await result.Content.CopyToAsync(ms);
+                byte[] roundTrippedBuffer = ms.ToArray();
+                CbObject cb = new CbObject(roundTrippedBuffer);
+                CbField needsField = cb["needs"];
+                Assert.AreNotEqual(CbField.Empty, needsField);
+                List<BlobId> missingBlobs = needsField.AsArray().Select(field => BlobId.FromIoHash(field.AsHash())).ToList();
+                Assert.AreEqual(0, missingBlobs.Count);
+            }
+
+            {
+                BucketId bucket = new BucketId("bucket");
+
+                RefRecord objectRecord = await ReferencesStore.GetAsync(TestNamespace, bucket, key, IReferencesStore.FieldFlags.None);
+
+                Assert.IsTrue(objectRecord.IsFinalized);
+                Assert.AreEqual(key, objectRecord.Name);
+                Assert.AreEqual(objectHash, objectRecord.BlobIdentifier);
+                Assert.IsNull(objectRecord.InlinePayload);
+            }
+            
+            {
+                HttpResponseMessage getResponse = await _httpClient.GetAsync(new Uri($"api/v1/refs/{TestNamespace}/bucket/{key}.json?fields=name", UriKind.Relative));
+                getResponse.EnsureSuccessStatusCode();
+                await using MemoryStream ms = new MemoryStream();
+                await getResponse.Content.CopyToAsync(ms);
+
+                byte[] roundTrippedBuffer = ms.ToArray();
+
+                string s = Encoding.ASCII.GetString(roundTrippedBuffer);
+                JsonNode? node = JsonNode.Parse(s);
+                Assert.IsNotNull(node);
+                Assert.AreEqual("thisIsAField", node["stringField"]!.GetValue<string>());
+            }
+        }
+
+        [TestMethod]
+        public async Task PutLargeCompactBinaryAsync()
+        {
+            byte[] data = await File.ReadAllBytesAsync($"Objects/Payloads/lyra.cb");
+            BlobId objectHash = BlobId.FromBlob(data);
+            RefId key = RefId.FromName("largeCompactBinary");
+            using HttpContent requestContent = new ByteArrayContent(data);
+            requestContent.Headers.ContentType = new MediaTypeHeaderValue(CustomMediaTypeNames.UnrealCompactBinary);
+            requestContent.Headers.Add(CommonHeaders.HashHeaderName, objectHash.ToString());
+
+            HttpResponseMessage result = await _httpClient!.PutAsync(new Uri($"api/v1/refs/{TestNamespace}/bucket/{key}.uecb", UriKind.Relative), requestContent);
+            result.EnsureSuccessStatusCode();
+
+            CbObject cb = new CbObject(await result.Content.ReadAsByteArrayAsync());
+            CbArray needsField = cb["needs"].AsArray();
+            Assert.IsNotNull(needsField);
+            Assert.AreEqual(4924, needsField.Count);
+        }
+
+        [TestMethod]
+        public async Task PutGetCompactBinaryHierarchyAsync()
+        {
+            CbWriter childObjectWriter = new CbWriter();
+            childObjectWriter.BeginObject();
+            childObjectWriter.WriteString("stringField", "thisIsAField");
+            childObjectWriter.EndObject();
+            byte[] childObjectData = childObjectWriter.ToByteArray();
+            BlobId childObjectHash = BlobId.FromBlob(childObjectData);
+
+            CbWriter parentObjectWriter = new CbWriter();
+            parentObjectWriter.BeginObject();
+            parentObjectWriter.WriteObjectAttachment("childObject", childObjectHash.AsIoHash());
+            parentObjectWriter.EndObject();
+            byte[] parentObjectData = parentObjectWriter.ToByteArray();
+            BlobId parentObjectHash = BlobId.FromBlob(parentObjectData);
+
+            RefId key = RefId.FromName("newHierarchyObject");
+            // this first upload should fail with the child object missing
+            {
+                using HttpContent requestContent = new ByteArrayContent(parentObjectData);
+                requestContent.Headers.ContentType = new MediaTypeHeaderValue(CustomMediaTypeNames.UnrealCompactBinary);
+                requestContent.Headers.Add(CommonHeaders.HashHeaderName, parentObjectHash.ToString());
+
+                HttpResponseMessage result = await _httpClient!.PutAsync(new Uri($"api/v1/refs/{TestNamespace}/bucket/{key}.uecb", UriKind.Relative), requestContent);
+                result.EnsureSuccessStatusCode();
+
+                Assert.AreEqual(CustomMediaTypeNames.UnrealCompactBinary, result!.Content.Headers.ContentType!.MediaType);
+                // check that one blobs is missing
+                await using MemoryStream ms = new MemoryStream();
+                await result.Content.CopyToAsync(ms);
+                byte[] roundTrippedBuffer = ms.ToArray();
+                CbObject cb = new CbObject(roundTrippedBuffer);
+                CbField needsField = cb["needs"];
+                Assert.AreNotEqual(CbField.Empty, needsField);
+                List<BlobId> missingBlobs = needsField.AsArray().Select(field => BlobId.FromIoHash(field.AsHash())).ToList();
+                Assert.AreEqual(1, missingBlobs.Count);
+                Assert.AreEqual(childObjectHash, missingBlobs[0]);
+            }
+
+            // upload the child object
+            {
+                using HttpRequestMessage request = new HttpRequestMessage(HttpMethod.Put, new Uri($"api/v1/objects/{TestNamespace}/{childObjectHash}", UriKind.Relative));
+                request.Headers.Accept.Add(new MediaTypeWithQualityHeaderValue(CustomMediaTypeNames.UnrealCompactBinary));
+
+                HttpContent requestContent = new ByteArrayContent(childObjectData);
+                requestContent.Headers.ContentType = new MediaTypeHeaderValue(CustomMediaTypeNames.UnrealCompactBinary);
+                requestContent.Headers.Add(CommonHeaders.HashHeaderName, childObjectHash.ToString());
+                request.Content = requestContent;
+
+                HttpResponseMessage result = await _httpClient!.SendAsync(request);
+                result.EnsureSuccessStatusCode();
+
+                Assert.AreEqual(CustomMediaTypeNames.UnrealCompactBinary, result!.Content.Headers.ContentType!.MediaType);
+                // check that one blobs is missing
+                await using MemoryStream ms = new MemoryStream();
+                await result.Content.CopyToAsync(ms);
+                byte[] roundTrippedBuffer = ms.ToArray();
+                CbObject cb = new CbObject(roundTrippedBuffer);
+                CbField value = cb["identifier"];
+                Assert.AreNotEqual(CbField.Empty, value);
+                Assert.AreEqual(childObjectHash, BlobId.FromIoHash(value.AsHash()));
+            }
+
+            // since we have now uploaded the child object putting the object again should result in no missing references
+            {
+                using HttpContent requestContent = new ByteArrayContent(parentObjectData);
+                requestContent.Headers.ContentType = new MediaTypeHeaderValue(CustomMediaTypeNames.UnrealCompactBinary);
+                requestContent.Headers.Add(CommonHeaders.HashHeaderName, parentObjectHash.ToString());
+
+                HttpResponseMessage result = await _httpClient!.PutAsync(new Uri($"api/v1/refs/{TestNamespace}/bucket/{key}.uecb", UriKind.Relative), requestContent);
+                result.EnsureSuccessStatusCode();
+
+                Assert.AreEqual(result!.Content.Headers.ContentType!.MediaType, CustomMediaTypeNames.UnrealCompactBinary);
+                // check that one blobs is missing
+                await using MemoryStream ms = new MemoryStream();
+                await result.Content.CopyToAsync(ms);
+                byte[] roundTrippedBuffer = ms.ToArray();
+                CbObject cb = new CbObject(roundTrippedBuffer);
+                CbField needsField = cb["needs"];
+                List<BlobId> missingBlobs = needsField.AsArray().Select(field => BlobId.FromIoHash(field.AsHash())).ToList();
+                Assert.AreEqual(0, missingBlobs.Count);
+            }
+
+            {
+                BucketId bucket = new BucketId("bucket");
+
+                RefRecord objectRecord = await ReferencesStore.GetAsync(TestNamespace, bucket, key, IReferencesStore.FieldFlags.IncludePayload);
+
+                Assert.IsTrue(objectRecord.IsFinalized);
+                Assert.AreEqual(key, objectRecord.Name);
+                Assert.AreEqual(parentObjectHash, objectRecord.BlobIdentifier);
+                Assert.IsNotNull(objectRecord.InlinePayload);
+            }
+
+            {
+                HttpResponseMessage getResponse = await _httpClient.GetAsync(new Uri($"api/v1/refs/{TestNamespace}/bucket/{key}.raw", UriKind.Relative));
+                getResponse.EnsureSuccessStatusCode();
+                await using MemoryStream ms = new MemoryStream();
+                await getResponse.Content.CopyToAsync(ms);
+
+                byte[] roundTrippedBuffer = ms.ToArray();
+
+                CollectionAssert.AreEqual(parentObjectData, roundTrippedBuffer);
+                Assert.AreEqual(parentObjectHash, BlobId.FromBlob(roundTrippedBuffer));
+            }
+
+            {
+                HttpResponseMessage getResponse = await _httpClient.GetAsync(new Uri($"api/v1/refs/{TestNamespace}/bucket/{key}.uecb", UriKind.Relative));
+                getResponse.EnsureSuccessStatusCode();
+                await using MemoryStream ms = new MemoryStream();
+                await getResponse.Content.CopyToAsync(ms);
+
+                byte[] roundTrippedBuffer = ms.ToArray();
+                CbObject cb = new CbObject(roundTrippedBuffer);
+                List<CbField> fields = cb.ToList();
+
+                Assert.AreEqual(1, fields.Count);
+                CbField childObjectField = fields[0];
+                Assert.AreEqual("childObject", childObjectField.Name);
+                Assert.AreEqual(childObjectHash, BlobId.FromIoHash(childObjectField.AsHash()));
+            }
+
+            {
+                HttpResponseMessage getResponse = await _httpClient.GetAsync(new Uri($"api/v1/refs/{TestNamespace}/bucket/{key}.json", UriKind.Relative));
+                getResponse.EnsureSuccessStatusCode();
+                await using MemoryStream ms = new MemoryStream();
+                await getResponse.Content.CopyToAsync(ms);
+
+                byte[] roundTrippedBuffer = ms.ToArray();
+
+                string s = Encoding.ASCII.GetString(roundTrippedBuffer);
+                JsonNode? node = JsonNode.Parse(s);
+                Assert.IsNotNull(node);
+                Assert.AreEqual(childObjectHash.ToString().ToLower(), node["childObject"]!.GetValue<string>());
+            }
+
+            {
+                HttpResponseMessage getResponse = await _httpClient.GetAsync(new Uri($"api/v1/objects/{TestNamespace}/{childObjectHash}", UriKind.Relative));
+                getResponse.EnsureSuccessStatusCode();
+                await using MemoryStream ms = new MemoryStream();
+                await getResponse.Content.CopyToAsync(ms);
+
+                byte[] roundTrippedBuffer = ms.ToArray();
+
+                CollectionAssert.AreEqual(childObjectData, roundTrippedBuffer);
+                Assert.AreEqual(childObjectHash, BlobId.FromBlob(roundTrippedBuffer));
+            }
+        }
+
+        [TestMethod]
+        public async Task ExistsChecksAsync()
+        {
+            const string objectContents = "This is treated as a opaque blob";
+            byte[] data = Encoding.ASCII.GetBytes(objectContents);
+            BlobId objectHash = BlobId.FromBlob(data);
+            RefId key = RefId.FromName("newObject");
+            using HttpContent requestContent = new ByteArrayContent(data);
+            requestContent.Headers.ContentType = new MediaTypeHeaderValue(MediaTypeNames.Application.Octet);
+            requestContent.Headers.Add(CommonHeaders.HashHeaderName, objectHash.ToString());
+
+            {
+                HttpResponseMessage result = await _httpClient!.PutAsync(new Uri($"api/v1/refs/{TestNamespace}/bucket/{key}", UriKind.Relative), requestContent);
+                result.EnsureSuccessStatusCode();
+            }
+
+            {
+                using HttpRequestMessage message = new(HttpMethod.Head, new Uri($"api/v1/refs/{TestNamespace}/bucket/{key}", UriKind.Relative));
+                HttpResponseMessage result = await _httpClient!.SendAsync(message);
+                result.EnsureSuccessStatusCode();
+            }
+
+            {
+                using HttpRequestMessage message = new(HttpMethod.Head, new Uri($"api/v1/refs/{TestNamespace}/bucket/{RefId.FromName("missingObject")}", UriKind.Relative));
+                HttpResponseMessage result = await _httpClient!.SendAsync(message);
+                Assert.AreEqual(HttpStatusCode.NotFound, result.StatusCode);
+            }
+        }
+
+        [TestMethod]
+        public async Task ExistsChecksMultipleAsync()
+        {
+            BucketId bucket = new BucketId("bucket");
+            RefId existingObject = RefId.FromName("existingObject");
+
+            const string objectContents = "This is treated as a opaque blob";
+            byte[] data = Encoding.ASCII.GetBytes(objectContents);
+            BlobId objectHash = BlobId.FromBlob(data);
+            using HttpContent requestContent = new ByteArrayContent(data);
+            requestContent.Headers.ContentType = new MediaTypeHeaderValue(MediaTypeNames.Application.Octet);
+            requestContent.Headers.Add(CommonHeaders.HashHeaderName, objectHash.ToString());
+
+            {
+                HttpResponseMessage result = await _httpClient!.PutAsync(new Uri($"api/v1/refs/{TestNamespace}/{bucket}/{existingObject}", UriKind.Relative), requestContent);
+                result.EnsureSuccessStatusCode();
+            }
+
+            RefId missingObject = RefId.FromName("missingObject");
+
+            string queryString = $"?names={bucket}.{existingObject}&names={bucket}.{missingObject}";
+
+            {
+                HttpResponseMessage result = await _httpClient!.GetAsync(new Uri($"api/v1/refs/{TestNamespace}/exists" + queryString, UriKind.Relative));
+                result.EnsureSuccessStatusCode();
+                ExistCheckMultipleRefsResponse? response = await result.Content.ReadFromJsonAsync<ExistCheckMultipleRefsResponse>();
+
+                Assert.IsNotNull(response);
+                Assert.AreEqual(1, response.Missing.Count);
+                Assert.AreEqual(bucket, response.Missing[0].Bucket);
+                Assert.AreEqual(missingObject, response.Missing[0].Key);
+            }
+        }
+
+        [TestMethod]
+        public async Task PutGetObjectHierarchyAsync()
+        {
+            string blobContents = "This is a string that is referenced as a blob";
+            byte[] blobData = Encoding.ASCII.GetBytes(blobContents);
+            BlobId blobHash = BlobId.FromBlob(blobData);
+            await Service.PutObjectAsync(TestNamespace, blobData, blobHash);
+
+            string blobContentsChild = "This string is also referenced as a blob but from a child object";
+            byte[] dataChild = Encoding.ASCII.GetBytes(blobContentsChild);
+            BlobId blobHashChild = BlobId.FromBlob(dataChild);
+            await Service.PutObjectAsync(TestNamespace, dataChild, blobHashChild);
+
+            CbWriter writerChild = new CbWriter();
+            writerChild.BeginObject();
+            writerChild.WriteBinaryAttachment("blob", blobHashChild.AsIoHash());
+            writerChild.EndObject();
+
+            byte[] childDataObject = writerChild.ToByteArray();
+            BlobId childDataObjectHash = BlobId.FromBlob(childDataObject);
+            await Service.PutObjectAsync(TestNamespace, childDataObject, childDataObjectHash);
+
+            CbWriter writerParent = new CbWriter();
+            writerParent.BeginObject();
+            
+            writerParent.WriteBinaryAttachment("blobAttachment", blobHash.AsIoHash());
+            writerParent.WriteObjectAttachment("objectAttachment", childDataObjectHash.AsIoHash());
+            writerParent.EndObject();
+
+            byte[] objectData = writerParent.ToByteArray();
+            BlobId objectHash = BlobId.FromBlob(objectData);
+
+            RefId key = RefId.FromName("newHierarchyObject");
+
+            using HttpContent requestContent = new ByteArrayContent(objectData);
+            requestContent.Headers.ContentType = new MediaTypeHeaderValue(CustomMediaTypeNames.UnrealCompactBinary);
+            requestContent.Headers.Add(CommonHeaders.HashHeaderName, objectHash.ToString());
+
+            HttpResponseMessage result = await _httpClient!.PutAsync(new Uri($"api/v1/refs/{TestNamespace}/bucket/{key}.uecb", UriKind.Relative), requestContent);
+            result.EnsureSuccessStatusCode();
+
+            // check the response
+            {
+                Assert.AreEqual(result!.Content.Headers.ContentType!.MediaType, CustomMediaTypeNames.UnrealCompactBinary);
+                // check that no blobs are missing
+                await using MemoryStream ms = new MemoryStream();
+                await result.Content.CopyToAsync(ms);
+                byte[] roundTrippedBuffer = ms.ToArray();
+                CbObject cb = new CbObject(roundTrippedBuffer);
+                CbField needsField = cb["needs"];
+                List<BlobId> missingBlobs = needsField.AsArray().Select(field => BlobId.FromIoHash(field.AsHash())).ToList();
+                Assert.AreEqual(0, missingBlobs.Count);
+            }
+
+            // check that actual internal representation
+            {
+                BucketId bucket = new BucketId("bucket");
+
+                RefRecord objectRecord = await ReferencesStore.GetAsync(TestNamespace, bucket, key, IReferencesStore.FieldFlags.IncludePayload);
+
+                Assert.IsTrue(objectRecord.IsFinalized);
+                Assert.AreEqual(key, objectRecord.Name);
+                Assert.AreEqual(objectHash, objectRecord.BlobIdentifier);
+                Assert.IsNotNull(objectRecord.InlinePayload);
+            }
+
+            // verify attachments
+            {
+                HttpResponseMessage getResponse = await _httpClient.GetAsync(new Uri($"api/v1/refs/{TestNamespace}/bucket/{key}.raw", UriKind.Relative));
+                getResponse.EnsureSuccessStatusCode();
+                await using MemoryStream ms = new MemoryStream();
+                await getResponse.Content.CopyToAsync(ms);
+
+                byte[] roundTrippedBuffer = ms.ToArray();
+
+                CollectionAssert.AreEqual(objectData, roundTrippedBuffer);
+                Assert.AreEqual(objectHash, BlobId.FromBlob(roundTrippedBuffer));
+            }
+
+            {
+                BlobId blobAttachment;
+                BlobId objectAttachment;
+                {
+                    HttpResponseMessage getResponse = await _httpClient.GetAsync(new Uri($"api/v1/refs/{TestNamespace}/bucket/{key}.uecb", UriKind.Relative));
+                    getResponse.EnsureSuccessStatusCode();
+                    await using MemoryStream ms = new MemoryStream();
+                    await getResponse.Content.CopyToAsync(ms);
+
+                    byte[] roundTrippedBuffer = ms.ToArray();
+                    ReadOnlyMemory<byte> localMemory = new ReadOnlyMemory<byte>(roundTrippedBuffer);
+                    CbObject cb = new CbObject(roundTrippedBuffer);
+                    Assert.AreEqual(2, cb.Count());
+
+                    CbField blobAttachmentField = cb["blobAttachment"];
+                    Assert.AreNotEqual(CbField.Empty, blobAttachmentField);
+                    blobAttachment = BlobId.FromIoHash(blobAttachmentField.AsBinaryAttachment());
+                    CbField objectAttachmentField = cb["objectAttachment"];
+                    Assert.AreNotEqual(CbField.Empty, objectAttachmentField);
+                    objectAttachment = BlobId.FromIoHash(objectAttachmentField.AsObjectAttachment().Hash);
+                }
+
+                {
+                    HttpResponseMessage getAttachment = await _httpClient.GetAsync(new Uri($"api/v1/blobs/{TestNamespace}/{blobAttachment}", UriKind.Relative));
+                    getAttachment.EnsureSuccessStatusCode();
+                    await using MemoryStream ms = new MemoryStream();
+                    await getAttachment.Content.CopyToAsync(ms);
+                    byte[] roundTrippedBuffer = ms.ToArray();
+                    string roundTrippedString = Encoding.ASCII.GetString(roundTrippedBuffer);
+
+                    Assert.AreEqual(blobContents, roundTrippedString);
+                    Assert.AreEqual(blobHash, BlobId.FromBlob(roundTrippedBuffer));
+                }
+
+                BlobId attachedBlobIdentifier;
+                {
+                    HttpResponseMessage getAttachment = await _httpClient.GetAsync(new Uri($"api/v1/blobs/{TestNamespace}/{objectAttachment}", UriKind.Relative));
+                    getAttachment.EnsureSuccessStatusCode();
+                    await using MemoryStream ms = new MemoryStream();
+                    await getAttachment.Content.CopyToAsync(ms);
+                    byte[] roundTrippedBuffer = ms.ToArray();
+                    ReadOnlyMemory<byte> localMemory = new ReadOnlyMemory<byte>(roundTrippedBuffer);
+                    CbObject cb = new CbObject(roundTrippedBuffer);
+                    Assert.AreEqual(1, cb.Count());
+
+                    CbField blobField = cb["blob"];
+                    Assert.AreNotEqual(CbField.Empty, blobField);
+
+                    attachedBlobIdentifier = BlobId.FromIoHash(blobField!.AsBinaryAttachment());
+                }
+
+                {
+                    HttpResponseMessage getAttachment = await _httpClient.GetAsync(new Uri($"api/v1/blobs/{TestNamespace}/{attachedBlobIdentifier}", UriKind.Relative));
+                    getAttachment.EnsureSuccessStatusCode();
+                    await using MemoryStream ms = new MemoryStream();
+                    await getAttachment.Content.CopyToAsync(ms);
+                    byte[] roundTrippedBuffer = ms.ToArray();
+                    string roundTrippedString = Encoding.ASCII.GetString(roundTrippedBuffer);
+
+                    Assert.AreEqual(blobContentsChild, roundTrippedString);
+                    Assert.AreEqual(blobHashChild, BlobId.FromBlob(roundTrippedBuffer));
+                }
+            }
+
+            // check json representation
+            {
+                HttpResponseMessage getResponse = await _httpClient.GetAsync(new Uri($"api/v1/refs/{TestNamespace}/bucket/{key}.json", UriKind.Relative));
+                getResponse.EnsureSuccessStatusCode();
+                await using MemoryStream ms = new MemoryStream();
+                await getResponse.Content.CopyToAsync(ms);
+
+                byte[] roundTrippedBuffer = ms.ToArray();
+                string s = Encoding.ASCII.GetString(roundTrippedBuffer);
+                JsonNode? node = JsonNode.Parse(s);
+                Assert.IsNotNull(node);
+                Assert.AreEqual(blobHash, BlobId.Parse(node["blobAttachment"]!.GetValue<string>()));
+                Assert.AreEqual(childDataObjectHash, BlobId.Parse(node["objectAttachment"]!.GetValue<string>()));
+            }
+        }
+        
+        [TestMethod]
+        public async Task PutPartialHierarchyAsync()
+        {
+            // do not submit the content of the blobs, which should be reported in the response of the put
+            string blobContents = "This is a string that is referenced as a blob";
+            byte[] blobData = Encoding.ASCII.GetBytes(blobContents);
+            BlobId blobHash = BlobId.FromBlob(blobData);
+
+            string blobContentsChild = "This string is also referenced as a blob but from a child object";
+            byte[] dataChild = Encoding.ASCII.GetBytes(blobContentsChild);
+            BlobId blobHashChild = BlobId.FromBlob(dataChild);
+
+            CbWriter writerChild = new CbWriter();
+            writerChild.BeginObject();
+            writerChild.WriteBinaryAttachment("blob", blobHashChild.AsIoHash());
+            writerChild.EndObject();
+
+            byte[] childDataObject = writerChild.ToByteArray();
+            BlobId childDataObjectHash = BlobId.FromBlob(childDataObject);
+            await Service.PutObjectAsync(TestNamespace, childDataObject, childDataObjectHash);
+
+            CbWriter writerParent = new CbWriter();
+            writerParent.BeginObject();
+            
+            writerParent.WriteBinaryAttachment("blobAttachment", blobHash.AsIoHash());
+            writerParent.WriteObjectAttachment("objectAttachment", childDataObjectHash.AsIoHash());
+            writerParent.EndObject();
+
+            byte[] objectData = writerParent.ToByteArray();
+            BlobId objectHash = BlobId.FromBlob(objectData);
+            
+            RefId key = RefId.FromName("newHierarchyObject");
+
+            {
+                using HttpContent requestContent = new ByteArrayContent(objectData);
+                requestContent.Headers.ContentType = new MediaTypeHeaderValue(CustomMediaTypeNames.UnrealCompactBinary);
+                requestContent.Headers.Add(CommonHeaders.HashHeaderName, objectHash.ToString());
+
+                HttpResponseMessage result = await _httpClient!.PutAsync(new Uri($"api/v1/refs/{TestNamespace}/bucket/{key}.uecb", UriKind.Relative), requestContent);
+                result.EnsureSuccessStatusCode();
+
+                {
+                    Assert.AreEqual(result!.Content.Headers.ContentType!.MediaType, CustomMediaTypeNames.UnrealCompactBinary);
+                    await using MemoryStream ms = new MemoryStream();
+                    await result.Content.CopyToAsync(ms);
+                    byte[] roundTrippedBuffer = ms.ToArray();
+                    ReadOnlyMemory<byte> localMemory = new ReadOnlyMemory<byte>(roundTrippedBuffer);
+                    CbObject cb = new CbObject(roundTrippedBuffer);
+                    CbField needsField = cb["needs"];
+                    List<BlobId> missingBlobs = needsField.AsArray().Select(field => BlobId.FromIoHash(field.AsHash())).ToList();
+                    Assert.AreEqual(2, missingBlobs.Count);
+                    Assert.IsTrue(missingBlobs.Contains(blobHash));
+                    Assert.IsTrue(missingBlobs.Contains(blobHashChild));
+                }
+            }
+
+            {
+                using HttpContent requestContent = new ByteArrayContent(objectData);
+                requestContent.Headers.ContentType = new MediaTypeHeaderValue(CustomMediaTypeNames.UnrealCompactBinary);
+                requestContent.Headers.Add(CommonHeaders.HashHeaderName, objectHash.ToString());
+
+                HttpResponseMessage result = await _httpClient!.PutAsync( new Uri($"api/v1/refs/{TestNamespace}/bucket/{key}.json", UriKind.Relative), requestContent);
+                result.EnsureSuccessStatusCode();
+
+                {
+                    Assert.AreEqual(result!.Content.Headers.ContentType!.MediaType, MediaTypeNames.Application.Json);
+                    await using MemoryStream ms = new MemoryStream();
+                    await result.Content.CopyToAsync(ms);
+                    byte[] roundTrippedBuffer = ms.ToArray();
+                    string s = Encoding.ASCII.GetString(roundTrippedBuffer);
+
+					JsonSerializerOptions jsonSerializerOptions = new JsonSerializerOptions();
+					Startup.ConfigureJsonSerializer(jsonSerializerOptions);
+
+                    PutObjectResponse? response = JsonSerializer.Deserialize<PutObjectResponse>(s, jsonSerializerOptions);
+                    Assert.IsNotNull(response);
+
+                    BlobId[] missingBlobs = response.Needs.Select(hash => new BlobId(hash)).ToArray();
+
+                    Assert.AreEqual(2, missingBlobs.Length);
+                    Assert.IsTrue(missingBlobs.Contains(blobHash));
+                    Assert.IsTrue(missingBlobs.Contains(blobHashChild));
+                }
+            }
+        }
+
+        
+        [TestMethod]
+        public async Task PutContentIdMissingBlobAsync()
+        {
+            IContentIdService? contentIdStore = ServiceProvider.GetService<IContentIdService>();
+            Assert.IsNotNull(contentIdStore);
+
+            // submit a object which contains a content id, which exists but points to a blob that does not exist
+            string blobContents = "This is a string that is referenced as a blob";
+            byte[] blobData = Encoding.ASCII.GetBytes(blobContents);
+            BlobId blobHash = BlobId.FromBlob(blobData);
+            ContentId contentId = ContentId.Parse("0000000000000000000000000000000000000000");
+
+            await contentIdStore.PutAsync(TestNamespace, contentId, blobHash, blobData.Length);
+
+            CbWriter writer = new CbWriter();
+            writer.BeginObject();
+            writer.WriteBinaryAttachment("blob", contentId.AsIoHash());
+            writer.EndObject();
+
+            byte[] objectData = writer.ToByteArray();
+            BlobId objectHash = BlobId.FromBlob(objectData);
+            
+            RefId key = RefId.FromName("putContentIdMissingBlob");
+
+            {
+                using HttpContent requestContent = new ByteArrayContent(objectData);
+                requestContent.Headers.ContentType = new MediaTypeHeaderValue(CustomMediaTypeNames.UnrealCompactBinary);
+                requestContent.Headers.Add(CommonHeaders.HashHeaderName, objectHash.ToString());
+
+                HttpResponseMessage result = await _httpClient!.PutAsync(new Uri($"api/v1/refs/{TestNamespace}/bucket/{key}.uecb", UriKind.Relative), requestContent);
+                result.EnsureSuccessStatusCode();
+
+                {
+                    Assert.AreEqual(result!.Content.Headers.ContentType!.MediaType, CustomMediaTypeNames.UnrealCompactBinary);
+                    await using MemoryStream ms = new MemoryStream();
+                    await result.Content.CopyToAsync(ms);
+                    byte[] roundTrippedBuffer = ms.ToArray();
+                    CbObject cb = new CbObject(roundTrippedBuffer);
+                    CbField needsField = cb["needs"];
+                    Assert.AreNotEqual(CbField.Empty, needsField);
+                    List<IoHash> missingBlobs = needsField.AsArray().Select(field => field.AsHash()).ToList();
+                    Assert.AreEqual(1, missingBlobs.Count);
+
+                    Assert.AreNotEqual(blobHash, missingBlobs[0], "Refs should not be returning the mapped blob identifiers as this is unknown to the client attempting to put a new ref");
+                    Assert.AreEqual(contentId.AsBlobIdentifier(), BlobId.FromIoHash(missingBlobs[0]));
+                }
+            }
+
+            {
+                using HttpContent requestContent = new ByteArrayContent(objectData);
+                requestContent.Headers.ContentType = new MediaTypeHeaderValue(CustomMediaTypeNames.UnrealCompactBinary);
+                requestContent.Headers.Add(CommonHeaders.HashHeaderName, objectHash.ToString());
+
+                HttpResponseMessage result = await _httpClient!.PutAsync(new Uri($"api/v1/refs/{TestNamespace}/bucket/{key}.json", UriKind.Relative), requestContent);
+                result.EnsureSuccessStatusCode();
+
+                {
+                    Assert.AreEqual(result!.Content.Headers.ContentType!.MediaType, MediaTypeNames.Application.Json);
+                    await using MemoryStream ms = new MemoryStream();
+                    await result.Content.CopyToAsync(ms);
+                    byte[] roundTrippedBuffer = ms.ToArray();
+                    string s = Encoding.ASCII.GetString(roundTrippedBuffer);
+
+					JsonSerializerOptions jsonSerializerOptions = new JsonSerializerOptions();
+					Startup.ConfigureJsonSerializer(jsonSerializerOptions);
+
+                    PutObjectResponse? response = JsonSerializer.Deserialize<PutObjectResponse>(s, jsonSerializerOptions);
+                    Assert.IsNotNull(response);
+
+                    BlobId[] missingBlobs = response.Needs.Select(field => new BlobId(field)).ToArray();
+
+                    Assert.AreEqual(1, missingBlobs.Length);
+                    Assert.AreEqual(contentId.AsBlobIdentifier(), missingBlobs[0]);
+                }
+            }
+        }
+
+        
+        [TestMethod]
+        public async Task PutMissingAttachmentComplexAsync()
+        {
+            string blobContents = "This is a string that is referenced as a blob but will not be uploaded";
+            byte[] blobData = Encoding.ASCII.GetBytes(blobContents);
+            BlobId blobHash = BlobId.FromBlob(blobData);
+
+            CbWriter writer = new CbWriter();
+            writer.BeginObject();
+            writer.WriteString("name", "randomStringContent");
+            writer.BeginArray("values");
+            writer.BeginObject();
+            writer.WriteInteger("rawSize", 200);
+            writer.WriteBinaryAttachment("rawHash", blobHash.AsIoHash());
+            writer.EndObject();
+            writer.EndArray();
+            writer.EndObject();
+
+            byte[] objectData = writer.ToByteArray();
+            BlobId objectHash = BlobId.FromBlob(objectData);
+
+            RefId key = RefId.FromName("putContentIdMissingBlobComplex");
+
+            {
+                using HttpContent requestContent = new ByteArrayContent(objectData);
+                requestContent.Headers.ContentType = new MediaTypeHeaderValue(CustomMediaTypeNames.UnrealCompactBinary);
+                requestContent.Headers.Add(CommonHeaders.HashHeaderName, objectHash.ToString());
+
+                HttpResponseMessage result = await _httpClient!.PutAsync(new Uri($"api/v1/refs/{TestNamespace}/bucket/{key}.uecb", UriKind.Relative), requestContent);
+                result.EnsureSuccessStatusCode();
+
+                {
+                    Assert.AreEqual(result!.Content.Headers.ContentType!.MediaType, CustomMediaTypeNames.UnrealCompactBinary);
+                    await using MemoryStream ms = new MemoryStream();
+                    await result.Content.CopyToAsync(ms);
+                    byte[] roundTrippedBuffer = ms.ToArray();
+                    CbObject cb = new CbObject(roundTrippedBuffer);
+                    CbField needsField = cb["needs"];
+                    Assert.AreNotEqual(CbField.Empty, needsField);
+                    List<IoHash> missingBlobs = needsField.AsArray().Select(field => field.AsHash()).ToList();
+                    Assert.AreEqual(1, missingBlobs.Count);
+
+                    Assert.AreEqual(blobHash.AsIoHash(), missingBlobs[0], "Expected to find missing hash even for attachments not directly in the root object");
+                }
+            }
+        }
+
+        [TestMethod]
+        public async Task PutAndFinalizeAsync()
+        {
+            BucketId bucket = new BucketId("bucket");
+            RefId key = RefId.FromName("willFinalizeObject");
+
+            // do not submit the content of the blobs, which should be reported in the response of the put
+            string blobContents = "This is a string that is referenced as a blob";
+            byte[] blobData = Encoding.ASCII.GetBytes(blobContents);
+            BlobId blobHash = BlobId.FromBlob(blobData);
+
+            string blobContentsChild = "This string is also referenced as a blob but from a child object";
+            byte[] dataChild = Encoding.ASCII.GetBytes(blobContentsChild);
+            BlobId blobHashChild = BlobId.FromBlob(dataChild);
+
+            CbWriter writerChild = new CbWriter();
+            writerChild.BeginObject();
+            writerChild.WriteBinaryAttachment("blob", blobHashChild.AsIoHash());
+            writerChild.EndObject();
+
+            byte[] childDataObject = writerChild.ToByteArray();
+            BlobId childDataObjectHash = BlobId.FromBlob(childDataObject);
+            await Service.PutObjectAsync(TestNamespace, childDataObject, childDataObjectHash);
+
+            CbWriter writerParent = new CbWriter();
+            writerParent.BeginObject();
+            
+            writerParent.WriteBinaryAttachment("blobAttachment", blobHash.AsIoHash());
+            writerParent.WriteObjectAttachment("objectAttachment", childDataObjectHash.AsIoHash());
+            writerParent.EndObject();
+
+            byte[] objectData = writerParent.ToByteArray();
+            BlobId objectHash = BlobId.FromBlob(objectData);
+
+            {
+                using HttpContent requestContent = new ByteArrayContent(objectData);
+                requestContent.Headers.ContentType = new MediaTypeHeaderValue(CustomMediaTypeNames.UnrealCompactBinary);
+                requestContent.Headers.Add(CommonHeaders.HashHeaderName, objectHash.ToString());
+
+                HttpResponseMessage result = await _httpClient!.PutAsync(new Uri($"api/v1/refs/{TestNamespace}/bucket/{key}.uecb", UriKind.Relative), requestContent);
+                result.EnsureSuccessStatusCode();
+
+                {
+                    Assert.AreEqual(result!.Content.Headers.ContentType!.MediaType, CustomMediaTypeNames.UnrealCompactBinary);
+                    await using MemoryStream ms = new MemoryStream();
+                    await result.Content.CopyToAsync(ms);
+                    byte[] roundTrippedBuffer = ms.ToArray();
+                    CbObject cb = new CbObject(roundTrippedBuffer);
+                    CbField needsField = cb["needs"];
+                    List<BlobId> missingBlobs = needsField.AsArray().Select(field => BlobId.FromIoHash(field.AsHash())).ToList();
+                    Assert.AreEqual(2, missingBlobs.Count);
+                    Assert.IsTrue(missingBlobs.Contains(blobHash));
+                    Assert.IsTrue(missingBlobs.Contains(blobHashChild));
+                }
+            }
+
+            // check that actual internal representation
+            {
+                RefRecord objectRecord = await ReferencesStore.GetAsync(TestNamespace, bucket, key, IReferencesStore.FieldFlags.IncludePayload);
+
+                Assert.IsFalse(objectRecord.IsFinalized);
+                Assert.AreEqual(key, objectRecord.Name);
+                Assert.AreEqual(objectHash, objectRecord.BlobIdentifier);
+                Assert.IsNotNull(objectRecord.InlinePayload);
+            }
+
+            // upload missing pieces
+            {
+                await Service.PutObjectAsync(TestNamespace, blobData, blobHash);
+                await Service.PutObjectAsync(TestNamespace, dataChild, blobHashChild);
+            }
+
+            // finalize the object as no pieces is now missing
+            {
+                using HttpContent requestContent = new ByteArrayContent(Array.Empty<byte>());
+
+                HttpResponseMessage result = await _httpClient!.PostAsync(new Uri($"api/v1/refs/{TestNamespace}/bucket/{key}/finalize/{objectHash}.uecb", UriKind.Relative), requestContent);
+                result.EnsureSuccessStatusCode();
+
+                {
+                    Assert.AreEqual(result!.Content.Headers.ContentType!.MediaType, CustomMediaTypeNames.UnrealCompactBinary);
+                    await using MemoryStream ms = new MemoryStream();
+                    await result.Content.CopyToAsync(ms);
+                    byte[] roundTrippedBuffer = ms.ToArray();
+                    ReadOnlyMemory<byte> localMemory = new ReadOnlyMemory<byte>(roundTrippedBuffer);
+                    CbObject cb = new CbObject(roundTrippedBuffer);
+                    CbField? needsField = cb["needs"];
+                    List<BlobId> missingBlobs = needsField.AsArray().Select(field => BlobId.FromIoHash(field.AsHash())).ToList();
+                    Assert.AreEqual(0, missingBlobs.Count);
+                }
+            }
+
+            // check that actual internal representation has updated its state
+            {
+                RefRecord objectRecord = await ReferencesStore.GetAsync(TestNamespace, bucket, key, IReferencesStore.FieldFlags.None);
+
+                Assert.IsTrue(objectRecord.IsFinalized);
+                Assert.AreEqual(key, objectRecord.Name);
+                Assert.AreEqual(objectHash, objectRecord.BlobIdentifier);
+            }
+        }
+
+        [TestMethod]
+        public async Task GetMissingContentIdRecordAsync()
+        {
+            string blobContents = "This is a blob";
+            byte[] blobData = Encoding.ASCII.GetBytes(blobContents);
+            BlobId uncompressedHash = BlobId.FromBlob(blobData);
+            RefId key = RefId.FromName("compressedObject");
+
+            CompressedBufferUtils bufferUtils = ServiceProvider.GetService<CompressedBufferUtils>()!;
+            using MemoryStream compressedStream = new MemoryStream();
+            bufferUtils.CompressContent(compressedStream, OoodleCompressorMethod.Mermaid, OoodleCompressionLevel.VeryFast, blobData);
+            byte[] compressedBuffer = compressedStream.ToArray();
+            BlobId compressedHash = BlobId.FromBlob(compressedBuffer);
+
+            CbObject cbObject = CbObject.Build(writer => writer.WriteBinaryAttachment("Attachment", uncompressedHash.AsIoHash()));
+            byte[] cbObjectData = cbObject.GetView().ToArray();
+            BlobId cbObjectHash = BlobId.FromBlob(cbObjectData);
+
+            {
+                // upload compressed blob
+                using HttpContent requestContent = new ByteArrayContent(compressedBuffer);
+                requestContent.Headers.ContentType = new MediaTypeHeaderValue(CustomMediaTypeNames.UnrealCompressedBuffer);
+
+                HttpResponseMessage result = await _httpClient!.PutAsync(new Uri($"api/v1/compressed-blobs/{TestNamespace}/{uncompressedHash}", UriKind.Relative), requestContent);
+                result.EnsureSuccessStatusCode();
+            }
+
+            {
+                // upload ref
+                using HttpContent requestContent = new ByteArrayContent(cbObjectData);
+                requestContent.Headers.ContentType = new MediaTypeHeaderValue(CustomMediaTypeNames.UnrealCompactBinary);
+                requestContent.Headers.Add(CommonHeaders.HashHeaderName, cbObjectHash.ToString());
+
+                HttpResponseMessage result = await _httpClient!.PutAsync(new Uri($"api/v1/refs/{TestNamespace}/bucket/{key}.uecb", UriKind.Relative), requestContent);
+                result.EnsureSuccessStatusCode();
+
+                {
+                    Assert.AreEqual(CustomMediaTypeNames.UnrealCompactBinary, result!.Content.Headers.ContentType!.MediaType);
+                    // check that no blobs are missing
+                    await using MemoryStream ms = new MemoryStream();
+                    await result.Content.CopyToAsync(ms);
+                    byte[] roundTrippedBuffer = ms.ToArray();
+                    CbObject cb = new CbObject(roundTrippedBuffer);
+                    CbField needsField = cb["needs"];
+                    List<BlobId> missingBlobs = needsField.AsArray().Select(field => BlobId.FromIoHash(field.AsHash())).ToList();
+                    Assert.AreEqual(0, missingBlobs.Count);
+                }
+            }
+
+            {
+                // verify we can fetch the blob properly
+                HttpResponseMessage getResponse = await _httpClient.GetAsync(new Uri($"api/v1/refs/{TestNamespace}/bucket/{key}.uecb", UriKind.Relative));
+                getResponse.EnsureSuccessStatusCode();
+                await using MemoryStream ms = new MemoryStream();
+                await getResponse.Content.CopyToAsync(ms);
+
+                byte[] roundTrippedBuffer = ms.ToArray();
+
+                CollectionAssert.AreEqual(cbObjectData, roundTrippedBuffer);
+                Assert.AreEqual(cbObjectHash, BlobId.FromBlob(roundTrippedBuffer));
+            }
+
+            {
+                // verify that head checks find the object
+                using HttpRequestMessage headRequest = new HttpRequestMessage(HttpMethod.Head, new Uri($"api/v1/refs/{TestNamespace}/bucket/{key}", UriKind.Relative));
+                HttpResponseMessage getResponse = await _httpClient.SendAsync(headRequest);
+                getResponse.EnsureSuccessStatusCode();
+            }
+
+            {
+                // verify that the exists check doesnt find any issue
+                HttpResponseMessage existsResponse = await _httpClient.GetAsync(new Uri($"api/v1/refs/{TestNamespace}/exists?names=bucket.{key}", UriKind.Relative));
+                existsResponse.EnsureSuccessStatusCode();
+
+                ExistCheckMultipleRefsResponse? response = await existsResponse.Content.ReadFromJsonAsync<ExistCheckMultipleRefsResponse>();
+                Assert.IsNotNull(response);
+                Assert.AreEqual(0, response.Missing.Count);
+            }
+
+            {
+                // delete the blob referenced by the compressed buffer
+                await Service.DeleteObjectAsync(TestNamespace, compressedHash);
+            }
+
+            {
+                // the compact binary object still exists, so this returns a success (trying to resolve the attachment will fail)
+                HttpResponseMessage getResponse = await _httpClient.GetAsync(new Uri($"api/v1/refs/{TestNamespace}/bucket/{key}.uecb", UriKind.Relative));
+                getResponse.EnsureSuccessStatusCode();
+            }
+
+            {
+                // verify that head checks fail
+                using HttpRequestMessage headRequest = new HttpRequestMessage(HttpMethod.Head, new Uri($"api/v1/refs/{TestNamespace}/bucket/{key}", UriKind.Relative));
+                HttpResponseMessage getResponse = await _httpClient.SendAsync(headRequest);
+                Assert.AreEqual(HttpStatusCode.NotFound, getResponse.StatusCode);
+            }
+
+            {
+                // verify that the exists check correctly returns a missing blob
+                HttpResponseMessage existsResponse = await _httpClient.GetAsync(new Uri($"api/v1/refs/{TestNamespace}/exists?names=bucket.{key}", UriKind.Relative));
+                existsResponse.EnsureSuccessStatusCode();
+
+                ExistCheckMultipleRefsResponse? response = await existsResponse.Content.ReadFromJsonAsync<ExistCheckMultipleRefsResponse>();
+                Assert.IsNotNull(response);
+                Assert.AreEqual(1, response.Missing.Count);
+                Assert.AreEqual(key, response.Missing[0].Key);
+                Assert.AreEqual("bucket", response.Missing[0].Bucket.ToString());
+            }
+        }
+
+        [TestMethod]
+        public async Task GetMissingCompressedBufferAttachmentAsync()
+        {
+            CbObject cbObjectAttachment = CbObject.Build(writer => writer.WriteString("ValueField", "This field has a value"));
+            byte[] cbAttachmentData = cbObjectAttachment.GetView().ToArray();
+            BlobId cbAttachmentHash = BlobId.FromBlob(cbAttachmentData);
+            RefId key = RefId.FromName("compressedAttachedObject");
+
+            CbObject cbObject = CbObject.Build(writer => writer.WriteObjectAttachment("Attachment", cbAttachmentHash.AsIoHash()));
+            byte[] cbObjectData = cbObject.GetView().ToArray();
+            BlobId cbObjectHash = BlobId.FromBlob(cbObjectData);
+
+            {
+                // upload compressed blob
+                using HttpContent requestContent = new ByteArrayContent(cbAttachmentData);
+                requestContent.Headers.ContentType = new MediaTypeHeaderValue(CustomMediaTypeNames.UnrealCompactBinary);
+
+                HttpResponseMessage result = await _httpClient!.PutAsync(new Uri($"api/v1/objects/{TestNamespace}/{cbAttachmentHash}", UriKind.Relative), requestContent);
+                result.EnsureSuccessStatusCode();
+            }
+
+            {
+                // upload ref
+                using HttpContent requestContent = new ByteArrayContent(cbObjectData);
+                requestContent.Headers.ContentType = new MediaTypeHeaderValue(CustomMediaTypeNames.UnrealCompactBinary);
+                requestContent.Headers.Add(CommonHeaders.HashHeaderName, cbObjectHash.ToString());
+
+                HttpResponseMessage result = await _httpClient!.PutAsync(new Uri($"api/v1/refs/{TestNamespace}/bucket/{key}.uecb", UriKind.Relative), requestContent);
+                result.EnsureSuccessStatusCode();
+
+                {
+                    Assert.AreEqual(CustomMediaTypeNames.UnrealCompactBinary, result!.Content.Headers.ContentType!.MediaType);
+                    // check that no blobs are missing
+                    await using MemoryStream ms = new MemoryStream();
+                    await result.Content.CopyToAsync(ms);
+                    byte[] roundTrippedBuffer = ms.ToArray();
+                    CbObject cb = new CbObject(roundTrippedBuffer);
+                    CbField needsField = cb["needs"];
+                    List<BlobId> missingBlobs = needsField.AsArray().Select(field => BlobId.FromIoHash(field.AsHash())).ToList();
+                    Assert.AreEqual(0, missingBlobs.Count);
+                }
+            }
+
+            {
+                // verify we can fetch the blob properly
+                HttpResponseMessage getResponse = await _httpClient.GetAsync(new Uri($"api/v1/refs/{TestNamespace}/bucket/{key}.uecb", UriKind.Relative));
+                getResponse.EnsureSuccessStatusCode();
+                await using MemoryStream ms = new MemoryStream();
+                await getResponse.Content.CopyToAsync(ms);
+
+                byte[] roundTrippedBuffer = ms.ToArray();
+
+                CollectionAssert.AreEqual(cbObjectData, roundTrippedBuffer);
+                Assert.AreEqual(cbObjectHash, BlobId.FromBlob(roundTrippedBuffer));
+            }
+
+            {
+                // verify that head checks find the object
+                using HttpRequestMessage headRequest = new HttpRequestMessage(HttpMethod.Head, new Uri($"api/v1/refs/{TestNamespace}/bucket/{key}", UriKind.Relative));
+                HttpResponseMessage getResponse = await _httpClient.SendAsync(headRequest);
+                getResponse.EnsureSuccessStatusCode();
+            }
+
+            {
+                // verify that the exists check doesnt find any issue
+                HttpResponseMessage existsResponse = await _httpClient.GetAsync(new Uri($"api/v1/refs/{TestNamespace}/exists?names=bucket.{key}", UriKind.Relative));
+                existsResponse.EnsureSuccessStatusCode();
+
+                ExistCheckMultipleRefsResponse? response = await existsResponse.Content.ReadFromJsonAsync<ExistCheckMultipleRefsResponse>();
+                Assert.IsNotNull(response);
+                Assert.AreEqual(0, response.Missing.Count);
+            }
+
+            {
+                // delete the blob referenced by the compressed buffer
+                await Service.DeleteObjectAsync(TestNamespace, cbAttachmentHash);
+            }
+
+            {
+                // the compact binary object still exists, so this returns a success (trying to resolve the attachment will fail)
+                HttpResponseMessage getResponse = await _httpClient.GetAsync(new Uri($"api/v1/refs/{TestNamespace}/bucket/{key}.uecb", UriKind.Relative));
+                getResponse.EnsureSuccessStatusCode();
+            }
+
+            {
+                // verify that head checks fail
+                using HttpRequestMessage headRequest = new HttpRequestMessage(HttpMethod.Head, new Uri($"api/v1/refs/{TestNamespace}/bucket/{key}", UriKind.Relative));
+                HttpResponseMessage getResponse = await _httpClient.SendAsync(headRequest);
+                Assert.AreEqual(HttpStatusCode.NotFound, getResponse.StatusCode);
+            }
+
+            {
+                // verify that the exists check correctly returns a missing blob
+                HttpResponseMessage existsResponse = await _httpClient.GetAsync(new Uri($"api/v1/refs/{TestNamespace}/exists?names=bucket.{key}", UriKind.Relative));
+                existsResponse.EnsureSuccessStatusCode();
+
+                ExistCheckMultipleRefsResponse? response = await existsResponse.Content.ReadFromJsonAsync<ExistCheckMultipleRefsResponse>();
+                Assert.IsNotNull(response);
+                Assert.AreEqual(1, response.Missing.Count);
+                Assert.AreEqual(key, response.Missing[0].Key);
+                Assert.AreEqual("bucket", response.Missing[0].Bucket.ToString());
+            }
+        }
+
+        [TestMethod]
+        public async Task GetMissingBlobRecordAsync()
+        {
+            string blobContents = "This is a blob";
+            byte[] blobData = Encoding.ASCII.GetBytes(blobContents);
+            BlobId blobHash = BlobId.FromBlob(blobData);
+            RefId key = RefId.FromName("newReferenceObject");
+
+            using HttpContent requestContent = new ByteArrayContent(blobData);
+            requestContent.Headers.ContentType = new MediaTypeHeaderValue(MediaTypeNames.Application.Octet);
+            requestContent.Headers.Add(CommonHeaders.HashHeaderName, blobHash.ToString());
+
+            HttpResponseMessage result = await _httpClient!.PutAsync(new Uri($"api/v1/refs/{TestNamespace}/bucket/{key}.uecb", UriKind.Relative), requestContent);
+            result.EnsureSuccessStatusCode();
+
+            {
+                Assert.AreEqual(CustomMediaTypeNames.UnrealCompactBinary, result!.Content.Headers.ContentType!.MediaType);
+                // check that no blobs are missing
+                await using MemoryStream ms = new MemoryStream();
+                await result.Content.CopyToAsync(ms);
+                byte[] roundTrippedBuffer = ms.ToArray();
+                CbObject cb = new CbObject(roundTrippedBuffer);
+                CbField needsField = cb["needs"];
+                List<BlobId> missingBlobs = needsField.AsArray().Select(field => BlobId.FromIoHash(field.AsHash())).ToList();
+                Assert.AreEqual(0, missingBlobs.Count);
+            }
+
+            {
+                // verify we can fetch the blob properly
+                HttpResponseMessage getResponse = await _httpClient.GetAsync(new Uri($"api/v1/refs/{TestNamespace}/bucket/{key}.raw", UriKind.Relative));
+                getResponse.EnsureSuccessStatusCode();
+                await using MemoryStream ms = new MemoryStream();
+                await getResponse.Content.CopyToAsync(ms);
+
+                byte[] roundTrippedBuffer = ms.ToArray();
+
+                CollectionAssert.AreEqual(blobData, roundTrippedBuffer);
+                Assert.AreEqual(blobHash, BlobId.FromBlob(roundTrippedBuffer));
+            }
+
+            {
+                // verify that head checks find the object
+                using HttpRequestMessage headRequest = new HttpRequestMessage(HttpMethod.Head, new Uri($"api/v1/refs/{TestNamespace}/bucket/{key}", UriKind.Relative));
+                HttpResponseMessage getResponse = await _httpClient.SendAsync(headRequest);
+                getResponse.EnsureSuccessStatusCode();
+            }
+
+            {
+                // verify that the exists check doesnt find any issue
+                HttpResponseMessage existsResponse = await _httpClient.GetAsync(new Uri($"api/v1/refs/{TestNamespace}/exists?names=bucket.{key}", UriKind.Relative));
+                existsResponse.EnsureSuccessStatusCode();
+
+                ExistCheckMultipleRefsResponse? response = await existsResponse.Content.ReadFromJsonAsync<ExistCheckMultipleRefsResponse>();
+                Assert.IsNotNull(response);
+                Assert.AreEqual(0, response.Missing.Count);
+            }
+
+            {
+                // delete the blob 
+                await Service.DeleteObjectAsync(TestNamespace, blobHash);
+            }
+
+            {
+                // the compact binary object still exists, so this returns a success (trying to resolve the attachment will fail)
+                HttpResponseMessage getResponse = await _httpClient.GetAsync(new Uri($"api/v1/refs/{TestNamespace}/bucket/{key}.uecb", UriKind.Relative));
+                getResponse.EnsureSuccessStatusCode();
+            }
+
+            {
+                // the compact binary object still exists, so this returns a success (trying to resolve the attachment will fail)
+                HttpResponseMessage getResponse = await _httpClient.GetAsync(new Uri($"api/v1/refs/{TestNamespace}/bucket/{key}.json", UriKind.Relative));
+                getResponse.EnsureSuccessStatusCode();
+            }
+
+            {
+                // we should now see a 404 as the blob is missing
+                HttpResponseMessage getResponse = await _httpClient.GetAsync(new Uri($"api/v1/refs/{TestNamespace}/bucket/{key}.raw", UriKind.Relative));
+                Assert.AreEqual(HttpStatusCode.NotFound, getResponse.StatusCode);
+                Assert.AreEqual("application/problem+json", getResponse.Content.Headers.ContentType!.MediaType);
+                string s = await getResponse.Content.ReadAsStringAsync();
+                ProblemDetails? problem = JsonSerializer.Deserialize<ProblemDetails>(s);
+                Assert.IsNotNull(problem);
+                Assert.AreEqual($"Object {blobHash} in {TestNamespace} not found", problem.Title);
+            }
+
+            {
+                // verify that head checks fail
+                using HttpRequestMessage headRequest = new HttpRequestMessage(HttpMethod.Head, new Uri($"api/v1/refs/{TestNamespace}/bucket/{key}", UriKind.Relative));
+                HttpResponseMessage getResponse = await _httpClient.SendAsync(headRequest);
+                Assert.AreEqual(HttpStatusCode.NotFound, getResponse.StatusCode);
+            }
+
+            {
+                // verify that the exists check correctly returns a missing blob
+                HttpResponseMessage existsResponse = await _httpClient.GetAsync(new Uri($"api/v1/refs/{TestNamespace}/exists?names=bucket.{key}", UriKind.Relative));
+                existsResponse.EnsureSuccessStatusCode();
+
+                ExistCheckMultipleRefsResponse? response = await existsResponse.Content.ReadFromJsonAsync<ExistCheckMultipleRefsResponse>();
+                Assert.IsNotNull(response);
+                Assert.AreEqual(1, response.Missing.Count);
+                Assert.AreEqual(key, response.Missing[0].Key);
+                Assert.AreEqual("bucket", response.Missing[0].Bucket.ToString());
+            }
+        }
+
+        [TestMethod]
+        public async Task DeleteObjectAsync()
+        {
+            const string objectContents = "This is treated as a opaque blob";
+            byte[] data = Encoding.ASCII.GetBytes(objectContents);
+            BlobId objectHash = BlobId.FromBlob(data);
+
+            using HttpContent requestContent = new ByteArrayContent(data);
+            requestContent.Headers.ContentType = new MediaTypeHeaderValue(MediaTypeNames.Application.Octet);
+            requestContent.Headers.Add(CommonHeaders.HashHeaderName, objectHash.ToString());
+
+            RefId key = RefId.FromName("deletableObject");
+            // submit the object
+            {
+                HttpResponseMessage result = await _httpClient!.PutAsync(new Uri($"api/v1/refs/{TestNamespace}/bucket/{key}", UriKind.Relative), requestContent);
+                result.EnsureSuccessStatusCode();
+            }
+
+            // verify it is present
+            {
+                HttpResponseMessage result = await _httpClient!.GetAsync(new Uri($"api/v1/refs/{TestNamespace}/bucket/{key}.raw", UriKind.Relative));
+                result.EnsureSuccessStatusCode();
+            }
+
+            // delete the object
+            {
+                HttpResponseMessage result = await _httpClient!.DeleteAsync(new Uri($"api/v1/refs/{TestNamespace}/bucket/{key}", UriKind.Relative));
+                result.EnsureSuccessStatusCode();
+            }
+
+            // ensure the object is not present anymore
+            {
+                HttpResponseMessage result = await _httpClient!.GetAsync(new Uri($"api/v1/refs/{TestNamespace}/bucket/{key}.raw", UriKind.Relative));
+                Assert.AreEqual(HttpStatusCode.NotFound, result.StatusCode);
+            }
+
+            // delete the object again which doesn't exist anymore
+            {
+                HttpResponseMessage result = await _httpClient!.DeleteAsync(new Uri($"api/v1/refs/{TestNamespace}/bucket/{key}", UriKind.Relative));
+                Assert.AreEqual(HttpStatusCode.NotFound, result.StatusCode);
+            }
+        }
+
+        
+        [TestMethod]
+        public async Task DropBucketAsync()
+        {
+            const string BucketToDelete = "delete-bucket";
+
+            const string objectContents = "This is treated as a opaque blob";
+            byte[] data = Encoding.ASCII.GetBytes(objectContents);
+            BlobId objectHash = BlobId.FromBlob(data);
+
+            using HttpContent requestContent = new ByteArrayContent(data);
+            requestContent.Headers.ContentType = new MediaTypeHeaderValue(MediaTypeNames.Application.Octet);
+            requestContent.Headers.Add(CommonHeaders.HashHeaderName, objectHash.ToString());
+
+            RefId key = RefId.FromName("deletableObject");
+
+            // submit the object into multiple buckets
+            {
+                HttpResponseMessage result = await _httpClient!.PutAsync(new Uri($"api/v1/refs/{TestNamespace}/{BucketToDelete}/{key}", UriKind.Relative), requestContent);
+                result.EnsureSuccessStatusCode();
+            }
+
+            {
+                HttpResponseMessage result = await _httpClient!.PutAsync(new Uri($"api/v1/refs/{TestNamespace}/bucket/{key}", UriKind.Relative), requestContent);
+                result.EnsureSuccessStatusCode();
+            }
+
+            // verify it is present
+            {
+                HttpResponseMessage result = await _httpClient!.GetAsync(new Uri($"api/v1/refs/{TestNamespace}/bucket/{key}.raw", UriKind.Relative));
+                result.EnsureSuccessStatusCode();
+            }
+
+            // verify it is present
+            {
+                HttpResponseMessage result = await _httpClient!.GetAsync(new Uri($"api/v1/refs/{TestNamespace}/{BucketToDelete}/{key}.raw", UriKind.Relative));
+                result.EnsureSuccessStatusCode();
+            }
+
+            // delete the bucket
+            {
+                HttpResponseMessage result = await _httpClient!.DeleteAsync(new Uri($"api/v1/refs/{TestNamespace}/{BucketToDelete}", UriKind.Relative));
+                result.EnsureSuccessStatusCode();
+            }
+
+            // ensure the object is present in not deleted bucket
+            {
+                HttpResponseMessage result = await _httpClient!.GetAsync(new Uri($"api/v1/refs/{TestNamespace}/bucket/{key}.raw", UriKind.Relative));
+                result.EnsureSuccessStatusCode();
+            }
+
+            // ensure the object is not present anymore
+            {
+                HttpResponseMessage result = await _httpClient!.GetAsync(new Uri($"api/v1/refs/{TestNamespace}/{BucketToDelete}/{key}.raw", UriKind.Relative));
+                Assert.AreEqual(HttpStatusCode.NotFound, result.StatusCode);
+            }
+        }
+
+        
+        [TestMethod]
+        public async Task DeleteNamespaceAsync()
+        {
+            const string NamespaceToBeDeleted = "test-delete-namespace";
+
+            const string objectContents = "This is treated as a opaque blob";
+            byte[] data = Encoding.ASCII.GetBytes(objectContents);
+            BlobId objectHash = BlobId.FromBlob(data);
+
+            using HttpContent requestContent = new ByteArrayContent(data);
+            requestContent.Headers.ContentType = new MediaTypeHeaderValue(MediaTypeNames.Application.Octet);
+            requestContent.Headers.Add(CommonHeaders.HashHeaderName, objectHash.ToString());
+
+            RefId key = RefId.FromName("deletableObject");
+            // submit the object into multiple namespaces
+            {
+                HttpResponseMessage result = await _httpClient!.PutAsync(new Uri($"api/v1/refs/{TestNamespace}/bucket/{key}", UriKind.Relative), requestContent);
+                result.EnsureSuccessStatusCode();
+            }
+
+            {
+                HttpResponseMessage result = await _httpClient!.PutAsync(new Uri($"api/v1/refs/{NamespaceToBeDeleted}/bucket/{key}", UriKind.Relative), requestContent);
+                result.EnsureSuccessStatusCode();
+            }
+
+            // verify it is present
+            {
+                HttpResponseMessage result = await _httpClient!.GetAsync(new Uri($"api/v1/refs/{TestNamespace}/bucket/{key}.raw", UriKind.Relative));
+                result.EnsureSuccessStatusCode();
+            }
+
+            // verify it is present
+            {
+                HttpResponseMessage result = await _httpClient!.GetAsync(new Uri($"api/v1/refs/{NamespaceToBeDeleted}/bucket/{key}.raw", UriKind.Relative));
+                result.EnsureSuccessStatusCode();
+            }
+
+            // delete the namespace
+            {
+                HttpResponseMessage result = await _httpClient!.DeleteAsync(new Uri($"api/v1/refs/{NamespaceToBeDeleted}", UriKind.Relative));
+                result.EnsureSuccessStatusCode();
+            }
+
+            // ensure the object is present in not deleted namespace
+            {
+                HttpResponseMessage result = await _httpClient!.GetAsync(new Uri($"api/v1/refs/{TestNamespace}/bucket/{key}.raw", UriKind.Relative));
+                result.EnsureSuccessStatusCode();
+            }
+
+            // ensure the object is not present anymore
+            {
+                HttpResponseMessage result = await _httpClient!.GetAsync(new Uri($"api/v1/refs/{NamespaceToBeDeleted}/bucket/{key}.raw", UriKind.Relative));
+                Assert.AreEqual(HttpStatusCode.NotFound, result.StatusCode);
+            }
+
+            // make sure the namespace is not considered a valid namespace anymore
+            {
+                HttpResponseMessage result = await _httpClient!.GetAsync(new Uri($"api/v1/refs", UriKind.Relative));
+                result.EnsureSuccessStatusCode();
+                Assert.AreEqual(result!.Content.Headers.ContentType!.MediaType, MediaTypeNames.Application.Json);
+                GetNamespacesResponse? response = await result.Content.ReadFromJsonAsync<GetNamespacesResponse>();
+                Assert.IsNotNull(response);
+                CollectionAssert.DoesNotContain(response.Namespaces, NamespaceToBeDeleted);
+            }
+        }
+
+        
+        [TestMethod]
+        public async Task ListNamespacesAsync()
+        {
+            const string objectContents = "This is treated as a opaque blob";
+            byte[] data = Encoding.ASCII.GetBytes(objectContents);
+            BlobId objectHash = BlobId.FromBlob(data);
+            RefId key = RefId.FromName("notUsedObject");
+
+            using HttpContent requestContent = new ByteArrayContent(data);
+            requestContent.Headers.ContentType = new MediaTypeHeaderValue(MediaTypeNames.Application.Octet);
+            requestContent.Headers.Add(CommonHeaders.HashHeaderName, objectHash.ToString());
+
+            // submit a object to make sure a namespace is created
+            {
+                HttpResponseMessage result = await _httpClient!.PutAsync(new Uri($"api/v1/refs/{TestNamespace}/bucket/{key}", UriKind.Relative), requestContent);
+                result.EnsureSuccessStatusCode();
+            }
+
+            {
+                HttpResponseMessage result = await _httpClient!.GetAsync(new Uri($"api/v1/refs", UriKind.Relative));
+                result.EnsureSuccessStatusCode();
+                Assert.AreEqual(result!.Content.Headers.ContentType!.MediaType, MediaTypeNames.Application.Json);
+                GetNamespacesResponse? response = await result.Content.ReadFromJsonAsync<GetNamespacesResponse>();
+                Assert.IsNotNull(response);
+                Assert.IsTrue(response.Namespaces.Contains(TestNamespace));
+            }
+        }
+        
+        [TestMethod]
+        public async Task GetOldRecordsAsync()
+        {
+            const string objectContents = "This is treated as a opaque blob";
+            byte[] data = Encoding.ASCII.GetBytes(objectContents);
+            BlobId objectHash = BlobId.FromBlob(data);
+
+            using HttpContent requestContent = new ByteArrayContent(data);
+            requestContent.Headers.ContentType = new MediaTypeHeaderValue(MediaTypeNames.Application.Octet);
+            requestContent.Headers.Add(CommonHeaders.HashHeaderName, objectHash.ToString());
+
+            RefId key = RefId.FromName("oldRecord");
+            // submit some contents
+            {
+                HttpResponseMessage result = await _httpClient!.PutAsync(new Uri($"api/v1/refs/{TestNamespace}/bucket/{key}", UriKind.Relative), requestContent);
+                result.EnsureSuccessStatusCode();
+            }
+
+            List<(NamespaceId, BucketId, RefId, DateTime)> records = await ReferencesStore.GetRecordsAsync().ToListAsync();
+
+            (NamespaceId oldNs, BucketId oldBucket, RefId oldName, DateTime oldDate) = records.Where(tuple => tuple.Item1 == TestNamespace).First(record => record.Item3 == key);
+            Assert.AreEqual(key, oldName);
+            Assert.AreEqual("bucket", oldBucket.ToString());
+        }
+        
+        [TestMethod]
+        public async Task BatchJsonRequestAsync()
+        {
+            // verifies that json request against the batch endpoint will fail
+            {
+                using HttpContent requestContent = new StringContent("{}");
+                requestContent.Headers.ContentType = new MediaTypeHeaderValue(MediaTypeNames.Application.Json);
+
+                HttpResponseMessage result = await _httpClient!.PostAsync(new Uri($"api/v1/refs/{TestNamespace}", UriKind.Relative), requestContent);
+
+                Assert.AreEqual(HttpStatusCode.UnsupportedMediaType, result.StatusCode);
+            }
+        }
+
+        [TestMethod]
+        public async Task BatchErrorOperationsAsync()
+        {
+            // seed some data
+            BucketId bucket = new BucketId("bucket");
+            RefId newBlobObjectKey = RefId.FromName("thisObjectDoesNotExist");
+
+            RefId putObjectKey = RefId.FromName("putObjectFail");
+            CbObject ref1 = CbObject.Empty;
+
+            CbWriter getObjectOp = new CbWriter();
+            getObjectOp.BeginObject();
+            getObjectOp.WriteInteger( "opId",0);
+            getObjectOp.WriteString("op", RefBatchOps.RefBatchOp.RefOperation.GET.ToString());
+            getObjectOp.WriteString("bucket", bucket.ToString());
+            getObjectOp.WriteString("key", newBlobObjectKey.ToString());
+            getObjectOp.EndObject();
+
+            CbWriter putObjectOp = new CbWriter();
+            putObjectOp.BeginObject();
+            putObjectOp.WriteInteger("opId", 1);
+            putObjectOp.WriteString("op", RefBatchOps.RefBatchOp.RefOperation.PUT.ToString());
+            putObjectOp.WriteString("bucket", bucket.ToString());
+            putObjectOp.WriteString("key", putObjectKey.ToString());
+            putObjectOp.WriteObject("payload", ref1);
+            // omit the hash for a error object1.WriteHash("payloadHash", IoHash.Compute(ref1.GetView().Span));
+            putObjectOp.EndObject();
+
+            CbObject[] ops = new[]
+            {
+                getObjectOp.ToObject(),
+                putObjectOp.ToObject(),
+            };
+
+            CbWriter batchRequestWriter = new CbWriter();
+            batchRequestWriter.BeginObject();
+            batchRequestWriter.BeginUniformArray("ops", CbFieldType.Object);
+            foreach (CbObject o in ops)
+            {
+                batchRequestWriter.WriteObject(o);
+            }
+            batchRequestWriter.EndUniformArray();
+            batchRequestWriter.EndObject();
+            byte[] batchRequestData = batchRequestWriter.ToByteArray();
+
+            {
+                using HttpContent requestContent = new ByteArrayContent(batchRequestData);
+                requestContent.Headers.ContentType = new MediaTypeHeaderValue(CustomMediaTypeNames.UnrealCompactBinary);
+
+                HttpResponseMessage result = await _httpClient!.PostAsync(new Uri($"api/v1/refs/{TestNamespace}", UriKind.Relative), requestContent);
+                result.EnsureSuccessStatusCode();
+
+                await using MemoryStream ms = new MemoryStream();
+                await result.Content.CopyToAsync(ms);
+                byte[] roundTrippedBuffer = ms.ToArray();
+
+                BatchOpsResponse response = CbSerializer.Deserialize<BatchOpsResponse>(roundTrippedBuffer);
+                Assert.AreEqual(2, response.Results.Count);
+                
+                BatchOpsResponse.OpResponses op0 = response.Results.First(r => r.OpId == 0);
+                Assert.IsNotNull(op0.Response);
+                Assert.AreEqual(404, op0.StatusCode);
+                Assert.IsTrue(!op0.Response["title"].Equals(CbField.Empty));
+                Assert.AreEqual("Object not found 29911b58b3c970ba39d9690f2dca66839dd6f5d9 in bucket bucket namespace test-namespace", op0.Response["title"].AsString());
+
+                BatchOpsResponse.OpResponses op1 = response.Results.First(r => r.OpId == 1);
+                Assert.IsNotNull(op1.Response);
+                Assert.AreEqual(500, op1.StatusCode);
+                Assert.IsTrue(!op1.Response["title"].Equals(CbField.Empty));
+                Assert.AreEqual("Missing payload for operation: 1", op1.Response["title"].AsString());
+            }
+        }
+
+        [TestMethod]
+        public async Task BatchGetOperationsAsync()
+        {
+            // seed some data
+            BucketId bucket = new BucketId("bucket");
+            RefId newBlobObjectKey = RefId.FromName("newBlobObject");
+            CbObject newBlobObject = CbObject.Build(writer => writer.WriteString("String", "this-has-contents"));
+
+            {
+                byte[] cbObjectBytes = newBlobObject.GetView().ToArray();
+                BlobId blobHash = BlobId.FromBlob(cbObjectBytes);
+
+                using HttpContent requestContent = new ByteArrayContent(cbObjectBytes);
+                requestContent.Headers.ContentType = new MediaTypeHeaderValue(CustomMediaTypeNames.UnrealCompactBinary);
+                requestContent.Headers.Add(CommonHeaders.HashHeaderName, blobHash.ToString());
+
+                HttpResponseMessage result = await _httpClient!.PutAsync(new Uri($"api/v1/refs/{TestNamespace}/{bucket}/{newBlobObjectKey}.uecb", UriKind.Relative), requestContent);
+                result.EnsureSuccessStatusCode();
+                byte[] content = await result.Content.ReadAsByteArrayAsync();
+                PutObjectResponse? response = CbSerializer.Deserialize<PutObjectResponse?>(content);
+                Assert.IsNotNull(response);
+                Assert.IsTrue(response.Needs.Length == 0);
+            }
+
+            byte[] blobContents = Encoding.ASCII.GetBytes("This is a attached blob");
+            CbObject newReferenceObject = CbObject.Build(writer => writer.WriteBinaryAttachment("Attachment", IoHash.Compute(blobContents)));
+            RefId newReferenceObjectKey = RefId.FromName("newReferenceObject");
+
+            {
+                BlobId blobHash = BlobId.FromBlob(blobContents);
+
+                using HttpContent requestContent = new ByteArrayContent(blobContents);
+                requestContent.Headers.ContentType = new MediaTypeHeaderValue(MediaTypeNames.Application.Octet);
+                requestContent.Headers.Add(CommonHeaders.HashHeaderName, blobHash.ToString());
+
+                HttpResponseMessage result = await _httpClient!.PutAsync(new Uri($"api/v1/blobs/{TestNamespace}/{blobHash}", UriKind.Relative), requestContent);
+                result.EnsureSuccessStatusCode();
+            }
+
+            byte[] blobContentsMissing = Encoding.ASCII.GetBytes("This contents will not be submitted");
+            CbObject missingAttachmentObject = CbObject.Build(writer => writer.WriteBinaryAttachment("Attachment", IoHash.Compute(blobContentsMissing)));
+            RefId missingAttachmentKey = RefId.FromName("blobMissingAttachment");
+
+            {
+                BlobId blobHash = BlobId.FromBlob(blobContents);
+
+                using HttpContent requestContent = new ByteArrayContent(blobContents);
+                requestContent.Headers.ContentType = new MediaTypeHeaderValue(MediaTypeNames.Application.Octet);
+                requestContent.Headers.Add(CommonHeaders.HashHeaderName, blobHash.ToString());
+
+                HttpResponseMessage result = await _httpClient!.PutAsync(new Uri($"api/v1/blobs/{TestNamespace}/{blobHash}", UriKind.Relative), requestContent);
+                result.EnsureSuccessStatusCode();
+            }
+
+            {
+                byte[] cbObjectBytes = newReferenceObject.GetView().ToArray();
+                BlobId blobHash = BlobId.FromBlob(cbObjectBytes);
+
+                using HttpContent requestContent = new ByteArrayContent(cbObjectBytes);
+                requestContent.Headers.ContentType = new MediaTypeHeaderValue(CustomMediaTypeNames.UnrealCompactBinary);
+                requestContent.Headers.Add(CommonHeaders.HashHeaderName, blobHash.ToString());
+
+                HttpResponseMessage result = await _httpClient!.PutAsync(new Uri($"api/v1/refs/{TestNamespace}/{bucket}/{newReferenceObjectKey}.uecb", UriKind.Relative), requestContent);
+                result.EnsureSuccessStatusCode();
+
+                result.EnsureSuccessStatusCode();
+                byte[] content = await result.Content.ReadAsByteArrayAsync();
+                PutObjectResponse response = CbSerializer.Deserialize<PutObjectResponse>(content);
+                Assert.IsTrue(response.Needs.Length == 0);
+            }
+            CbWriter getObjectOp = new CbWriter();
+            getObjectOp.BeginObject();
+            getObjectOp.WriteInteger( "opId",0);
+            getObjectOp.WriteString("op", RefBatchOps.RefBatchOp.RefOperation.GET.ToString());
+            getObjectOp.WriteString("bucket", bucket.ToString());
+            getObjectOp.WriteString("key", newBlobObjectKey.ToString());
+            getObjectOp.EndObject();
+
+            CbWriter getObjectOp2 = new CbWriter();
+            getObjectOp2.BeginObject();
+            getObjectOp2.WriteInteger("opId", 1);
+            getObjectOp2.WriteString("op", RefBatchOps.RefBatchOp.RefOperation.GET.ToString());
+            getObjectOp2.WriteString("bucket", bucket.ToString());
+            getObjectOp2.WriteString("key", newReferenceObjectKey.ToString());
+            getObjectOp2.EndObject();
+
+            
+            CbWriter getObjectOp3 = new CbWriter();
+            getObjectOp3.BeginObject();
+            getObjectOp3.WriteInteger("opId", 2);
+            getObjectOp3.WriteString("op", RefBatchOps.RefBatchOp.RefOperation.GET.ToString());
+            getObjectOp3.WriteString("bucket", bucket.ToString());
+            getObjectOp3.WriteString("key", missingAttachmentKey.ToString());
+            getObjectOp3.WriteBool("resolveAttachments", true);
+            getObjectOp3.EndObject();
+
+            CbObject[] ops = new[]
+            {
+                getObjectOp.ToObject(),
+                getObjectOp2.ToObject(),
+                getObjectOp3.ToObject(),
+            };
+
+            CbWriter batchRequestWriter = new CbWriter();
+            batchRequestWriter.BeginObject();
+            batchRequestWriter.BeginUniformArray("ops", CbFieldType.Object);
+            foreach (CbObject o in ops)
+            {
+                batchRequestWriter.WriteObject(o);
+            }
+            batchRequestWriter.EndUniformArray();
+            batchRequestWriter.EndObject();
+            byte[] batchRequestData = batchRequestWriter.ToByteArray();
+
+            {
+                using HttpContent requestContent = new ByteArrayContent(batchRequestData);
+                requestContent.Headers.ContentType = new MediaTypeHeaderValue(CustomMediaTypeNames.UnrealCompactBinary);
+
+                HttpResponseMessage result = await _httpClient!.PostAsync(new Uri($"api/v1/refs/{TestNamespace}", UriKind.Relative), requestContent);
+                result.EnsureSuccessStatusCode();
+
+                await using MemoryStream ms = new MemoryStream();
+                await result.Content.CopyToAsync(ms);
+                byte[] roundTrippedBuffer = ms.ToArray();
+
+                BatchOpsResponse response = CbSerializer.Deserialize<BatchOpsResponse>(roundTrippedBuffer);
+                Assert.AreEqual(3, response.Results.Count);
+                
+                BatchOpsResponse.OpResponses op0 = response.Results.First(r => r.OpId == 0);
+                Assert.IsNotNull(op0.Response);
+                Assert.AreEqual(200, op0.StatusCode);
+                Assert.AreEqual(newBlobObject, op0.Response);
+
+                BatchOpsResponse.OpResponses op1 = response.Results.First(r => r.OpId == 1);
+                Assert.IsNotNull(op1.Response);
+                Assert.AreEqual(200, op1.StatusCode);
+                Assert.AreEqual(newReferenceObject, op1.Response);
+
+                BatchOpsResponse.OpResponses op2 = response.Results.First(r => r.OpId == 2);
+                Assert.IsNotNull(op2.Response);
+                Assert.AreEqual(404, op2.StatusCode);
+                Assert.AreNotEqual(missingAttachmentObject, op2.Response);
+            }
+        }
+
+        [TestMethod]
+        public async Task BatchHeadOperationsAsync()
+        {
+            // seed some data
+            BucketId bucket = new BucketId("bucket");
+            RefId newBlobObjectKey = RefId.FromName("newBlobObject");
+            CbObject newBlobObject = CbObject.Build(writer => writer.WriteString("String", "this-has-contents"));
+
+            {
+                byte[] cbObjectBytes = newBlobObject.GetView().ToArray();
+                BlobId blobHash = BlobId.FromBlob(cbObjectBytes);
+
+                using HttpContent requestContent = new ByteArrayContent(cbObjectBytes);
+                requestContent.Headers.ContentType = new MediaTypeHeaderValue(CustomMediaTypeNames.UnrealCompactBinary);
+                requestContent.Headers.Add(CommonHeaders.HashHeaderName, blobHash.ToString());
+
+                HttpResponseMessage result = await _httpClient!.PutAsync(new Uri($"api/v1/refs/{TestNamespace}/{bucket}/{newBlobObjectKey}.uecb", UriKind.Relative), requestContent);
+                result.EnsureSuccessStatusCode();
+                byte[] content = await result.Content.ReadAsByteArrayAsync();
+                PutObjectResponse? response = CbSerializer.Deserialize<PutObjectResponse?>(content);
+                Assert.IsNotNull(response);
+                Assert.IsTrue(response.Needs.Length == 0);
+            }
+
+            byte[] blobContents = Encoding.ASCII.GetBytes("This is a attached blob");
+            CbObject newReferenceObject = CbObject.Build(writer => writer.WriteBinaryAttachment("Attachment", IoHash.Compute(blobContents)));
+            RefId newReferenceObjectKey = RefId.FromName("newReferenceObject");
+
+            {
+                BlobId blobHash = BlobId.FromBlob(blobContents);
+
+                using HttpContent requestContent = new ByteArrayContent(blobContents);
+                requestContent.Headers.ContentType = new MediaTypeHeaderValue(MediaTypeNames.Application.Octet);
+                requestContent.Headers.Add(CommonHeaders.HashHeaderName, blobHash.ToString());
+
+                HttpResponseMessage result = await _httpClient!.PutAsync(new Uri($"api/v1/blobs/{TestNamespace}/{blobHash}", UriKind.Relative), requestContent);
+                result.EnsureSuccessStatusCode();
+            }
+
+            {
+                byte[] cbObjectBytes = newReferenceObject.GetView().ToArray();
+                BlobId blobHash = BlobId.FromBlob(cbObjectBytes);
+
+                using HttpContent requestContent = new ByteArrayContent(cbObjectBytes);
+                requestContent.Headers.ContentType = new MediaTypeHeaderValue(CustomMediaTypeNames.UnrealCompactBinary);
+                requestContent.Headers.Add(CommonHeaders.HashHeaderName, blobHash.ToString());
+
+                HttpResponseMessage result = await _httpClient!.PutAsync(new Uri($"api/v1/refs/{TestNamespace}/{bucket}/{newReferenceObjectKey}.uecb", UriKind.Relative), requestContent);
+                result.EnsureSuccessStatusCode();
+
+                result.EnsureSuccessStatusCode();
+                byte[] content = await result.Content.ReadAsByteArrayAsync();
+                PutObjectResponse response = CbSerializer.Deserialize<PutObjectResponse>(content);
+                Assert.IsTrue(response.Needs.Length == 0);
+            }
+            CbWriter getObjectOp = new CbWriter();
+            getObjectOp.BeginObject();
+            getObjectOp.WriteInteger( "opId",0);
+            getObjectOp.WriteString("op", RefBatchOps.RefBatchOp.RefOperation.HEAD.ToString());
+            getObjectOp.WriteString("bucket", bucket.ToString());
+            getObjectOp.WriteString("key", newBlobObjectKey.ToString());
+            getObjectOp.EndObject();
+
+            CbWriter getObjectOp2 = new CbWriter();
+            getObjectOp2.BeginObject();
+            getObjectOp2.WriteInteger("opId", 1);
+            getObjectOp2.WriteString("op", RefBatchOps.RefBatchOp.RefOperation.HEAD.ToString());
+            getObjectOp2.WriteString("bucket", bucket.ToString());
+            getObjectOp2.WriteString("key", newReferenceObjectKey.ToString());
+            getObjectOp2.EndObject();
+
+            CbObject[] ops = new[]
+            {
+                getObjectOp.ToObject(),
+                getObjectOp2.ToObject(),
+            };
+
+            CbWriter batchRequestWriter = new CbWriter();
+            batchRequestWriter.BeginObject();
+            batchRequestWriter.BeginUniformArray("ops", CbFieldType.Object);
+            foreach (CbObject o in ops)
+            {
+                batchRequestWriter.WriteObject(o);
+            }
+            batchRequestWriter.EndUniformArray();
+            batchRequestWriter.EndObject();
+            byte[] batchRequestData = batchRequestWriter.ToByteArray();
+            
+            {
+                using HttpContent requestContent = new ByteArrayContent(batchRequestData);
+                requestContent.Headers.ContentType = new MediaTypeHeaderValue(CustomMediaTypeNames.UnrealCompactBinary);
+
+                HttpResponseMessage result = await _httpClient!.PostAsync(new Uri($"api/v1/refs/{TestNamespace}", UriKind.Relative), requestContent);
+                result.EnsureSuccessStatusCode();
+
+                await using MemoryStream ms = new MemoryStream();
+                await result.Content.CopyToAsync(ms);
+                byte[] roundTrippedBuffer = ms.ToArray();
+
+                BatchOpsResponse response = CbSerializer.Deserialize<BatchOpsResponse>(roundTrippedBuffer);
+                Assert.AreEqual(2, response.Results.Count);
+                
+                BatchOpsResponse.OpResponses op0 = response.Results.First(r => r.OpId == 0);
+                Assert.IsNotNull(op0.Response);
+                Assert.AreEqual(200, op0.StatusCode);
+                Assert.IsTrue(!op0.Response["exists"].Equals(CbField.Empty));
+                Assert.IsTrue(op0.Response["exists"].AsBool());
+
+                BatchOpsResponse.OpResponses op1 = response.Results.First(r => r.OpId == 1);
+                Assert.IsNotNull(op1.Response);
+                Assert.AreEqual(200, op1.StatusCode);
+                Assert.IsTrue(!op1.Response["exists"].Equals(CbField.Empty));
+                Assert.IsTrue(op1.Response["exists"].AsBool());
+            }
+        }
+
+        [TestMethod]
+        public async Task BatchPutOperationsAsync()
+        {
+            BucketId bucket = new BucketId("bucket");
+            RefId ref0name = RefId.FromName("putRef0");
+            CbObject ref0 = CbObject.Build(writer => writer.WriteString("foo", "bar"));
+
+            RefId ref1name = RefId.FromName("putRef1");
+            CbObject ref1 = CbObject.Build(writer => writer.WriteInteger("baz", 1337));
+
+            CbWriter object0 = new CbWriter();
+            object0.BeginObject();
+            object0.WriteInteger( "opId",0);
+            object0.WriteString("op", RefBatchOps.RefBatchOp.RefOperation.PUT.ToString());
+            object0.WriteString("bucket", bucket.ToString());
+            object0.WriteString("key", ref0name.ToString());
+            object0.WriteObject("payload", ref0);
+            object0.WriteHash("payloadHash", IoHash.Compute(ref0.GetView().Span));
+            object0.EndObject();
+
+            CbWriter object1 = new CbWriter();
+            object1.BeginObject();
+            object1.WriteInteger("opId", 1);
+            object1.WriteString("op", RefBatchOps.RefBatchOp.RefOperation.PUT.ToString());
+            object1.WriteString("bucket", bucket.ToString());
+            object1.WriteString("key", ref1name.ToString());
+            object1.WriteObject("payload", ref1);
+            object1.WriteHash("payloadHash", IoHash.Compute(ref1.GetView().Span));
+            object1.EndObject();
+
+            CbObject[] ops = new[]
+            {
+                object0.ToObject(),
+                object1.ToObject(),
+            };
+
+            CbWriter batchRequestWriter = new CbWriter();
+            batchRequestWriter.BeginObject();
+            batchRequestWriter.BeginUniformArray("ops", CbFieldType.Object);
+            foreach (CbObject o in ops)
+            {
+                batchRequestWriter.WriteObject(o);
+            }
+            batchRequestWriter.EndUniformArray();
+            batchRequestWriter.EndObject();
+            byte[] batchRequestData = batchRequestWriter.ToByteArray();
+
+            {
+                using HttpContent requestContent = new ByteArrayContent(batchRequestData);
+                requestContent.Headers.ContentType = new MediaTypeHeaderValue(CustomMediaTypeNames.UnrealCompactBinary);
+
+                HttpResponseMessage result = await _httpClient!.PostAsync(new Uri($"api/v1/refs/{TestNamespace}", UriKind.Relative), requestContent);
+                result.EnsureSuccessStatusCode();
+
+                await using MemoryStream ms = new MemoryStream();
+                await result.Content.CopyToAsync(ms);
+                byte[] roundTrippedBuffer = ms.ToArray();
+
+                BatchOpsResponse response = CbSerializer.Deserialize<BatchOpsResponse>(roundTrippedBuffer);
+                Assert.AreEqual(2, response.Results.Count);
+                
+                BatchOpsResponse.OpResponses op0 = response.Results.First(r => r.OpId == 0);
+                Assert.IsNotNull(op0.Response);
+                Assert.AreEqual(200, op0.StatusCode, $"Expected 200 response, got {op0.StatusCode} . Response: {op0.Response.ToJson()}");
+                Assert.IsTrue(!op0.Response["needs"].Equals(CbField.Empty));
+                CollectionAssert.AreEqual(Array.Empty<IoHash>(), op0.Response["needs"].ToArray());
+
+                BatchOpsResponse.OpResponses op1 = response.Results.First(r => r.OpId == 1);
+                Assert.IsNotNull(op1.Response);
+                Assert.AreEqual(200, op1.StatusCode, $"Expected 200 response, got {op1.StatusCode} . Response: {op1.Response.ToJson()}");
+                Assert.IsTrue(!op1.Response["needs"].Equals(CbField.Empty));
+                CollectionAssert.AreEqual(Array.Empty<IoHash>(), op1.Response["needs"].ToArray());
+            }
+        }
+
+        [TestMethod]
+        public async Task BatchMixedOperationsAsync()
+        {
+            // seed some data
+            BucketId bucket = new BucketId("bucket");
+            RefId getObjectKey = RefId.FromName("getBlobObject");
+            CbObject getObject = CbObject.Build(writer => writer.WriteString("String", "this-has-contents"));
+
+            RefId putObjectKey = RefId.FromName("putBlobObject");
+            CbObject putObject = CbObject.Build(writer => writer.WriteString("String", "this-has-contents"));
+
+            RefId missingObjectKey = RefId.FromName("thisKeyDoesNotExist");
+
+            {
+                byte[] cbObjectBytes = getObject.GetView().ToArray();
+                BlobId blobHash = BlobId.FromBlob(cbObjectBytes);
+
+                using HttpContent requestContent = new ByteArrayContent(cbObjectBytes);
+                requestContent.Headers.ContentType = new MediaTypeHeaderValue(CustomMediaTypeNames.UnrealCompactBinary);
+                requestContent.Headers.Add(CommonHeaders.HashHeaderName, blobHash.ToString());
+
+                HttpResponseMessage result = await _httpClient!.PutAsync(new Uri($"api/v1/refs/{TestNamespace}/{bucket}/{getObjectKey}.uecb", UriKind.Relative), requestContent);
+                result.EnsureSuccessStatusCode();
+                byte[] content = await result.Content.ReadAsByteArrayAsync();
+                PutObjectResponse? response = CbSerializer.Deserialize<PutObjectResponse?>(content);
+                Assert.IsNotNull(response);
+                Assert.IsTrue(response.Needs.Length == 0);
+            }
+
+            CbWriter getObjectOp = new CbWriter();
+            getObjectOp.BeginObject();
+            getObjectOp.WriteInteger( "opId",0);
+            getObjectOp.WriteString("op", RefBatchOps.RefBatchOp.RefOperation.GET.ToString());
+            getObjectOp.WriteString("bucket", bucket.ToString());
+            getObjectOp.WriteString("key", getObjectKey.ToString());
+            getObjectOp.EndObject();
+
+            CbWriter putObjectOp = new CbWriter();
+            putObjectOp.BeginObject();
+            putObjectOp.WriteInteger("opId", 1);
+            putObjectOp.WriteString("op", RefBatchOps.RefBatchOp.RefOperation.PUT.ToString());
+            putObjectOp.WriteString("bucket", bucket.ToString());
+            putObjectOp.WriteString("key", putObjectKey.ToString());
+            putObjectOp.WriteObject("payload", putObject);
+            putObjectOp.WriteHash("payloadHash", IoHash.Compute(putObject.GetView().Span));
+            putObjectOp.EndObject();
+
+            CbWriter errorObjectOp = new CbWriter();
+            errorObjectOp.BeginObject();
+            errorObjectOp.WriteInteger( "opId",2);
+            errorObjectOp.WriteString("op", RefBatchOps.RefBatchOp.RefOperation.GET.ToString());
+            errorObjectOp.WriteString("bucket", bucket.ToString());
+            errorObjectOp.WriteString("key", missingObjectKey.ToString());
+            errorObjectOp.EndObject();
+
+            CbWriter headObjectOp = new CbWriter();
+            headObjectOp.BeginObject();
+            headObjectOp.WriteInteger("opId", 3);
+            headObjectOp.WriteString("op", RefBatchOps.RefBatchOp.RefOperation.HEAD.ToString());
+            headObjectOp.WriteString("bucket", bucket.ToString());
+            headObjectOp.WriteString("key", getObjectKey.ToString());
+            headObjectOp.EndObject();
+
+            CbObject[] ops = new[]
+            {
+                getObjectOp.ToObject(),
+                putObjectOp.ToObject(),
+                errorObjectOp.ToObject(),
+                headObjectOp.ToObject(),
+            };
+
+            CbWriter batchRequestWriter = new CbWriter();
+            batchRequestWriter.BeginObject();
+            batchRequestWriter.BeginUniformArray("ops", CbFieldType.Object);
+            foreach (CbObject o in ops)
+            {
+                batchRequestWriter.WriteObject(o);
+            }
+            batchRequestWriter.EndUniformArray();
+            batchRequestWriter.EndObject();
+            byte[] batchRequestData = batchRequestWriter.ToByteArray();
+            
+            {
+                using HttpContent requestContent = new ByteArrayContent(batchRequestData);
+                requestContent.Headers.ContentType = new MediaTypeHeaderValue(CustomMediaTypeNames.UnrealCompactBinary);
+
+                HttpResponseMessage result = await _httpClient!.PostAsync(new Uri($"api/v1/refs/{TestNamespace}", UriKind.Relative), requestContent);
+                result.EnsureSuccessStatusCode();
+
+                await using MemoryStream ms = new MemoryStream();
+                await result.Content.CopyToAsync(ms);
+                byte[] roundTrippedBuffer = ms.ToArray();
+
+                BatchOpsResponse response = CbSerializer.Deserialize<BatchOpsResponse>(roundTrippedBuffer);
+                Assert.AreEqual(4, response.Results.Count);
+                
+                BatchOpsResponse.OpResponses getOp = response.Results.First(r => r.OpId == 0);
+                Assert.IsNotNull(getOp.Response);
+                Assert.AreEqual(200, getOp.StatusCode);
+                Assert.AreEqual(getObject, getOp.Response);
+
+                BatchOpsResponse.OpResponses putOp = response.Results.First(r => r.OpId == 1);
+                Assert.IsNotNull(putOp.Response);
+                Assert.AreEqual(200, putOp.StatusCode);
+                Assert.IsTrue(!putOp.Response["needs"].Equals(CbField.Empty));
+                CollectionAssert.AreEqual(Array.Empty<IoHash>(), putOp.Response["needs"].ToArray());
+
+                BatchOpsResponse.OpResponses errorOp = response.Results.First(r => r.OpId == 2);
+                Assert.IsNotNull(errorOp.Response);
+                Assert.AreEqual(404, errorOp.StatusCode);
+                Assert.IsTrue(!errorOp.Response["title"].Equals(CbField.Empty));
+                Assert.AreEqual("Object not found cbc3db15b9c8253f6106158962325c8fd848daef in bucket bucket namespace test-namespace", errorOp.Response["title"].AsString());
+                Assert.AreEqual(404, errorOp.Response["status"].AsInt32());
+
+                BatchOpsResponse.OpResponses headOp = response.Results.First(r => r.OpId == 3);
+                Assert.IsNotNull(headOp.Response);
+                Assert.AreEqual(200, headOp.StatusCode);
+                Assert.IsTrue(!headOp.Response["exists"].Equals(CbField.Empty));
+                Assert.IsTrue(headOp.Response["exists"].AsBool());
+            }
+        }
+    }
+}
diff --git a/Engine/Source/Programs/Horde/Horde.Server/Ddc/RequestHelper.cs b/Engine/Source/Programs/Horde/Horde.Server/Ddc/RequestHelper.cs
--- a/Engine/Source/Programs/Horde/Horde.Server/Ddc/RequestHelper.cs
+++ b/Engine/Source/Programs/Horde/Horde.Server/Ddc/RequestHelper.cs
@@ -29,7 +29,7 @@ namespace Horde.Server.Ddc
 			{
 				return new ForbidResult();
 			}
-			if (!aclActions.Any(x => !storageClient.Config.Authorize(x, user)))
+			if (aclActions.Any(x => !storageClient.Config.Authorize(x, user)))
 			{
 				return new ForbidResult();
 			}
diff --git a/Engine/Source/Programs/Horde/Horde.Server/Startup.cs b/Engine/Source/Programs/Horde/Horde.Server/Startup.cs
--- a/Engine/Source/Programs/Horde/Horde.Server/Startup.cs
+++ b/Engine/Source/Programs/Horde/Horde.Server/Startup.cs
@@ -521,10 +521,14 @@ namespace Horde.Server
 			}
 
 			// DDC
+			services.AddSingleton<IRequestHelper, RequestHelper>();
 			services.AddSingleton<IBlobService, BlobService>();
 			services.AddSingleton<IRefService, RefService>();
 			services.AddSingleton<IReferenceResolver, ReferenceResolver>();
 			services.AddSingleton<IContentIdService, ContentIdService>();
+			services.AddSingleton<BufferedPayloadFactory>();
+			services.AddSingleton<NginxRedirectHelper>();
+			services.AddSingleton<FormatResolver>();
 
 			// Secret providers
 			services.AddSingleton<SecretService>();
