Horde: Enable first batch of tests from Jupiter within Horde.

diff --git a/Engine/Source/Programs/Horde/Horde.Server.Tests/Ddc/Functional/Objects/ReferencesTests.cs b/Engine/Source/Programs/Horde/Horde.Server.Tests/Ddc/Functional/Objects/ReferencesTests.cs
--- a/Engine/Source/Programs/Horde/Horde.Server.Tests/Ddc/Functional/Objects/ReferencesTests.cs
+++ b/Engine/Source/Programs/Horde/Horde.Server.Tests/Ddc/Functional/Objects/ReferencesTests.cs
@@ -60,7 +60,6 @@ namespace Horde.Server.Tests.Ddc.FunctionalTests.References
 		}
 	}
 
-	[Ignore]
 	[TestClass]
     public class ReferencesTests : ControllerIntegrationTest
     {
@@ -86,7 +85,7 @@ namespace Horde.Server.Tests.Ddc.FunctionalTests.References
 			ConfigService configService = ServiceProvider.GetRequiredService<ConfigService>();
 			configService.Set(IoHash.Zero, globalConfig);
 
-			Service = ServiceProvider.GetRequiredService<IBlobService>();
+//			Service = ServiceProvider.GetRequiredService<IBlobService>();
 //			ReferencesStore = ServiceProvider.GetRequiredService<IRefService>();
 		}
 
@@ -200,7 +199,8 @@ namespace Horde.Server.Tests.Ddc.FunctionalTests.References
             }
         }
 
-        [TestMethod]
+		[Ignore]
+		[TestMethod]
         public async Task PutGetCompactBinaryAsync()
         {
             CbWriter writer = new CbWriter();
@@ -286,8 +286,8 @@ namespace Horde.Server.Tests.Ddc.FunctionalTests.References
             }
         }
 
-        
-        [TestMethod]
+		[Ignore]
+		[TestMethod]
         public async Task PutGetCompactBinaryFilteringAsync()
         {
             CbWriter writer = new CbWriter();
@@ -345,7 +345,8 @@ namespace Horde.Server.Tests.Ddc.FunctionalTests.References
             }
         }
 
-        [TestMethod]
+		[Ignore]
+		[TestMethod]
         public async Task PutLargeCompactBinaryAsync()
         {
             byte[] data = await File.ReadAllBytesAsync($"Objects/Payloads/lyra.cb");
@@ -364,7 +365,8 @@ namespace Horde.Server.Tests.Ddc.FunctionalTests.References
             Assert.AreEqual(4924, needsField.Count);
         }
 
-        [TestMethod]
+		[Ignore]
+		[TestMethod]
         public async Task PutGetCompactBinaryHierarchyAsync()
         {
             CbWriter childObjectWriter = new CbWriter();
@@ -514,7 +516,7 @@ namespace Horde.Server.Tests.Ddc.FunctionalTests.References
             }
         }
 
-        [TestMethod]
+		[TestMethod]
         public async Task ExistsChecksAsync()
         {
             const string objectContents = "This is treated as a opaque blob";
@@ -543,7 +545,7 @@ namespace Horde.Server.Tests.Ddc.FunctionalTests.References
             }
         }
 
-        [TestMethod]
+		[TestMethod]
         public async Task ExistsChecksMultipleAsync()
         {
             BucketId bucket = new BucketId("bucket");
@@ -577,7 +579,8 @@ namespace Horde.Server.Tests.Ddc.FunctionalTests.References
             }
         }
 
-        [TestMethod]
+		[Ignore]
+		[TestMethod]
         public async Task PutGetObjectHierarchyAsync()
         {
             string blobContents = "This is a string that is referenced as a blob";
@@ -735,8 +738,9 @@ namespace Horde.Server.Tests.Ddc.FunctionalTests.References
                 Assert.AreEqual(childDataObjectHash, BlobId.Parse(node["objectAttachment"]!.GetValue<string>()));
             }
         }
-        
-        [TestMethod]
+
+		[Ignore]
+		[TestMethod]
         public async Task PutPartialHierarchyAsync()
         {
             // do not submit the content of the blobs, which should be reported in the response of the put
@@ -822,8 +826,8 @@ namespace Horde.Server.Tests.Ddc.FunctionalTests.References
             }
         }
 
-        
-        [TestMethod]
+		[Ignore]
+		[TestMethod]
         public async Task PutContentIdMissingBlobAsync()
         {
             IContentIdService? contentIdStore = ServiceProvider.GetService<IContentIdService>();
@@ -900,8 +904,8 @@ namespace Horde.Server.Tests.Ddc.FunctionalTests.References
             }
         }
 
-        
-        [TestMethod]
+		[Ignore]
+		[TestMethod]
         public async Task PutMissingAttachmentComplexAsync()
         {
             string blobContents = "This is a string that is referenced as a blob but will not be uploaded";
@@ -948,7 +952,8 @@ namespace Horde.Server.Tests.Ddc.FunctionalTests.References
             }
         }
 
-        [TestMethod]
+		[Ignore]
+		[TestMethod]
         public async Task PutAndFinalizeAsync()
         {
             BucketId bucket = new BucketId("bucket");
@@ -1050,7 +1055,8 @@ namespace Horde.Server.Tests.Ddc.FunctionalTests.References
             }
         }
 
-        [TestMethod]
+		[Ignore]
+		[TestMethod]
         public async Task GetMissingContentIdRecordAsync()
         {
             string blobContents = "This is a blob";
@@ -1160,7 +1166,8 @@ namespace Horde.Server.Tests.Ddc.FunctionalTests.References
             }
         }
 
-        [TestMethod]
+		[Ignore]
+		[TestMethod]
         public async Task GetMissingCompressedBufferAttachmentAsync()
         {
             CbObject cbObjectAttachment = CbObject.Build(writer => writer.WriteString("ValueField", "This field has a value"));
@@ -1264,7 +1271,8 @@ namespace Horde.Server.Tests.Ddc.FunctionalTests.References
             }
         }
 
-        [TestMethod]
+		[Ignore]
+		[TestMethod]
         public async Task GetMissingBlobRecordAsync()
         {
             string blobContents = "This is a blob";
@@ -1369,7 +1377,8 @@ namespace Horde.Server.Tests.Ddc.FunctionalTests.References
             }
         }
 
-        [TestMethod]
+		[Ignore]
+		[TestMethod]
         public async Task DeleteObjectAsync()
         {
             const string objectContents = "This is treated as a opaque blob";
@@ -1412,8 +1421,8 @@ namespace Horde.Server.Tests.Ddc.FunctionalTests.References
             }
         }
 
-        
-        [TestMethod]
+		[Ignore]
+		[TestMethod]
         public async Task DropBucketAsync()
         {
             const string BucketToDelete = "delete-bucket";
@@ -1470,8 +1479,8 @@ namespace Horde.Server.Tests.Ddc.FunctionalTests.References
             }
         }
 
-        
-        [TestMethod]
+		[Ignore]
+		[TestMethod]
         public async Task DeleteNamespaceAsync()
         {
             const string NamespaceToBeDeleted = "test-delete-namespace";
@@ -1537,8 +1546,8 @@ namespace Horde.Server.Tests.Ddc.FunctionalTests.References
             }
         }
 
-        
-        [TestMethod]
+		[Ignore]
+		[TestMethod]
         public async Task ListNamespacesAsync()
         {
             const string objectContents = "This is treated as a opaque blob";
@@ -1565,8 +1574,9 @@ namespace Horde.Server.Tests.Ddc.FunctionalTests.References
                 Assert.IsTrue(response.Namespaces.Contains(TestNamespace));
             }
         }
-        
-        [TestMethod]
+
+		[Ignore]
+		[TestMethod]
         public async Task GetOldRecordsAsync()
         {
             const string objectContents = "This is treated as a opaque blob";
@@ -1590,8 +1600,8 @@ namespace Horde.Server.Tests.Ddc.FunctionalTests.References
             Assert.AreEqual(key, oldName);
             Assert.AreEqual("bucket", oldBucket.ToString());
         }
-        
-        [TestMethod]
+
+		[TestMethod]
         public async Task BatchJsonRequestAsync()
         {
             // verifies that json request against the batch endpoint will fail
@@ -1605,7 +1615,7 @@ namespace Horde.Server.Tests.Ddc.FunctionalTests.References
             }
         }
 
-        [TestMethod]
+		[TestMethod]
         public async Task BatchErrorOperationsAsync()
         {
             // seed some data
@@ -1678,7 +1688,7 @@ namespace Horde.Server.Tests.Ddc.FunctionalTests.References
             }
         }
 
-        [TestMethod]
+		[TestMethod]
         public async Task BatchGetOperationsAsync()
         {
             // seed some data
@@ -1823,7 +1833,7 @@ namespace Horde.Server.Tests.Ddc.FunctionalTests.References
             }
         }
 
-        [TestMethod]
+		[TestMethod]
         public async Task BatchHeadOperationsAsync()
         {
             // seed some data
@@ -1939,7 +1949,7 @@ namespace Horde.Server.Tests.Ddc.FunctionalTests.References
             }
         }
 
-        [TestMethod]
+		[TestMethod]
         public async Task BatchPutOperationsAsync()
         {
             BucketId bucket = new BucketId("bucket");
@@ -2014,7 +2024,7 @@ namespace Horde.Server.Tests.Ddc.FunctionalTests.References
             }
         }
 
-        [TestMethod]
+		[TestMethod]
         public async Task BatchMixedOperationsAsync()
         {
             // seed some data
diff --git a/Engine/Source/Programs/Horde/Horde.Server/Ddc/RefNode.cs b/Engine/Source/Programs/Horde/Horde.Server/Ddc/RefNode.cs
deleted file mode 100644
diff --git a/Engine/Source/Programs/Horde/Horde.Server/Ddc/RefService.cs b/Engine/Source/Programs/Horde/Horde.Server/Ddc/RefService.cs
--- a/Engine/Source/Programs/Horde/Horde.Server/Ddc/RefService.cs
+++ b/Engine/Source/Programs/Horde/Horde.Server/Ddc/RefService.cs
@@ -81,7 +81,10 @@ namespace Horde.Server.Ddc
 					throw new BlobNotFoundException(ns, blobHash);
 				}
 
-				await storageClient.WriteRefTargetAsync(GetRefName(bucket, key), handle, cancellationToken: cancellationToken);
+				DdcRefNode refNode = new DdcRefNode(blobHash.Hash);
+				refNode.References.Add((blobHash.Hash, handle));
+
+				await storageClient.WriteNodeAsync(GetRefName(bucket, key), refNode, cancellationToken: cancellationToken);
 			}
 
 			return (Array.Empty<ContentId>(), missingBlobIds.ToArray());
@@ -91,21 +94,24 @@ namespace Horde.Server.Ddc
 		{
 			IStorageClient storageClient = await _storageService.GetClientAsync(ns, cancellationToken);
 
-			RefNode? node = await storageClient.TryReadNodeAsync<RefNode>(GetRefName(bucket, key), cancellationToken: cancellationToken);
+			DdcRefNode? node = await storageClient.TryReadNodeAsync<DdcRefNode>(GetRefName(bucket, key), cancellationToken: cancellationToken);
 			if (node == null)
 			{
 				throw new RefNotFoundException(ns, bucket, key);
 			}
 
-			RefRecord record = new RefRecord(ns, bucket, key, DateTime.UtcNow, null, new BlobId(node.RootHash), false);
-			return (record, null);
+			BlobData data = await node.References.First(x => x.Hash == node.RootHash).Handle.ReadAsync(cancellationToken);
+			BlobContents contents = new BlobContents(data.Data.ToArray());
+
+			RefRecord record = new RefRecord(ns, bucket, key, DateTime.UtcNow, null, new BlobId(node.RootHash), true);
+			return (record, contents);
 		}
 
 		public async Task<List<BlobId>> GetReferencedBlobsAsync(NamespaceId ns, BucketId bucket, RefId key, CancellationToken cancellationToken)
 		{
 			IStorageClient storageClient = await _storageService.GetClientAsync(ns, cancellationToken);
 
-			RefNode? node = await storageClient.TryReadNodeAsync<RefNode>(GetRefName(bucket, key), cancellationToken: cancellationToken);
+			DdcRefNode? node = await storageClient.TryReadNodeAsync<DdcRefNode>(GetRefName(bucket, key), cancellationToken: cancellationToken);
 			if (node == null)
 			{
 				throw new RefNotFoundException(ns, bucket, key);
diff --git a/Engine/Source/Programs/Horde/Horde.Server/Startup.cs b/Engine/Source/Programs/Horde/Horde.Server/Startup.cs
--- a/Engine/Source/Programs/Horde/Horde.Server/Startup.cs
+++ b/Engine/Source/Programs/Horde/Horde.Server/Startup.cs
@@ -105,6 +105,8 @@ using EpicGames.Horde;
 using EpicGames.Horde.Api;
 using EpicGames.Horde.Storage.Bundles;
 using Horde.Server.Ddc;
+using System.Net.Mime;
+using Microsoft.Extensions.Logging.Abstractions;
 
 namespace Horde.Server
 {
@@ -789,7 +791,21 @@ namespace Horde.Server
 				options.InputFormatters.Add(new CbInputFormatter());
 				options.OutputFormatters.Add(new CbOutputFormatter());
 				options.OutputFormatters.Insert(0, new CbPreferredOutputFormatter());
+				options.OutputFormatters.Add(new RawOutputFormatter(NullLogger.Instance));
+				options.FormatterMappings.SetMediaTypeMappingForFormat("raw", MediaTypeNames.Application.Octet);
 				options.FormatterMappings.SetMediaTypeMappingForFormat("uecb", CustomMediaTypeNames.UnrealCompactBinary);
+				options.FormatterMappings.SetMediaTypeMappingForFormat("uecbpkg", CustomMediaTypeNames.UnrealCompactBinaryPackage);
+			}).ConfigureApiBehaviorOptions(options =>
+			{
+				options.InvalidModelStateResponseFactory = context =>
+				{
+					BadRequestObjectResult result = new BadRequestObjectResult(context.ModelState);
+					// always return errors as json objects
+					// we could allow more types here, but we do not want raw for instance
+					result.ContentTypes.Add(MediaTypeNames.Application.Json);
+
+					return result;
+				};
 			});
 
 			services.AddSwaggerGen(config =>
diff --git a/Engine/Source/Programs/Horde/Horde.Server/Utilities/RawOutputFormatter.cs b/Engine/Source/Programs/Horde/Horde.Server/Utilities/RawOutputFormatter.cs
new file mode 100644
--- /dev/null
+++ b/Engine/Source/Programs/Horde/Horde.Server/Utilities/RawOutputFormatter.cs
@@ -0,0 +1,84 @@
+﻿// Copyright Epic Games, Inc. All Rights Reserved.
+
+using System;
+using System.Reflection;
+using System.Text;
+using System.Threading.Tasks;
+using Microsoft.AspNetCore.Http;
+using Microsoft.AspNetCore.Mvc;
+using Microsoft.AspNetCore.Mvc.Formatters;
+using Microsoft.Extensions.Logging;
+using Microsoft.Net.Http.Headers;
+
+namespace Horde.Server.Utilities
+{
+	class RawOutputFormatter : OutputFormatter
+	{
+		private readonly ILogger _logger;
+
+		public RawOutputFormatter(ILogger logger)
+		{
+			_logger = logger;
+			SupportedMediaTypes.Add(MediaTypeHeaderValue.Parse("application/octet-stream"));
+		}
+
+		public override async Task WriteResponseBodyAsync(OutputFormatterWriteContext context)
+		{
+			if (context.Object == null)
+			{
+				return;
+			}
+
+			HttpResponse response = context.HttpContext.Response;
+
+			object o = context.Object;
+			PropertyInfo? outputProperty = null;
+			foreach (PropertyInfo propertyInfo in o.GetType().GetProperties())
+			{
+				object[] attr =
+					propertyInfo.GetCustomAttributes(attributeType: typeof(RawOutputPropertyAttribute), true);
+				if (attr.Length != 0)
+				{
+					outputProperty = propertyInfo;
+				}
+			}
+
+			if (o is ProblemDetails p)
+			{
+				context.ContentType = "text/plain";
+				await response.Body.WriteAsync(Encoding.UTF8.GetBytes(p.Title!));
+				return;
+			}
+
+			if (outputProperty == null)
+			{
+				context.ContentType = "text/plain";
+
+				_logger.LogDebug("No RawOutputProperty detected, unable to use raw formatter on this object of type {Type}", o.GetType().FullName);
+				return;
+			}
+
+			if (outputProperty.PropertyType != typeof(byte[]))
+			{
+				throw new Exception("RawOutputProperty can only be set on a byte[]");
+			}
+
+			object? nullableValue = outputProperty.GetValue(o);
+			if (nullableValue == null)
+			{
+				throw new Exception("RawOutputProperty field can not be null");
+			}
+
+			byte[] value = (byte[]) nullableValue;
+			await response.Body.WriteAsync(value);
+		}
+	}
+
+	/// <summary>
+	///  Used to indicate that the field in question is what we want to return as a raw set of bytes
+	/// </summary>
+	[AttributeUsage(AttributeTargets.Property)]
+	public sealed class RawOutputPropertyAttribute : Attribute
+	{
+	}
+}
diff --git a/Engine/Source/Programs/Shared/EpicGames.Horde/Storage/Nodes/DdcRefNode.cs b/Engine/Source/Programs/Shared/EpicGames.Horde/Storage/Nodes/DdcRefNode.cs
new file mode 100644
--- /dev/null
+++ b/Engine/Source/Programs/Shared/EpicGames.Horde/Storage/Nodes/DdcRefNode.cs
@@ -0,0 +1,56 @@
+// Copyright Epic Games, Inc. All Rights Reserved.
+
+using System.Collections.Generic;
+using EpicGames.Core;
+using EpicGames.Horde.Storage;
+
+namespace Horde.Server.Ddc
+{
+	/// <summary>
+	/// A node containing ref data
+	/// </summary>
+	[NodeType("{0C7E5F25-454B-4B55-9B4A-F4635106D074}", 1)]
+	public class DdcRefNode : Node
+	{
+		/// <summary>
+		/// Hash of the root node
+		/// </summary>
+		public IoHash RootHash { get; }
+
+		/// <summary>
+		/// References to attachments
+		/// </summary>
+		public List<(IoHash Hash, BlobHandle Handle)> References { get; }
+
+		/// <summary>
+		/// Constructor
+		/// </summary>
+		public DdcRefNode(IoHash rootHash)
+		{
+			RootHash = rootHash;
+			References = new List<(IoHash, BlobHandle)>();
+		}
+
+		/// <summary>
+		/// Constructor
+		/// </summary>
+		public DdcRefNode(NodeReader reader)
+		{
+			RootHash = reader.ReadIoHash();
+			References = reader.ReadList(x => (reader.ReadIoHash(), reader.ReadNodeHandle()));
+		}
+
+		/// <inheritdoc/>
+		public override void Serialize(NodeWriter writer)
+		{
+			writer.WriteIoHash(RootHash);
+			writer.WriteList(References, x => WriteReference(writer, x.Hash, x.Handle));
+		}
+
+		static void WriteReference(NodeWriter writer, IoHash hash, BlobHandle handle)
+		{
+			writer.WriteIoHash(hash);
+			writer.WriteNodeHandle(handle);
+		}
+	}
+}
