Horde: Fix 4 more DDC tests.

diff --git a/Engine/Source/Programs/Horde/Horde.Server.Tests/Ddc/Functional/Objects/ReferencesTests.cs b/Engine/Source/Programs/Horde/Horde.Server.Tests/Ddc/Functional/Objects/ReferencesTests.cs
--- a/Engine/Source/Programs/Horde/Horde.Server.Tests/Ddc/Functional/Objects/ReferencesTests.cs
+++ b/Engine/Source/Programs/Horde/Horde.Server.Tests/Ddc/Functional/Objects/ReferencesTests.cs
@@ -199,7 +199,6 @@ namespace Horde.Server.Tests.Ddc.FunctionalTests.References
             }
         }
 
-		[Ignore]
 		[TestMethod]
         public async Task PutGetCompactBinaryAsync()
         {
@@ -235,12 +234,14 @@ namespace Horde.Server.Tests.Ddc.FunctionalTests.References
             {
                 BucketId bucket = new BucketId("bucket");
 
-                RefRecord objectRecord = await ReferencesStore.GetAsync(TestNamespace, bucket, key, IReferencesStore.FieldFlags.IncludePayload);
+				IRefService refService = ServiceProvider.GetRequiredService<IRefService>();
+				(RefRecord objectRecord, BlobContents? contents) = await refService.GetAsync(TestNamespace, bucket, key, Array.Empty<string>());
 
                 Assert.IsTrue(objectRecord.IsFinalized);
                 Assert.AreEqual(key, objectRecord.Name);
                 Assert.AreEqual(objectHash, objectRecord.BlobIdentifier);
-                Assert.IsNotNull(objectRecord.InlinePayload);
+				// Note: Horde does not use the InlinePayload field
+				// Assert.IsNotNull(objectRecord.InlinePayload);
             }
 
             {
@@ -286,7 +287,6 @@ namespace Horde.Server.Tests.Ddc.FunctionalTests.References
             }
         }
 
-		[Ignore]
 		[TestMethod]
         public async Task PutGetCompactBinaryFilteringAsync()
         {
@@ -322,7 +322,8 @@ namespace Horde.Server.Tests.Ddc.FunctionalTests.References
             {
                 BucketId bucket = new BucketId("bucket");
 
-                RefRecord objectRecord = await ReferencesStore.GetAsync(TestNamespace, bucket, key, IReferencesStore.FieldFlags.None);
+				IRefService refService = ServiceProvider.GetRequiredService<IRefService>();
+				(RefRecord objectRecord, _) = await refService.GetAsync(TestNamespace, bucket, key, Array.Empty<string>());
 
                 Assert.IsTrue(objectRecord.IsFinalized);
                 Assert.AreEqual(key, objectRecord.Name);
@@ -365,7 +366,6 @@ namespace Horde.Server.Tests.Ddc.FunctionalTests.References
             Assert.AreEqual(4924, needsField.Count);
         }
 
-		[Ignore]
 		[TestMethod]
         public async Task PutGetCompactBinaryHierarchyAsync()
         {
@@ -453,12 +453,13 @@ namespace Horde.Server.Tests.Ddc.FunctionalTests.References
             {
                 BucketId bucket = new BucketId("bucket");
 
-                RefRecord objectRecord = await ReferencesStore.GetAsync(TestNamespace, bucket, key, IReferencesStore.FieldFlags.IncludePayload);
+				IRefService refService = ServiceProvider.GetRequiredService<IRefService>();
+				(RefRecord objectRecord, _) = await refService.GetAsync(TestNamespace, bucket, key, Array.Empty<string>());
 
                 Assert.IsTrue(objectRecord.IsFinalized);
                 Assert.AreEqual(key, objectRecord.Name);
                 Assert.AreEqual(parentObjectHash, objectRecord.BlobIdentifier);
-                Assert.IsNotNull(objectRecord.InlinePayload);
+//                Assert.IsNotNull(objectRecord.InlinePayload);
             }
 
             {
@@ -1377,7 +1378,6 @@ namespace Horde.Server.Tests.Ddc.FunctionalTests.References
             }
         }
 
-		[Ignore]
 		[TestMethod]
         public async Task DeleteObjectAsync()
         {
diff --git a/Engine/Source/Programs/Horde/Horde.Server/Commands/Bundles/BundleWriteTest.cs b/Engine/Source/Programs/Horde/Horde.Server/Commands/Bundles/BundleWriteTest.cs
--- a/Engine/Source/Programs/Horde/Horde.Server/Commands/Bundles/BundleWriteTest.cs
+++ b/Engine/Source/Programs/Horde/Horde.Server/Commands/Bundles/BundleWriteTest.cs
@@ -25,7 +25,7 @@ namespace Horde.Server.Commands.Bundles
 			{
 			}
 
-			public Task DeleteRefAsync(RefName name, CancellationToken cancellationToken = default) => Task.CompletedTask;
+			public Task<bool> DeleteRefAsync(RefName name, CancellationToken cancellationToken = default) => Task.FromResult(true);
 			public Task<Bundle> ReadBundleAsync(BundleLocator locator, CancellationToken cancellationToken = default) => throw new NotImplementedException();
 			public Task<ReadOnlyMemory<byte>> ReadBundleRangeAsync(BundleLocator locator, int offset, int length, CancellationToken cancellationToken = default) => throw new NotImplementedException();
 			public Task AddAliasAsync(Utf8String name, BlobHandle handle, CancellationToken cancellationToken = default) => throw new NotImplementedException();
diff --git a/Engine/Source/Programs/Horde/Horde.Server/Ddc/RefService.cs b/Engine/Source/Programs/Horde/Horde.Server/Ddc/RefService.cs
--- a/Engine/Source/Programs/Horde/Horde.Server/Ddc/RefService.cs
+++ b/Engine/Source/Programs/Horde/Horde.Server/Ddc/RefService.cs
@@ -28,8 +28,7 @@ namespace Horde.Server.Ddc
 		public async Task<bool> DeleteAsync(NamespaceId ns, BucketId bucket, RefId key, CancellationToken cancellationToken)
 		{
 			IStorageClient storageClient = await _storageService.GetClientAsync(ns, cancellationToken);
-			await storageClient.DeleteRefAsync(GetRefName(bucket, key), cancellationToken);
-			return true;
+			return await storageClient.DeleteRefAsync(GetRefName(bucket, key), cancellationToken);
 		}
 
 		public async Task<bool> ExistsAsync(NamespaceId ns, BucketId bucket, RefId key, CancellationToken cancellationToken)
diff --git a/Engine/Source/Programs/Horde/Horde.Server/Storage/StorageService.cs b/Engine/Source/Programs/Horde/Horde.Server/Storage/StorageService.cs
--- a/Engine/Source/Programs/Horde/Horde.Server/Storage/StorageService.cs
+++ b/Engine/Source/Programs/Horde/Horde.Server/Storage/StorageService.cs
@@ -260,7 +260,7 @@ namespace Horde.Server.Storage
 			public override Task WriteRefTargetAsync(RefName name, BundleNodeLocator target, RefOptions? options = null, CancellationToken cancellationToken = default) => _outer.WriteRefTargetAsync(NamespaceId, name, target, options, cancellationToken);
 
 			/// <inheritdoc/>
-			public override Task DeleteRefAsync(RefName name, CancellationToken cancellationToken = default) => _outer.DeleteRefAsync(NamespaceId, name, cancellationToken);
+			public override Task<bool> DeleteRefAsync(RefName name, CancellationToken cancellationToken = default) => _outer.DeleteRefAsync(NamespaceId, name, cancellationToken);
 
 			#endregion
 		}
@@ -819,10 +819,10 @@ namespace Horde.Server.Storage
 		}
 
 		/// <inheritdoc/>
-		async Task DeleteRefAsync(NamespaceId namespaceId, RefName name, CancellationToken cancellationToken = default)
+		async Task<bool> DeleteRefAsync(NamespaceId namespaceId, RefName name, CancellationToken cancellationToken = default)
 		{
 			FilterDefinition<RefInfo> filter = Builders<RefInfo>.Filter.Expr(x => x.NamespaceId == namespaceId && x.Name == name);
-			await DeleteRefInternalAsync(namespaceId, name, filter, cancellationToken);
+			return await DeleteRefInternalAsync(namespaceId, name, filter, cancellationToken);
 		}
 
 		/// <summary>
@@ -834,15 +834,19 @@ namespace Horde.Server.Storage
 			await DeleteRefInternalAsync(refDocument.NamespaceId, refDocument.Name, filter, cancellationToken);
 		}
 
-		async Task DeleteRefInternalAsync(NamespaceId namespaceId, RefName name, FilterDefinition<RefInfo> filter, CancellationToken cancellationToken = default)
+		async Task<bool> DeleteRefInternalAsync(NamespaceId namespaceId, RefName name, FilterDefinition<RefInfo> filter, CancellationToken cancellationToken = default)
 		{
 			RefInfo? oldRefInfo = await _refCollection.FindOneAndDeleteAsync<RefInfo>(filter, cancellationToken: cancellationToken);
+			AddRefToCache(namespaceId, name, default);
+
 			if (oldRefInfo != null)
 			{
 				_logger.LogInformation("Deleted ref {NamespaceId}:{RefName}", namespaceId, name);
 				AddGcCheckRecord(namespaceId, oldRefInfo.BlobInfoId);
+				return true;
 			}
-			AddRefToCache(namespaceId, name, default);
+
+			return false;
 		}
 
 		/// <inheritdoc/>
diff --git a/Engine/Source/Programs/Shared/EpicGames.Horde/Compute/AgentStorageClient.cs b/Engine/Source/Programs/Shared/EpicGames.Horde/Compute/AgentStorageClient.cs
--- a/Engine/Source/Programs/Shared/EpicGames.Horde/Compute/AgentStorageClient.cs
+++ b/Engine/Source/Programs/Shared/EpicGames.Horde/Compute/AgentStorageClient.cs
@@ -93,7 +93,7 @@ namespace EpicGames.Horde.Compute
 		#region Refs
 
 		/// <inheritdoc/>
-		public override Task DeleteRefAsync(RefName name, CancellationToken cancellationToken = default)
+		public override Task<bool> DeleteRefAsync(RefName name, CancellationToken cancellationToken = default)
 		{
 			throw new NotSupportedException();
 		}
diff --git a/Engine/Source/Programs/Shared/EpicGames.Horde/Storage/Backends/FileStorageClient.cs b/Engine/Source/Programs/Shared/EpicGames.Horde/Storage/Backends/FileStorageClient.cs
--- a/Engine/Source/Programs/Shared/EpicGames.Horde/Storage/Backends/FileStorageClient.cs
+++ b/Engine/Source/Programs/Shared/EpicGames.Horde/Storage/Backends/FileStorageClient.cs
@@ -14,6 +14,7 @@ using EpicGames.Core;
 using Microsoft.Extensions.Caching.Memory;
 using Microsoft.Extensions.Logging;
 using EpicGames.Horde.Storage.Bundles;
+using Microsoft.CodeAnalysis.FlowAnalysis;
 
 namespace EpicGames.Horde.Storage.Backends
 {
@@ -123,11 +124,15 @@ namespace EpicGames.Horde.Storage.Backends
 		#region Refs
 
 		/// <inheritdoc/>
-		public override Task DeleteRefAsync(RefName name, CancellationToken cancellationToken = default)
+		public override Task<bool> DeleteRefAsync(RefName name, CancellationToken cancellationToken = default)
 		{
-			FileReference file = GetRefFile(name);
-			FileReference.Delete(file);
-			return Task.CompletedTask;
+			FileInfo file = GetRefFile(name).ToFileInfo();
+			if (file.Exists)
+			{
+				file.Delete();
+				return Task.FromResult(true);
+			}
+			return Task.FromResult(false);
 		}
 
 		/// <inheritdoc/>
diff --git a/Engine/Source/Programs/Shared/EpicGames.Horde/Storage/Backends/HttpStorageClient.cs b/Engine/Source/Programs/Shared/EpicGames.Horde/Storage/Backends/HttpStorageClient.cs
--- a/Engine/Source/Programs/Shared/EpicGames.Horde/Storage/Backends/HttpStorageClient.cs
+++ b/Engine/Source/Programs/Shared/EpicGames.Horde/Storage/Backends/HttpStorageClient.cs
@@ -235,7 +235,7 @@ namespace EpicGames.Horde.Storage.Backends
 		#region Refs
 
 		/// <inheritdoc/>
-		public override async Task DeleteRefAsync(RefName name, CancellationToken cancellationToken)
+		public override async Task<bool> DeleteRefAsync(RefName name, CancellationToken cancellationToken)
 		{
 			_logger.LogDebug("Deleting ref {RefName}", name);
 			using (HttpClient httpClient = _createClient())
@@ -244,7 +244,17 @@ namespace EpicGames.Horde.Storage.Backends
 				{
 					using (HttpResponseMessage response = await httpClient.SendAsync(request, cancellationToken))
 					{
+						if (response.IsSuccessStatusCode)
+						{
+							return true;
+						}
+						if (response.StatusCode == HttpStatusCode.NotFound)
+						{
+							return false;
+						}
+
 						response.EnsureSuccessStatusCode();
+						return false;
 					}
 				}
 			}
diff --git a/Engine/Source/Programs/Shared/EpicGames.Horde/Storage/Backends/MemoryStorageClient.cs b/Engine/Source/Programs/Shared/EpicGames.Horde/Storage/Backends/MemoryStorageClient.cs
--- a/Engine/Source/Programs/Shared/EpicGames.Horde/Storage/Backends/MemoryStorageClient.cs
+++ b/Engine/Source/Programs/Shared/EpicGames.Horde/Storage/Backends/MemoryStorageClient.cs
@@ -116,7 +116,7 @@ namespace EpicGames.Horde.Storage.Backends
 		#region Refs
 
 		/// <inheritdoc/>
-		public override Task DeleteRefAsync(RefName name, CancellationToken cancellationToken) => Task.FromResult(_refs.TryRemove(name, out _));
+		public override Task<bool> DeleteRefAsync(RefName name, CancellationToken cancellationToken) => Task.FromResult(_refs.TryRemove(name, out _));
 
 		/// <inheritdoc/>
 		public override Task<BundleNodeHandle?> TryReadRefTargetAsync(RefName name, RefCacheTime cacheTime = default, CancellationToken cancellationToken = default)
diff --git a/Engine/Source/Programs/Shared/EpicGames.Horde/Storage/Bundles/BundleStorageClient.cs b/Engine/Source/Programs/Shared/EpicGames.Horde/Storage/Bundles/BundleStorageClient.cs
--- a/Engine/Source/Programs/Shared/EpicGames.Horde/Storage/Bundles/BundleStorageClient.cs
+++ b/Engine/Source/Programs/Shared/EpicGames.Horde/Storage/Bundles/BundleStorageClient.cs
@@ -102,7 +102,7 @@ namespace EpicGames.Horde.Storage.Bundles
 		#region Refs
 
 		/// <inheritdoc/>
-		public abstract Task DeleteRefAsync(RefName name, CancellationToken cancellationToken = default);
+		public abstract Task<bool> DeleteRefAsync(RefName name, CancellationToken cancellationToken = default);
 
 		/// <inheritdoc cref="StorageClientExtensions.ReadRefTargetAsync(IStorageClient, RefName, RefCacheTime, CancellationToken)"/>
 		public async Task<BundleNodeHandle> ReadRefTargetAsync(RefName name, RefCacheTime cacheTime = default, CancellationToken cancellationToken = default)
diff --git a/Engine/Source/Programs/Shared/EpicGames.Horde/Storage/IStorageClient.cs b/Engine/Source/Programs/Shared/EpicGames.Horde/Storage/IStorageClient.cs
--- a/Engine/Source/Programs/Shared/EpicGames.Horde/Storage/IStorageClient.cs
+++ b/Engine/Source/Programs/Shared/EpicGames.Horde/Storage/IStorageClient.cs
@@ -138,7 +138,7 @@ namespace EpicGames.Horde.Storage
 		/// </summary>
 		/// <param name="name">The ref identifier</param>
 		/// <param name="cancellationToken">Cancellation token for the operation</param>
-		Task DeleteRefAsync(RefName name, CancellationToken cancellationToken = default);
+		Task<bool> DeleteRefAsync(RefName name, CancellationToken cancellationToken = default);
 
 		#endregion
 	}
