Horde: Fix remaining DDC tests. Some tests still ignored but are not currently supported in Horde.

diff --git a/Engine/Source/Programs/Horde/Horde.Server.Tests/Ddc/Functional/Objects/ReferencesTests.cs b/Engine/Source/Programs/Horde/Horde.Server.Tests/Ddc/Functional/Objects/ReferencesTests.cs
--- a/Engine/Source/Programs/Horde/Horde.Server.Tests/Ddc/Functional/Objects/ReferencesTests.cs
+++ b/Engine/Source/Programs/Horde/Horde.Server.Tests/Ddc/Functional/Objects/ReferencesTests.cs
@@ -66,7 +66,9 @@ namespace Horde.Server.Tests.Ddc.FunctionalTests.References
 		//		private static TestServer? _server;
 		private static HttpClient? _httpClient;
 
-		protected IBlobService Service { get; set; } = null!;
+		protected IBlobService BlobService => ServiceProvider.GetRequiredService<IBlobService>();
+		protected IRefService RefService => ServiceProvider.GetRequiredService<IRefService>();
+
 		protected IReferencesStore ReferencesStore { get; set; } = null!;
 		protected NamespaceId TestNamespace { get; } = new NamespaceId("test-namespace");
 
@@ -84,9 +86,6 @@ namespace Horde.Server.Tests.Ddc.FunctionalTests.References
 
 			ConfigService configService = ServiceProvider.GetRequiredService<ConfigService>();
 			configService.Set(IoHash.Zero, globalConfig);
-
-//			Service = ServiceProvider.GetRequiredService<IBlobService>();
-//			ReferencesStore = ServiceProvider.GetRequiredService<IRefService>();
 		}
 
 		[TestMethod]
@@ -346,11 +345,10 @@ namespace Horde.Server.Tests.Ddc.FunctionalTests.References
             }
         }
 
-		[Ignore]
 		[TestMethod]
         public async Task PutLargeCompactBinaryAsync()
         {
-            byte[] data = await File.ReadAllBytesAsync($"Objects/Payloads/lyra.cb");
+            byte[] data = await File.ReadAllBytesAsync($"Ddc/Functional/Objects/Payloads/lyra.cb");
             BlobId objectHash = BlobId.FromBlob(data);
             RefId key = RefId.FromName("largeCompactBinary");
             using HttpContent requestContent = new ByteArrayContent(data);
@@ -580,19 +578,18 @@ namespace Horde.Server.Tests.Ddc.FunctionalTests.References
             }
         }
 
-		[Ignore]
 		[TestMethod]
         public async Task PutGetObjectHierarchyAsync()
         {
             string blobContents = "This is a string that is referenced as a blob";
             byte[] blobData = Encoding.ASCII.GetBytes(blobContents);
             BlobId blobHash = BlobId.FromBlob(blobData);
-            await Service.PutObjectAsync(TestNamespace, blobData, blobHash);
+            await BlobService.PutObjectAsync(TestNamespace, blobData, blobHash);
 
             string blobContentsChild = "This string is also referenced as a blob but from a child object";
             byte[] dataChild = Encoding.ASCII.GetBytes(blobContentsChild);
             BlobId blobHashChild = BlobId.FromBlob(dataChild);
-            await Service.PutObjectAsync(TestNamespace, dataChild, blobHashChild);
+            await BlobService.PutObjectAsync(TestNamespace, dataChild, blobHashChild);
 
             CbWriter writerChild = new CbWriter();
             writerChild.BeginObject();
@@ -601,7 +598,7 @@ namespace Horde.Server.Tests.Ddc.FunctionalTests.References
 
             byte[] childDataObject = writerChild.ToByteArray();
             BlobId childDataObjectHash = BlobId.FromBlob(childDataObject);
-            await Service.PutObjectAsync(TestNamespace, childDataObject, childDataObjectHash);
+            await BlobService.PutObjectAsync(TestNamespace, childDataObject, childDataObjectHash);
 
             CbWriter writerParent = new CbWriter();
             writerParent.BeginObject();
@@ -639,12 +636,12 @@ namespace Horde.Server.Tests.Ddc.FunctionalTests.References
             {
                 BucketId bucket = new BucketId("bucket");
 
-                RefRecord objectRecord = await ReferencesStore.GetAsync(TestNamespace, bucket, key, IReferencesStore.FieldFlags.IncludePayload);
+				(RefRecord objectRecord, _) = await RefService.GetAsync(TestNamespace, bucket, key, Array.Empty<string>());
 
                 Assert.IsTrue(objectRecord.IsFinalized);
                 Assert.AreEqual(key, objectRecord.Name);
                 Assert.AreEqual(objectHash, objectRecord.BlobIdentifier);
-                Assert.IsNotNull(objectRecord.InlinePayload);
+//                Assert.IsNotNull(objectRecord.InlinePayload);
             }
 
             // verify attachments
@@ -740,7 +737,6 @@ namespace Horde.Server.Tests.Ddc.FunctionalTests.References
             }
         }
 
-		[Ignore]
 		[TestMethod]
         public async Task PutPartialHierarchyAsync()
         {
@@ -760,7 +756,7 @@ namespace Horde.Server.Tests.Ddc.FunctionalTests.References
 
             byte[] childDataObject = writerChild.ToByteArray();
             BlobId childDataObjectHash = BlobId.FromBlob(childDataObject);
-            await Service.PutObjectAsync(TestNamespace, childDataObject, childDataObjectHash);
+            await BlobService.PutObjectAsync(TestNamespace, childDataObject, childDataObjectHash);
 
             CbWriter writerParent = new CbWriter();
             writerParent.BeginObject();
@@ -827,7 +823,6 @@ namespace Horde.Server.Tests.Ddc.FunctionalTests.References
             }
         }
 
-		[Ignore]
 		[TestMethod]
         public async Task PutContentIdMissingBlobAsync()
         {
@@ -905,7 +900,6 @@ namespace Horde.Server.Tests.Ddc.FunctionalTests.References
             }
         }
 
-		[Ignore]
 		[TestMethod]
         public async Task PutMissingAttachmentComplexAsync()
         {
@@ -953,7 +947,6 @@ namespace Horde.Server.Tests.Ddc.FunctionalTests.References
             }
         }
 
-		[Ignore]
 		[TestMethod]
         public async Task PutAndFinalizeAsync()
         {
@@ -976,7 +969,7 @@ namespace Horde.Server.Tests.Ddc.FunctionalTests.References
 
             byte[] childDataObject = writerChild.ToByteArray();
             BlobId childDataObjectHash = BlobId.FromBlob(childDataObject);
-            await Service.PutObjectAsync(TestNamespace, childDataObject, childDataObjectHash);
+            await BlobService.PutObjectAsync(TestNamespace, childDataObject, childDataObjectHash);
 
             CbWriter writerParent = new CbWriter();
             writerParent.BeginObject();
@@ -1012,18 +1005,13 @@ namespace Horde.Server.Tests.Ddc.FunctionalTests.References
 
             // check that actual internal representation
             {
-                RefRecord objectRecord = await ReferencesStore.GetAsync(TestNamespace, bucket, key, IReferencesStore.FieldFlags.IncludePayload);
-
-                Assert.IsFalse(objectRecord.IsFinalized);
-                Assert.AreEqual(key, objectRecord.Name);
-                Assert.AreEqual(objectHash, objectRecord.BlobIdentifier);
-                Assert.IsNotNull(objectRecord.InlinePayload);
+				await Assert.ThrowsExceptionAsync<RefNotFoundException>(() => RefService.GetAsync(TestNamespace, bucket, key, Array.Empty<string>()));
             }
 
             // upload missing pieces
             {
-                await Service.PutObjectAsync(TestNamespace, blobData, blobHash);
-                await Service.PutObjectAsync(TestNamespace, dataChild, blobHashChild);
+                await BlobService.PutObjectAsync(TestNamespace, blobData, blobHash);
+                await BlobService.PutObjectAsync(TestNamespace, dataChild, blobHashChild);
             }
 
             // finalize the object as no pieces is now missing
@@ -1048,7 +1036,7 @@ namespace Horde.Server.Tests.Ddc.FunctionalTests.References
 
             // check that actual internal representation has updated its state
             {
-                RefRecord objectRecord = await ReferencesStore.GetAsync(TestNamespace, bucket, key, IReferencesStore.FieldFlags.None);
+				(RefRecord objectRecord, _) = await RefService.GetAsync(TestNamespace, bucket, key, Array.Empty<string>());
 
                 Assert.IsTrue(objectRecord.IsFinalized);
                 Assert.AreEqual(key, objectRecord.Name);
@@ -1056,7 +1044,7 @@ namespace Horde.Server.Tests.Ddc.FunctionalTests.References
             }
         }
 
-		[Ignore]
+		[Ignore("Horde does not support deleting objects manually (ie. BlobService.DeleteObjectAsync)")]
 		[TestMethod]
         public async Task GetMissingContentIdRecordAsync()
         {
@@ -1138,7 +1126,7 @@ namespace Horde.Server.Tests.Ddc.FunctionalTests.References
 
             {
                 // delete the blob referenced by the compressed buffer
-                await Service.DeleteObjectAsync(TestNamespace, compressedHash);
+                await BlobService.DeleteObjectAsync(TestNamespace, compressedHash);
             }
 
             {
@@ -1167,7 +1155,7 @@ namespace Horde.Server.Tests.Ddc.FunctionalTests.References
             }
         }
 
-		[Ignore]
+		[Ignore("Horde does not support deleting objects manually (ie. BlobService.DeleteObjectAsync)")]
 		[TestMethod]
         public async Task GetMissingCompressedBufferAttachmentAsync()
         {
@@ -1243,7 +1231,7 @@ namespace Horde.Server.Tests.Ddc.FunctionalTests.References
 
             {
                 // delete the blob referenced by the compressed buffer
-                await Service.DeleteObjectAsync(TestNamespace, cbAttachmentHash);
+                await BlobService.DeleteObjectAsync(TestNamespace, cbAttachmentHash);
             }
 
             {
@@ -1272,7 +1260,7 @@ namespace Horde.Server.Tests.Ddc.FunctionalTests.References
             }
         }
 
-		[Ignore]
+		[Ignore("Horde does not support deleting objects manually (ie. BlobService.DeleteObjectAsync)")]
 		[TestMethod]
         public async Task GetMissingBlobRecordAsync()
         {
@@ -1332,7 +1320,7 @@ namespace Horde.Server.Tests.Ddc.FunctionalTests.References
 
             {
                 // delete the blob 
-                await Service.DeleteObjectAsync(TestNamespace, blobHash);
+                await BlobService.DeleteObjectAsync(TestNamespace, blobHash);
             }
 
             {
@@ -1421,7 +1409,7 @@ namespace Horde.Server.Tests.Ddc.FunctionalTests.References
             }
         }
 
-		[Ignore]
+		[Ignore("Horde does not support dropping buckets")]
 		[TestMethod]
         public async Task DropBucketAsync()
         {
@@ -1479,7 +1467,7 @@ namespace Horde.Server.Tests.Ddc.FunctionalTests.References
             }
         }
 
-		[Ignore]
+		[Ignore("Horde does not support deleting namespaces")]
 		[TestMethod]
         public async Task DeleteNamespaceAsync()
         {
@@ -1546,7 +1534,7 @@ namespace Horde.Server.Tests.Ddc.FunctionalTests.References
             }
         }
 
-		[Ignore]
+		[Ignore("Horde does not support enumerating namespaces")]
 		[TestMethod]
         public async Task ListNamespacesAsync()
         {
@@ -1575,7 +1563,7 @@ namespace Horde.Server.Tests.Ddc.FunctionalTests.References
             }
         }
 
-		[Ignore]
+		[Ignore("Horde does not support enumerating all blobs")]
 		[TestMethod]
         public async Task GetOldRecordsAsync()
         {
diff --git a/Engine/Source/Programs/Horde/Horde.Server.Tests/Horde.Server.Tests.csproj b/Engine/Source/Programs/Horde/Horde.Server.Tests/Horde.Server.Tests.csproj
--- a/Engine/Source/Programs/Horde/Horde.Server.Tests/Horde.Server.Tests.csproj
+++ b/Engine/Source/Programs/Horde/Horde.Server.Tests/Horde.Server.Tests.csproj
@@ -58,6 +58,9 @@
   </ItemGroup>
 
   <ItemGroup>
+    <None Update="Ddc\Functional\Objects\Payloads\lyra.cb">
+      <CopyToOutputDirectory>Always</CopyToOutputDirectory>
+    </None>
     <None Update="Ddc\UnitTests\CompactBinaryObjects\build">
       <CopyToOutputDirectory>Always</CopyToOutputDirectory>
     </None>
diff --git a/Engine/Source/Programs/Horde/Horde.Server/Ddc/RefService.cs b/Engine/Source/Programs/Horde/Horde.Server/Ddc/RefService.cs
--- a/Engine/Source/Programs/Horde/Horde.Server/Ddc/RefService.cs
+++ b/Engine/Source/Programs/Horde/Horde.Server/Ddc/RefService.cs
@@ -40,13 +40,12 @@ namespace Horde.Server.Ddc
 
 		public async Task<(ContentId[], BlobId[])> FinalizeAsync(NamespaceId ns, BucketId bucket, RefId key, BlobId blobHash, CancellationToken cancellationToken)
 		{
-			HashSet<BlobId> missingBlobIds = new HashSet<BlobId>();
+			List<BlobId> missingBlobIds = new List<BlobId>();
 
 			// Read all the blobs back in
 			List<BlobId> blobIds = new List<BlobId>();
 			blobIds.Add(blobHash);
 
-			HashSet<BlobId> addedBlobIds = new HashSet<BlobId>();
 			for (int idx = 0; idx < blobIds.Count; idx++)
 			{
 				try
@@ -55,14 +54,7 @@ namespace Horde.Server.Ddc
 					byte[] data = await contents.Stream.ReadAllBytesAsync(cancellationToken);
 
 					CbObject obj = new CbObject(data);
-					obj.IterateAttachments(x =>
-					{
-						BlobId blobId = new BlobId(x.AsAttachment());
-						if (addedBlobIds.Add(blobId))
-						{
-							blobIds.Add(blobId);
-						}
-					});
+					obj.IterateAttachments(x => blobIds.Add(new BlobId(x.AsAttachment())));
 				}
 				catch (BlobNotFoundException)
 				{
