Horde: Merge support for compressed buffers from Jupiter.

diff --git a/Engine/Source/Programs/Horde/Horde.Server.Tests/Ddc/Functional/ContentId/CompressedBlobTests.cs b/Engine/Source/Programs/Horde/Horde.Server.Tests/Ddc/Functional/ContentId/CompressedBlobTests.cs
new file mode 100644
--- /dev/null
+++ b/Engine/Source/Programs/Horde/Horde.Server.Tests/Ddc/Functional/ContentId/CompressedBlobTests.cs
@@ -0,0 +1,388 @@
+// Copyright Epic Games, Inc. All Rights Reserved.
+
+using System;
+using System.Collections.Generic;
+using System.IO;
+using System.Net;
+using System.Net.Http;
+using System.Net.Http.Headers;
+using System.Net.Http.Json;
+using System.Net.Mime;
+using System.Text;
+using System.Threading.Tasks;
+using Blake3;
+using Horde.Server.Ddc;
+using Microsoft.AspNetCore.TestHost;
+using Microsoft.AspNetCore.Hosting;
+using Microsoft.Extensions.Configuration;
+using Microsoft.Extensions.DependencyInjection;
+using Microsoft.VisualStudio.TestTools.UnitTesting;
+using EpicGames.AspNet;
+using EpicGames.Core;
+using Horde.Server.Tests;
+using Horde.Server;
+using Horde.Server.Server;
+using Microsoft.Extensions.Logging;
+using Microsoft.Extensions.Options;
+using Horde.Server.Storage;
+using Horde.Server.Configuration;
+using EpicGames.Horde.Storage;
+using System.Threading;
+using System.Text.Json;
+
+namespace Horde.Server.Tests.Ddc.FunctionalTests.CompressedBlobs
+{
+	[TestClass]
+	public class CompressedBlobTests : ControllerIntegrationTest
+	{
+		protected const string TestNamespace = "test-namespace";
+
+		class InsertResponse
+		{
+			public BlobId? Identifier { get; set; }
+		}
+
+		public CompressedBlobTests()
+		{
+			ServerSettings serverSettings = ServiceProvider.GetRequiredService<IOptions<ServerSettings>>().Value;
+
+			GlobalConfig globalConfig = ServiceProvider.GetRequiredService<IOptions<GlobalConfig>>().Value;
+
+			globalConfig.Storage.Backends.Add(new BackendConfig { Id = new BackendId("default"), Type = StorageBackendType.Memory });
+			globalConfig.Storage.Namespaces.Add(new NamespaceConfig { Id = new NamespaceId(TestNamespace), Backend = globalConfig.Storage.Backends[^1].Id });
+			globalConfig.PostLoad(serverSettings);
+
+			ConfigService configService = ServiceProvider.GetRequiredService<ConfigService>();
+			configService.Set(IoHash.Zero, globalConfig);
+		}
+
+		[TestMethod]
+		[DataTestMethod]
+		[DataRow("7835c353d7dc67e8a0531c88fbc75ddfda10dee4", "7835c353d7dc67e8a0531c88fbc75ddfda10dee4")]
+		/*        [DataRow("4958689fe783e02fb35b13c14b0c3d7beb91e50c", "4958689fe783e02fb35b13c14b0c3d7beb91e50c")]
+				[DataRow("dce31eb416f3dcb4c8250ac545eda3930919d3ff", "dce31eb416f3dcb4c8250ac545eda3930919d3ff")]
+				[DataRow("05d7c699a2668efdecbe48f10db0d621d736f449.uecomp", "05D7C699A2668EFDECBE48F10DB0D621D736F449")]
+				[DataRow("dce31eb416f3dcb4c8250ac545eda3930919d3ff", "dce31eb416f3dcb4c8250ac545eda3930919d3ff")]
+				[DataRow("UncompressedTexture_CAS_dea81b6c3b565bb5089695377c98ce0f1c13b0c3.udd", "DEA81B6C3B565BB5089695377C98CE0F1C13B0C3")]
+				[DataRow("Oodle-f895ea954b37217270e88d8b728bd3c09152689c", "F895EA954B37217270E88D8B728BD3C09152689C")]
+				[DataRow("Oodle-f0b9c675fe21951ca27699f9baab9f9f5040b202", "F0B9C675FE21951CA27699F9BAAB9F9F5040B202")]
+				[DataRow("OodleTexture_CAS_dbda9040e75c4674fcec173f982fddf12b021e24.udd", "DBDA9040E75C4674FCEC173F982FDDF12B021E24")]
+		*/
+		public async Task PutPayloadsAsync(string payloadFilename, string uncompressedHash)
+		{
+			byte[] texturePayload = await File.ReadAllBytesAsync($"Ddc/Functional/ContentId/Payloads/{payloadFilename}");
+			BlobId compressedPayloadIdentifier = BlobId.FromBlob(texturePayload);
+			BlobId uncompressedPayloadIdentifier = BlobId.Parse(uncompressedHash);
+
+			using ByteArrayContent content = new(texturePayload);
+			content.Headers.ContentType = new MediaTypeHeaderValue(CustomMediaTypeNames.UnrealCompressedBuffer);
+			HttpResponseMessage result = await Client!.PutAsync(new Uri($"api/v1/compressed-blobs/{TestNamespace}/{uncompressedPayloadIdentifier}", UriKind.Relative), content);
+			result.EnsureSuccessStatusCode();
+
+			InsertResponse? response = await result.Content.ReadFromJsonAsync<InsertResponse>();
+			Assert.IsNotNull(response);
+			Assert.AreNotEqual(compressedPayloadIdentifier, response.Identifier);
+			Assert.AreEqual(uncompressedPayloadIdentifier, response.Identifier);
+		}
+
+		[TestMethod]
+		public async Task PutGetComplexTextureAsync()
+		{
+			byte[] texturePayload = await File.ReadAllBytesAsync("Ddc/Functional/ContentId/Payloads/UncompressedTexture_CAS_dea81b6c3b565bb5089695377c98ce0f1c13b0c3.udd");
+			BlobId compressedPayloadIdentifier = BlobId.FromBlob(texturePayload);
+			ContentId uncompressedPayloadIdentifier = ContentId.Parse("DEA81B6C3B565BB5089695377C98CE0F1C13B0C3");
+
+			{
+				using ByteArrayContent content = new(texturePayload);
+				content.Headers.ContentType = new MediaTypeHeaderValue(CustomMediaTypeNames.UnrealCompressedBuffer);
+				HttpResponseMessage result = await Client!.PutAsync(new Uri($"api/v1/compressed-blobs/{TestNamespace}/{uncompressedPayloadIdentifier}", UriKind.Relative), content);
+				result.EnsureSuccessStatusCode();
+
+				InsertResponse? response = await result.Content.ReadFromJsonAsync<InsertResponse>();
+				Assert.IsNotNull(response);
+				Assert.IsNotNull(response.Identifier);
+				Assert.AreNotEqual(compressedPayloadIdentifier, response.Identifier);
+				Assert.AreEqual(uncompressedPayloadIdentifier, ContentId.FromBlobId(response.Identifier.Value));
+			}
+
+			{
+				HttpResponseMessage result = await Client!.GetAsync(new Uri($"api/v1/compressed-blobs/{TestNamespace}/{uncompressedPayloadIdentifier}", UriKind.Relative));
+				result.EnsureSuccessStatusCode();
+				Assert.AreEqual(CustomMediaTypeNames.UnrealCompressedBuffer, result.Content.Headers.ContentType!.MediaType);
+
+				byte[] blobContent = await result.Content.ReadAsByteArrayAsync();
+				CollectionAssert.AreEqual(texturePayload, blobContent);
+			}
+
+			{
+				// verify the compressed blob can be retrieved in the blob store
+				HttpResponseMessage result = await Client!.GetAsync(new Uri($"api/v1/blobs/{TestNamespace}/{compressedPayloadIdentifier}", UriKind.Relative));
+				result.EnsureSuccessStatusCode();
+				Assert.AreEqual(MediaTypeNames.Application.Octet, result.Content.Headers.ContentType!.MediaType);
+
+				byte[] blobContent = await result.Content.ReadAsByteArrayAsync();
+				CollectionAssert.AreEqual(texturePayload, blobContent);
+			}
+
+			{
+				// the uncompressed payload should not be valid in the blob endpoint
+				HttpResponseMessage result = await Client!.GetAsync(new Uri($"api/v1/blobs/{TestNamespace}/{uncompressedPayloadIdentifier}", UriKind.Relative));
+				Assert.AreEqual(HttpStatusCode.NotFound, result.StatusCode);
+			}
+		}
+
+		/// <summary>
+		/// write a large file (bigger then what C# can have in memory)
+		/// </summary>
+		[Ignore]
+		[TestMethod]
+		[TestCategory("SlowTests")]
+		public async Task PutGetLargeCompressedPayloadAsync()
+		{
+			// we submit a blob so large that it can not fit using the memory blob store
+			FileInfo tempOutputFile = new FileInfo(Path.GetTempFileName());
+			FileInfo tempCompressedFile = new FileInfo(Path.GetTempFileName());
+
+			ILogger logger = ServiceProvider.GetRequiredService<ILogger<CompressedBlobTests>>()!;
+
+			try
+			{
+				logger.LogInformation("Generating large file");
+				int blockSize = 1024 * 1024;
+				// we want a file larger then 6GB, each block is 1 MB
+				int countOfBlocks = 6500;
+				IoHash uncompressedContentHash;
+				{
+					byte[] block = new byte[blockSize];
+					Random.Shared.NextBytes(block);
+
+					List<byte[]> blocksToCompress = new List<byte[]>();
+
+					using Hasher hasher = Hasher.New();
+					for (int i = 0; i < countOfBlocks; i++)
+					{
+						hasher.UpdateWithJoin(new ReadOnlySpan<byte>(block, 0, blockSize));
+						blocksToCompress.Add(block);
+					}
+					Hash blake3Hash = hasher.Finalize();
+					byte[] hash = blake3Hash.AsSpanUnsafe().Slice(0, 20).ToArray();
+					uncompressedContentHash = new IoHash(hash);
+
+					await using FileStream fs = tempCompressedFile.OpenWrite();
+					CompressedBufferUtils bufferUtils = ServiceProvider.GetService<CompressedBufferUtils>()!;
+					bufferUtils.CompressContent(fs, OoodleCompressorMethod.Mermaid, OoodleCompressionLevel.HyperFast4, blocksToCompress, blockSize);
+				}
+
+				logger.LogInformation("Hashing generated file");
+				BlobId blobIdentifier = BlobId.FromIoHash(uncompressedContentHash);
+				BlobId compressedContentHash;
+				{
+					await using FileStream fs = tempCompressedFile.OpenRead();
+					compressedContentHash = await BlobId.FromStreamAsync(fs, CancellationToken.None);
+				}
+
+				logger.LogInformation("Uploading large file");
+
+				// it takes a long time to upload this content and we will not get any response while it happens so we have to bump the timeout
+				Client!.Timeout = TimeSpan.FromMinutes(5.0);
+				{
+					await using FileStream fs = tempCompressedFile.OpenRead();
+					using StreamContent content = new StreamContent(fs);
+					content.Headers.ContentType = new MediaTypeHeaderValue(CustomMediaTypeNames.UnrealCompressedBuffer);
+					HttpResponseMessage result = await Client!.PutAsync(new Uri($"api/v1/compressed-blobs/{TestNamespace}/{blobIdentifier}", UriKind.Relative), content);
+					result.EnsureSuccessStatusCode();
+
+					InsertResponse? response = await result.Content.ReadFromJsonAsync<InsertResponse>();
+					Assert.IsNotNull(response);
+					Assert.AreEqual(blobIdentifier, response.Identifier);
+				}
+
+				logger.LogInformation("Large file uploaded");
+				logger.LogInformation("Downloading large file");
+
+				{
+					// verify we can fetch the blob again
+					HttpResponseMessage result = await Client!.GetAsync(new Uri($"api/v1/compressed-blobs/{TestNamespace}/{blobIdentifier}", UriKind.Relative), HttpCompletionOption.ResponseHeadersRead);
+					result.EnsureSuccessStatusCode();
+
+					{
+						await using Stream s = await result.Content.ReadAsStreamAsync();
+
+						// stream this to disk so we have something to look at in case there is an error
+						await using FileStream fs = tempOutputFile.OpenWrite();
+						await s.CopyToAsync(fs);
+					}
+
+					await using FileStream downloadedFile = tempOutputFile.OpenRead();
+
+					BlobId downloadedBlobIdentifier = await BlobId.FromStreamAsync(downloadedFile, CancellationToken.None);
+					Assert.AreEqual(compressedContentHash, downloadedBlobIdentifier);
+				}
+
+				logger.LogInformation("Download completed");
+			}
+			finally
+			{
+				if (tempCompressedFile.Exists)
+				{
+					tempCompressedFile.Delete();
+				}
+
+				if (tempOutputFile.Exists)
+				{
+					tempOutputFile.Delete();
+				}
+			}
+		}
+
+		[Ignore]
+		[TestMethod]
+		public async Task RecompressionTestAsync()
+		{
+			ContentId uncompressedPayloadIdentifier = ContentId.Parse("A2AC0ECED768698F7413F131D064D36B7EC6F7DA");
+			byte[] texturePayloadSmaller = await File.ReadAllBytesAsync("Ddc/Functional/ContentId/Payloads/smallerfile");
+			BlobId compressedPayloadIdentifierSmaller = BlobId.FromBlob(texturePayloadSmaller);
+
+			{
+				using ByteArrayContent content = new(texturePayloadSmaller);
+				content.Headers.ContentType = new MediaTypeHeaderValue(CustomMediaTypeNames.UnrealCompressedBuffer);
+				HttpResponseMessage result = await Client!.PutAsync(new Uri($"api/v1/compressed-blobs/{TestNamespace}/{uncompressedPayloadIdentifier}", UriKind.Relative), content);
+				result.EnsureSuccessStatusCode();
+
+				InsertResponse response = JsonSerializer.Deserialize<InsertResponse>(await result.Content.ReadAsByteArrayAsync())!;
+				Assert.IsNotNull(response.Identifier);
+				Assert.AreNotEqual(compressedPayloadIdentifierSmaller, response.Identifier);
+				Assert.AreEqual(uncompressedPayloadIdentifier, ContentId.FromBlobId(response.Identifier.Value));
+			}
+
+			byte[] texturePayloadLarger = await File.ReadAllBytesAsync("Ddc/Functional/ContentId/Payloads/largerfile");
+			BlobId compressedPayloadIdentifierLarger = BlobId.FromBlob(texturePayloadLarger);
+
+			{
+				using ByteArrayContent content = new(texturePayloadLarger);
+				content.Headers.ContentType = new MediaTypeHeaderValue(CustomMediaTypeNames.UnrealCompressedBuffer);
+				HttpResponseMessage result = await Client!.PutAsync(new Uri($"api/v1/compressed-blobs/{TestNamespace}/{uncompressedPayloadIdentifier}", UriKind.Relative), content);
+				result.EnsureSuccessStatusCode();
+
+				InsertResponse response = JsonSerializer.Deserialize<InsertResponse>(await result.Content.ReadAsByteArrayAsync())!;
+				Assert.IsNotNull(response.Identifier);
+				Assert.AreNotEqual(compressedPayloadIdentifierLarger, response.Identifier);
+				Assert.AreEqual(uncompressedPayloadIdentifier, ContentId.FromBlobId(response.Identifier.Value));
+			}
+
+			{
+				// fetch the overloaded content id, it should be returning the smaller payload
+				HttpResponseMessage result = await Client!.GetAsync(new Uri($"api/v1/compressed-blobs/{TestNamespace}/{uncompressedPayloadIdentifier}", UriKind.Relative));
+				result.EnsureSuccessStatusCode();
+				Assert.AreEqual(CustomMediaTypeNames.UnrealCompressedBuffer, result.Content.Headers.ContentType!.MediaType);
+
+				byte[] blobContent = await result.Content.ReadAsByteArrayAsync();
+				Assert.AreEqual(compressedPayloadIdentifierSmaller, BlobId.FromBlob(blobContent));
+				CollectionAssert.AreEqual(texturePayloadSmaller, blobContent);
+			}
+		}
+
+		[TestMethod]
+		public async Task PutWrongIdentifierAsync()
+		{
+			byte[] texturePayload = await File.ReadAllBytesAsync("Ddc/Functional/ContentId/Payloads/UncompressedTexture_CAS_dea81b6c3b565bb5089695377c98ce0f1c13b0c3.udd");
+			BlobId compressedPayloadIdentifier = BlobId.FromBlob(texturePayload);
+
+			{
+				using ByteArrayContent content = new(texturePayload);
+				content.Headers.ContentType = new MediaTypeHeaderValue(CustomMediaTypeNames.UnrealCompressedBuffer);
+				// we purposefully use the compressed identifier here which is not what is expected
+				HttpResponseMessage result = await Client!.PutAsync(new Uri($"api/v1/compressed-blobs/{TestNamespace}/{compressedPayloadIdentifier}", UriKind.Relative), content);
+
+				Assert.AreEqual(HttpStatusCode.BadRequest, result.StatusCode);
+			}
+		}
+
+		[TestMethod]
+		public async Task PostNoIdentifierAsync()
+		{
+			byte[] texturePayload = await File.ReadAllBytesAsync("Ddc/Functional/ContentId/Payloads/UncompressedTexture_CAS_dea81b6c3b565bb5089695377c98ce0f1c13b0c3.udd");
+			BlobId compressedPayloadIdentifier = BlobId.FromBlob(texturePayload);
+			BlobId uncompressedPayloadIdentifier = BlobId.Parse("DEA81B6C3B565BB5089695377C98CE0F1C13B0C3");
+			{
+				using ByteArrayContent content = new(texturePayload);
+				content.Headers.ContentType = new MediaTypeHeaderValue(CustomMediaTypeNames.UnrealCompressedBuffer);
+				// we purposefully use the compressed identifier here which is not what is expected
+				HttpResponseMessage result = await Client!.PostAsync(new Uri($"api/v1/compressed-blobs/{TestNamespace}", UriKind.Relative), content);
+
+				result.EnsureSuccessStatusCode();
+
+				InsertResponse? response = await result.Content.ReadFromJsonAsync<InsertResponse>();
+				Assert.IsNotNull(response);
+				Assert.IsNotNull(response.Identifier);
+				Assert.AreEqual(uncompressedPayloadIdentifier, response.Identifier);
+				Assert.AreNotEqual(compressedPayloadIdentifier, response.Identifier);
+			}
+		}
+
+		[TestMethod]
+		public async Task GetUncompressedContentAsync()
+		{
+			string stringContent = "this is just a random string";
+			byte[] payload = Encoding.ASCII.GetBytes(stringContent);
+			BlobId blobIdentifier = BlobId.FromBlob(payload);
+
+			// upload a uncompressed blob
+			{
+				using ByteArrayContent content = new(payload);
+				content.Headers.ContentType = new MediaTypeHeaderValue(MediaTypeNames.Application.Octet);
+				HttpResponseMessage result = await Client!.PutAsync(new Uri($"api/v1/blobs/{TestNamespace}/{blobIdentifier}", UriKind.Relative), content);
+				result.EnsureSuccessStatusCode();
+			}
+
+			{
+				HttpResponseMessage result = await Client!.GetAsync(new Uri($"api/v1/compressed-blobs/{TestNamespace}/{blobIdentifier}", UriKind.Relative));
+				result.EnsureSuccessStatusCode();
+
+				// verify that we return the uncompressed content but also that the media type indicates that
+				byte[] blobContent = await result.Content.ReadAsByteArrayAsync();
+				Assert.AreEqual(MediaTypeNames.Application.Octet, result.Content.Headers.ContentType!.MediaType);
+				CollectionAssert.AreEqual(payload, blobContent);
+			}
+		}
+
+
+		[TestMethod]
+		public async Task GetUncompressedContentAsCompressedBufferAsync()
+		{
+			string stringContent = "this is just a random string";
+			byte[] payload = Encoding.ASCII.GetBytes(stringContent);
+			BlobId blobIdentifier = BlobId.FromBlob(payload);
+
+			// upload a uncompressed blob
+			{
+				using ByteArrayContent content = new(payload);
+				content.Headers.ContentType = new MediaTypeHeaderValue(MediaTypeNames.Application.Octet);
+				HttpResponseMessage result = await Client!.PutAsync(new Uri($"api/v1/blobs/{TestNamespace}/{blobIdentifier}", UriKind.Relative), content);
+				result.EnsureSuccessStatusCode();
+			}
+
+			{
+				using HttpRequestMessage request = new HttpRequestMessage(HttpMethod.Get, $"api/v1/compressed-blobs/{TestNamespace}/{blobIdentifier}");
+				request.Headers.Add("Accept", CustomMediaTypeNames.UnrealCompactBinary);
+
+				// asking for a compressed buffer when we only have the uncompressed content results in a 415 (unsupported media type) as transcoding this is to complicated when the compressed buffer header requires the full blake3 hash of the content
+				HttpResponseMessage result = await Client!.SendAsync(request);
+				Assert.AreEqual(HttpStatusCode.UnsupportedMediaType, result.StatusCode);
+			}
+
+			{
+				// ask for any accept type
+				using HttpRequestMessage request = new HttpRequestMessage(HttpMethod.Get, $"api/v1/compressed-blobs/{TestNamespace}/{blobIdentifier}");
+				request.Headers.Add("Accept", "*/*");
+
+				HttpResponseMessage result = await Client!.SendAsync(request);
+				result.EnsureSuccessStatusCode();
+				Assert.AreEqual(MediaTypeNames.Application.Octet, result.Content.Headers.ContentType!.MediaType);
+
+				byte[] blobContent = await result.Content.ReadAsByteArrayAsync();
+				CollectionAssert.AreEqual(payload, blobContent);
+			}
+		}
+	}
+}
diff --git a/Engine/Source/Programs/Horde/Horde.Server.Tests/Horde.Server.Tests.csproj b/Engine/Source/Programs/Horde/Horde.Server.Tests/Horde.Server.Tests.csproj
--- a/Engine/Source/Programs/Horde/Horde.Server.Tests/Horde.Server.Tests.csproj
+++ b/Engine/Source/Programs/Horde/Horde.Server.Tests/Horde.Server.Tests.csproj
@@ -58,6 +58,39 @@
   </ItemGroup>
 
   <ItemGroup>
+    <None Update="Ddc\Functional\ContentId\Payloads\05d7c699a2668efdecbe48f10db0d621d736f449.uecomp">
+      <CopyToOutputDirectory>Always</CopyToOutputDirectory>
+    </None>
+    <None Update="Ddc\Functional\ContentId\Payloads\0b29a878f71cc1eb6a26663a1152537e8fe7dca9.uecomp">
+      <CopyToOutputDirectory>Always</CopyToOutputDirectory>
+    </None>
+    <None Update="Ddc\Functional\ContentId\Payloads\4958689fe783e02fb35b13c14b0c3d7beb91e50c">
+      <CopyToOutputDirectory>Always</CopyToOutputDirectory>
+    </None>
+    <None Update="Ddc\Functional\ContentId\Payloads\7835c353d7dc67e8a0531c88fbc75ddfda10dee4">
+      <CopyToOutputDirectory>Always</CopyToOutputDirectory>
+    </None>
+    <None Update="Ddc\Functional\ContentId\Payloads\dce31eb416f3dcb4c8250ac545eda3930919d3ff">
+      <CopyToOutputDirectory>Always</CopyToOutputDirectory>
+    </None>
+    <None Update="Ddc\Functional\ContentId\Payloads\largerfile">
+      <CopyToOutputDirectory>Always</CopyToOutputDirectory>
+    </None>
+    <None Update="Ddc\Functional\ContentId\Payloads\Oodle-f0b9c675fe21951ca27699f9baab9f9f5040b202">
+      <CopyToOutputDirectory>Always</CopyToOutputDirectory>
+    </None>
+    <None Update="Ddc\Functional\ContentId\Payloads\Oodle-f895ea954b37217270e88d8b728bd3c09152689c">
+      <CopyToOutputDirectory>Always</CopyToOutputDirectory>
+    </None>
+    <None Update="Ddc\Functional\ContentId\Payloads\OodleTexture_CAS_dbda9040e75c4674fcec173f982fddf12b021e24.udd">
+      <CopyToOutputDirectory>Always</CopyToOutputDirectory>
+    </None>
+    <None Update="Ddc\Functional\ContentId\Payloads\smallerfile">
+      <CopyToOutputDirectory>Always</CopyToOutputDirectory>
+    </None>
+    <None Update="Ddc\Functional\ContentId\Payloads\UncompressedTexture_CAS_dea81b6c3b565bb5089695377c98ce0f1c13b0c3.udd">
+      <CopyToOutputDirectory>Always</CopyToOutputDirectory>
+    </None>
     <None Update="Ddc\Functional\Objects\Payloads\lyra.cb">
       <CopyToOutputDirectory>Always</CopyToOutputDirectory>
     </None>
diff --git a/Engine/Source/Programs/Horde/Horde.Server/Ddc/CompressedBlobController.cs b/Engine/Source/Programs/Horde/Horde.Server/Ddc/CompressedBlobController.cs
new file mode 100644
--- /dev/null
+++ b/Engine/Source/Programs/Horde/Horde.Server/Ddc/CompressedBlobController.cs
@@ -0,0 +1,335 @@
+﻿// Copyright Epic Games, Inc. All Rights Reserved.
+
+using System.Collections.Concurrent;
+using System.Collections.Generic;
+using System.ComponentModel.DataAnnotations;
+using System.Linq;
+using System.Net;
+using System.Net.Mime;
+using System.Threading;
+using System.Threading.Tasks;
+using EpicGames.AspNet;
+using EpicGames.Horde.Storage;
+using EpicGames.Serialization;
+using Horde.Server.Storage;
+using Microsoft.AspNetCore.Authorization;
+using Microsoft.AspNetCore.Mvc;
+using Microsoft.Extensions.Primitives;
+
+#pragma warning disable CS1591
+
+namespace Horde.Server.Ddc
+{
+	using IDiagnosticContext = Serilog.IDiagnosticContext;
+
+	[ApiController]
+	[Authorize]
+	[Route("api/v1/compressed-blobs")]
+	public class CompressedBlobController : ControllerBase
+	{
+		private readonly IBlobService _storage;
+		private readonly IContentIdService _contentIdStore;
+		private readonly IDiagnosticContext _diagnosticContext;
+		private readonly IRequestHelper _requestHelper;
+		private readonly BufferedPayloadFactory _bufferedPayloadFactory;
+		private readonly NginxRedirectHelper _nginxRedirectHelper;
+
+		public CompressedBlobController(IBlobService storage, IContentIdService contentIdStore, IDiagnosticContext diagnosticContext, IRequestHelper requestHelper, BufferedPayloadFactory bufferedPayloadFactory, NginxRedirectHelper nginxRedirectHelper)
+		{
+			_storage = storage;
+			_contentIdStore = contentIdStore;
+			_diagnosticContext = diagnosticContext;
+			_requestHelper = requestHelper;
+			_bufferedPayloadFactory = bufferedPayloadFactory;
+			_nginxRedirectHelper = nginxRedirectHelper;
+		}
+
+		[HttpGet("{ns}/{id}")]
+		[ProducesResponseType(type: typeof(byte[]), 200)]
+		[ProducesResponseType(type: typeof(ValidationProblemDetails), 400)]
+		[Produces(CustomMediaTypeNames.UnrealCompressedBuffer, MediaTypeNames.Application.Octet)]
+
+		public async Task<IActionResult> GetAsync(
+			[Required] NamespaceId ns,
+			[Required] ContentId id)
+		{
+			ActionResult? result = await _requestHelper.HasAccessToNamespaceAsync(User, Request, ns, new [] { StorageAclAction.ReadBlobs });
+			if (result != null)
+			{
+				return result;
+			}
+
+			try
+			{
+				(BlobContents blobContents, string mediaType) = await _storage.GetCompressedObjectAsync(ns, id, HttpContext.RequestServices, supportsRedirectUri: true);
+
+				StringValues acceptHeader = Request.Headers["Accept"];
+				if (!acceptHeader.Contains("*/*") && acceptHeader.Count != 0 && !acceptHeader.Contains(mediaType))
+				{
+					return new UnsupportedMediaTypeResult();
+				}
+
+				if (blobContents.RedirectUri != null)
+				{
+					return Redirect(blobContents.RedirectUri.ToString());
+				}
+
+				if (_nginxRedirectHelper.CanRedirect(Request, blobContents))
+				{
+					return _nginxRedirectHelper.CreateActionResult(blobContents, mediaType);
+				}
+
+				return File(blobContents.Stream, mediaType, enableRangeProcessing: true);
+			}
+			catch (BlobNotFoundException e)
+			{
+				return NotFound(new ValidationProblemDetails { Title = $"Object {e.Blob} not found" });
+			}
+			catch (ContentIdResolveException e)
+			{
+				return NotFound(new ValidationProblemDetails { Title = $"Content Id {e.ContentId} not found" });
+			}
+		}
+		
+		[HttpHead("{ns}/{id}")]
+		[ProducesDefaultResponseType]
+		public async Task<IActionResult> HeadAsync(
+			[Required] NamespaceId ns,
+			[Required] ContentId id)
+		{
+			ActionResult? result = await _requestHelper.HasAccessToNamespaceAsync(User, Request, ns, new [] { StorageAclAction.ReadBlobs });
+			if (result != null)
+			{
+				return result;
+			}
+
+			BlobId[]? chunks = await _contentIdStore.ResolveAsync(ns, id, mustBeContentId: false);
+			if (chunks == null || chunks.Length == 0)
+			{
+				return NotFound();
+			}
+
+			Task<bool>[] tasks = new Task<bool>[chunks.Length];
+			for (int i = 0; i < chunks.Length; i++)
+			{
+				tasks[i] = _storage.ExistsAsync(ns, chunks[i]);
+			}
+
+			await Task.WhenAll(tasks);
+
+			bool exists = tasks.All(task => task.Result);
+
+			if (!exists)
+			{
+				return NotFound();
+			}
+
+			return Ok();
+		}
+
+		[HttpPost("{ns}/exists")]
+		[ProducesDefaultResponseType]
+		public async Task<IActionResult> ExistsMultipleAsync(
+			[Required] NamespaceId ns,
+			[Required] [FromQuery] List<ContentId> id)
+		{
+			ActionResult? result = await _requestHelper.HasAccessToNamespaceAsync(User, Request, ns, new [] { StorageAclAction.ReadBlobs });
+			if (result != null)
+			{
+				return result;
+			}
+
+			ConcurrentBag<ContentId> partialContentIds = new ConcurrentBag<ContentId>();
+			ConcurrentBag<ContentId> invalidContentIds = new ConcurrentBag<ContentId>();
+
+			IEnumerable<Task> tasks = id.Select(async blob =>
+			{
+				BlobId[]? chunks = await _contentIdStore.ResolveAsync(ns, blob, mustBeContentId: false);
+
+				if (chunks == null)
+				{
+					invalidContentIds.Add(blob);
+					return;
+				}
+
+				foreach (BlobId chunk in chunks)
+				{
+					if (!await _storage.ExistsAsync(ns, chunk))
+					{
+						partialContentIds.Add(blob);
+						break;
+					}
+				}
+			});
+			await Task.WhenAll(tasks);
+
+			List<ContentId> needs = new List<ContentId>(invalidContentIds);
+			needs.AddRange(partialContentIds);
+			 
+			return Ok(new ExistCheckMultipleContentIdResponse { Needs = needs.ToArray()});
+		}
+
+		[HttpPost("{ns}/exist")]
+		[ProducesDefaultResponseType]
+		public async Task<IActionResult> ExistsBodyAsync(
+			[Required] NamespaceId ns,
+			[FromBody] ContentId[] bodyIds)
+		{
+			ActionResult? result = await _requestHelper.HasAccessToNamespaceAsync(User, Request, ns, new [] { StorageAclAction.ReadBlobs });
+			if (result != null)
+			{
+				return result;
+			}
+
+			ConcurrentBag<ContentId> partialContentIds = new ConcurrentBag<ContentId>();
+			ConcurrentBag<ContentId> invalidContentIds = new ConcurrentBag<ContentId>();
+
+			IEnumerable<Task> tasks = bodyIds.Select(async blob =>
+			{
+				BlobId[]? chunks = await _contentIdStore.ResolveAsync(ns, blob, mustBeContentId: false);
+
+				if (chunks == null)
+				{
+					invalidContentIds.Add(blob);
+					return;
+				}
+
+				foreach (BlobId chunk in chunks)
+				{
+					if (!await _storage.ExistsAsync(ns, chunk))
+					{
+						partialContentIds.Add(blob);
+						break;
+					}
+				}
+			});
+			await Task.WhenAll(tasks);
+
+			List<ContentId> needs = new List<ContentId>(invalidContentIds);
+			needs.AddRange(partialContentIds);
+			 
+			return Ok(new ExistCheckMultipleContentIdResponse { Needs = needs.ToArray()});
+		}
+
+		[HttpPut("{ns}/{id}")]
+		[DisableRequestSizeLimit]
+		[RequiredContentType(CustomMediaTypeNames.UnrealCompressedBuffer)]
+		public async Task<IActionResult> PutAsync(
+			[Required] NamespaceId ns,
+			[Required] ContentId id,
+			CancellationToken cancellationToken)
+		{
+			ActionResult? result = await _requestHelper.HasAccessToNamespaceAsync(User, Request, ns, new [] { StorageAclAction.WriteBlobs });
+			if (result != null)
+			{
+				return result;
+			}
+
+			_diagnosticContext.Set("Content-Length", Request.ContentLength ?? -1);
+
+			try
+			{
+				using BufferedPayload payload = await _bufferedPayloadFactory.CreateFromRequest(Request);
+
+				ContentId identifier = await _storage.PutCompressedObjectAsync(ns, payload, id, HttpContext.RequestServices, cancellationToken);
+
+				return Ok(new { Identifier = identifier.ToString() });
+			}
+			catch (HashMismatchException e)
+			{
+				return BadRequest(new ProblemDetails
+				{
+					Title =
+						$"Incorrect hash, got hash \"{e.SuppliedHash}\" but hash of content was determined to be \"{e.ContentHash}\""
+				});
+			}
+			catch (ClientSendSlowException e)
+			{
+				return Problem(e.Message, null, (int)HttpStatusCode.RequestTimeout);
+			}
+		}
+
+		[HttpPost("{ns}")]
+		[DisableRequestSizeLimit]
+		[RequiredContentType(CustomMediaTypeNames.UnrealCompressedBuffer)]
+		public async Task<IActionResult> PostAsync(
+			[Required] NamespaceId ns,
+			CancellationToken cancellationToken)
+		{
+			ActionResult? result = await _requestHelper.HasAccessToNamespaceAsync(User, Request, ns, new [] { StorageAclAction.WriteBlobs });
+			if (result != null)
+			{
+				return result;
+			}
+
+			_diagnosticContext.Set("Content-Length", Request.ContentLength ?? -1);
+
+			try
+			{
+				using BufferedPayload payload = await _bufferedPayloadFactory.CreateFromRequest(Request);
+
+				ContentId identifier = await _storage.PutCompressedObjectAsync(ns, payload, null, HttpContext.RequestServices, cancellationToken);
+
+				return Ok(new
+				{
+					Identifier = identifier.ToString()
+				});
+			}
+			catch (HashMismatchException e)
+			{
+				return BadRequest(new ProblemDetails
+				{
+					Title = $"Incorrect hash, got hash \"{e.SuppliedHash}\" but hash of content was determined to be \"{e.ContentHash}\""
+				});
+			}
+			catch (ClientSendSlowException e)
+			{
+				return Problem(e.Message, null, (int)HttpStatusCode.RequestTimeout);
+			}
+		}
+
+		/*[HttpDelete("{ns}/{id}")]
+		public async Task<IActionResult> Delete(
+			[Required] string ns,
+			[Required] BlobIdentifier id)
+		{
+			ActionResult? result = await _requestHelper.HasAccessToNamespace(User, Request, ns);
+			if (result != null)
+			{
+				return result;
+			}
+
+			await DeleteImpl(ns, id);
+
+			return Ok();
+		}
+
+		
+		[HttpDelete("{ns}")]
+		public async Task<IActionResult> DeleteNamespace(
+			[Required] string ns)
+		{
+			ActionResult? result = await _requestHelper.HasAccessToNamespace(User, Request, ns);
+			if (result != null)
+			{
+				return result;
+			}
+
+			int deletedCount = await  _storage.DeleteNamespace(ns);
+
+			return Ok( new { Deleted = deletedCount });
+		}
+
+
+		private async Task DeleteImpl(string ns, BlobIdentifier id)
+		{
+			await _storage.DeleteObject(ns, id);
+		}*/
+	}
+
+	public class ExistCheckMultipleContentIdResponse
+	{
+		[CbField("needs")]
+		public ContentId[] Needs { get; set; } = null!;
+	}
+}
diff --git a/Engine/Source/Programs/Horde/Horde.Server/Ddc/ContentIdService.cs b/Engine/Source/Programs/Horde/Horde.Server/Ddc/ContentIdService.cs
--- a/Engine/Source/Programs/Horde/Horde.Server/Ddc/ContentIdService.cs
+++ b/Engine/Source/Programs/Horde/Horde.Server/Ddc/ContentIdService.cs
@@ -1,21 +1,58 @@
 // Copyright Epic Games, Inc. All Rights Reserved.
 
+using EpicGames.Core;
 using EpicGames.Horde.Storage;
-using System;
+using Horde.Server.Storage;
+using System.Linq;
+using System.Threading;
 using System.Threading.Tasks;
 
 namespace Horde.Server.Ddc
 {
 	class ContentIdService : IContentIdService
 	{
-		public Task<BlobId[]?> ResolveAsync(NamespaceId ns, ContentId contentId, bool mustBeContentId = false)
+		readonly StorageService _storageService;
+
+		public ContentIdService(StorageService storageService)
 		{
-			return Task.FromResult<BlobId[]?>(null);
+			_storageService = storageService;
 		}
 
-		public Task PutAsync(NamespaceId ns, ContentId contentId, BlobId blobIdentifier, int contentWeight)
+		static Utf8String GetAlias(BlobId blobId) => BlobService.GetAlias(blobId);
+		static Utf8String GetAlias(ContentId contentId) => $"cid:{contentId}";
+
+		public async Task<BlobId[]?> ResolveAsync(NamespaceId ns, ContentId contentId, bool mustBeContentId = false)
 		{
-			return Task.CompletedTask;
+			CancellationToken cancellationToken = CancellationToken.None;
+
+			IStorageClient storageClient = await _storageService.GetClientAsync(ns, cancellationToken);
+
+			BlobHandle? blobHandle = await storageClient.FindAliasAsync(GetAlias(contentId), cancellationToken).FirstOrDefaultAsync(cancellationToken);
+			if (blobHandle == null && !mustBeContentId)
+			{
+				blobHandle = await storageClient.FindAliasAsync(GetAlias(contentId.AsBlobIdentifier()), cancellationToken).FirstOrDefaultAsync(cancellationToken);
+			}
+			if (blobHandle == null)
+			{
+				return null;
+			}
+
+			return new[] { BlobId.FromIoHash(blobHandle.Hash) };
+		}
+
+		public async Task PutAsync(NamespaceId ns, ContentId contentId, BlobId blobId, int contentWeight)
+		{
+			CancellationToken cancellationToken = CancellationToken.None;
+
+			IStorageClient storageClient = await _storageService.GetClientAsync(ns, cancellationToken);
+
+			BlobHandle? blobHandle = await storageClient.FindAliasAsync(GetAlias(blobId), cancellationToken).FirstOrDefaultAsync(cancellationToken);
+			if (blobHandle == null)
+			{
+				throw new BlobNotFoundException(ns, blobId);
+			}
+
+			await storageClient.AddAliasAsync(GetAlias(contentId), blobHandle, cancellationToken);
 		}
 	}
 }
diff --git a/Engine/Source/Programs/Horde/Horde.Server/Ddc/IBlobService.cs b/Engine/Source/Programs/Horde/Horde.Server/Ddc/IBlobService.cs
--- a/Engine/Source/Programs/Horde/Horde.Server/Ddc/IBlobService.cs
+++ b/Engine/Source/Programs/Horde/Horde.Server/Ddc/IBlobService.cs
@@ -136,14 +136,13 @@ namespace Horde.Server.Ddc
 			// commit the mapping from the decompressed hash to the compressed hash, we run this in parallel with the blob store submit
 			// TODO: let users specify weight of the blob compared to previously submitted content ids
 			int contentIdWeight = (int)payload.Length;
-			Task contentIdStoreTask = contentIdStore.PutAsync(ns, identifierDecompressedPayload, identifierCompressedPayload, contentIdWeight);
 
 			// we still commit the compressed buffer to the object store using the hash of the compressed content
 			{
 				await blobService.PutObjectKnownHashAsync(ns, payload, identifierCompressedPayload, cancellationToken);
 			}
 
-			await contentIdStoreTask;
+			await contentIdStore.PutAsync(ns, identifierDecompressedPayload, identifierCompressedPayload, contentIdWeight);
 
 			return identifierDecompressedPayload;
 		}
@@ -164,7 +163,7 @@ namespace Horde.Server.Ddc
 			{
 				BlobId blobToReturn = chunks[0];
 				string mimeType = CustomMediaTypeNames.UnrealCompressedBuffer;
-				if (contentId.Equals(blobToReturn))
+				if (contentId.Hash == blobToReturn.Hash)
 				{
 					// this was actually the unmapped blob, meaning its not a compressed buffer
 					mimeType = MediaTypeNames.Application.Octet;
diff --git a/Engine/Source/Programs/Horde/Horde.Server/Startup.cs b/Engine/Source/Programs/Horde/Horde.Server/Startup.cs
--- a/Engine/Source/Programs/Horde/Horde.Server/Startup.cs
+++ b/Engine/Source/Programs/Horde/Horde.Server/Startup.cs
@@ -531,6 +531,7 @@ namespace Horde.Server
 			services.AddSingleton<BufferedPayloadFactory>();
 			services.AddSingleton<NginxRedirectHelper>();
 			services.AddSingleton<FormatResolver>();
+			services.AddSingleton<CompressedBufferUtils>();
 
 			// Secret providers
 			services.AddSingleton<SecretService>();
