Horde: Allow specifying a rank for aliases, where higher values indicate the target blob should be preferred over other blobs with the same alias but lower rank.

diff --git a/Engine/Source/Programs/Horde/Horde.Server.Tests/Ddc/Functional/Objects/ReferencesTests.cs b/Engine/Source/Programs/Horde/Horde.Server.Tests/Ddc/Functional/Objects/ReferencesTests.cs
--- a/Engine/Source/Programs/Horde/Horde.Server.Tests/Ddc/Functional/Objects/ReferencesTests.cs
+++ b/Engine/Source/Programs/Horde/Horde.Server.Tests/Ddc/Functional/Objects/ReferencesTests.cs
@@ -823,6 +823,7 @@ namespace Horde.Server.Tests.Ddc.FunctionalTests.References
             }
         }
 
+		[Ignore("Horde does not allow submitting content ids before blobs")]
 		[TestMethod]
         public async Task PutContentIdMissingBlobAsync()
         {
diff --git a/Engine/Source/Programs/Horde/Horde.Server/Commands/Bundles/BundleWriteTest.cs b/Engine/Source/Programs/Horde/Horde.Server/Commands/Bundles/BundleWriteTest.cs
--- a/Engine/Source/Programs/Horde/Horde.Server/Commands/Bundles/BundleWriteTest.cs
+++ b/Engine/Source/Programs/Horde/Horde.Server/Commands/Bundles/BundleWriteTest.cs
@@ -28,7 +28,7 @@ namespace Horde.Server.Commands.Bundles
 			public Task<bool> DeleteRefAsync(RefName name, CancellationToken cancellationToken = default) => Task.FromResult(true);
 			public Task<Bundle> ReadBundleAsync(BundleLocator locator, CancellationToken cancellationToken = default) => throw new NotImplementedException();
 			public Task<ReadOnlyMemory<byte>> ReadBundleRangeAsync(BundleLocator locator, int offset, int length, CancellationToken cancellationToken = default) => throw new NotImplementedException();
-			public Task AddAliasAsync(Utf8String name, BlobHandle handle, CancellationToken cancellationToken = default) => throw new NotImplementedException();
+			public Task AddAliasAsync(Utf8String name, BlobHandle handle, int rank = 0, CancellationToken cancellationToken = default) => throw new NotImplementedException();
 			public Task RemoveAliasAsync(Utf8String name, BlobHandle handle, CancellationToken cancellationToken = default) => throw new NotImplementedException();
 			public IAsyncEnumerable<BlobHandle> FindAliasAsync(Utf8String alias, CancellationToken cancellationToken = default) => throw new NotImplementedException();
 			public Task<BlobHandle?> TryReadRefTargetAsync(RefName name, RefCacheTime cacheTime = default, CancellationToken cancellationToken = default) => throw new NotImplementedException();
diff --git a/Engine/Source/Programs/Horde/Horde.Server/Ddc/BlobService.cs b/Engine/Source/Programs/Horde/Horde.Server/Ddc/BlobService.cs
--- a/Engine/Source/Programs/Horde/Horde.Server/Ddc/BlobService.cs
+++ b/Engine/Source/Programs/Horde/Horde.Server/Ddc/BlobService.cs
@@ -145,7 +145,7 @@ namespace Horde.Server.Ddc
 				await writer.FlushAsync(cancellationToken);
 			}
 
-			await storageClient.AddAliasAsync(GetAlias(identifier), blobHandle, cancellationToken);
+			await storageClient.AddAliasAsync(GetAlias(identifier), blobHandle, cancellationToken: cancellationToken);
 			return identifier;
 		}
 	}
diff --git a/Engine/Source/Programs/Horde/Horde.Server/Storage/StorageService.cs b/Engine/Source/Programs/Horde/Horde.Server/Storage/StorageService.cs
--- a/Engine/Source/Programs/Horde/Horde.Server/Storage/StorageService.cs
+++ b/Engine/Source/Programs/Horde/Horde.Server/Storage/StorageService.cs
@@ -85,8 +85,8 @@ namespace Horde.Server.Storage
 			SupportsRedirects = backend.SupportsRedirects && !config.EnableAliases;
 		}
 
-		/// <inheritdoc cref="IStorageClient.AddAliasAsync(Utf8String, BlobHandle, CancellationToken)"/>
-		public abstract Task AddAliasAsync(Utf8String name, BundleNodeLocator target, CancellationToken cancellationToken = default);
+		/// <inheritdoc cref="IStorageClient.AddAliasAsync(Utf8String, BlobHandle, Int32, CancellationToken)"/>
+		public abstract Task AddAliasAsync(Utf8String name, BundleNodeLocator target, int rank = 0, CancellationToken cancellationToken = default);
 
 		/// <inheritdoc cref="IStorageClient.RemoveAliasAsync(Utf8String, BlobHandle, CancellationToken)"/>
 		public abstract Task RemoveAliasAsync(Utf8String name, BundleNodeLocator target, CancellationToken cancellationToken = default);
@@ -214,10 +214,10 @@ namespace Horde.Server.Storage
 			#region Nodes
 
 			/// <inheritdoc/>
-			public override Task AddAliasAsync(Utf8String name, BundleNodeHandle handle, CancellationToken cancellationToken = default) => _outer.AddAliasAsync(NamespaceId, name, handle.GetLocator(), cancellationToken);
+			public override Task AddAliasAsync(Utf8String name, BundleNodeHandle handle, int rank = 0, CancellationToken cancellationToken = default) => _outer.AddAliasAsync(NamespaceId, name, handle.GetLocator(), rank, cancellationToken);
 
 			/// <inheritdoc/>
-			public override Task AddAliasAsync(Utf8String name, BundleNodeLocator locator, CancellationToken cancellationToken = default) => _outer.AddAliasAsync(NamespaceId, name, locator, cancellationToken);
+			public override Task AddAliasAsync(Utf8String name, BundleNodeLocator locator, int rank = 0, CancellationToken cancellationToken = default) => _outer.AddAliasAsync(NamespaceId, name, locator, rank, cancellationToken);
 
 			/// <inheritdoc/>
 			public override Task RemoveAliasAsync(Utf8String name, BundleNodeHandle handle, CancellationToken cancellationToken = default) => _outer.RemoveAliasAsync(NamespaceId, name, handle.GetLocator(), cancellationToken);
@@ -228,7 +228,7 @@ namespace Horde.Server.Storage
 			/// <inheritdoc/>
 			public override async IAsyncEnumerable<BundleNodeHandle> FindAliasAsync(Utf8String alias, [EnumeratorCancellation] CancellationToken cancellationToken = default)
 			{
-				await foreach (BundleNodeLocator locator in _outer.FindNodesAsync(NamespaceId, alias, cancellationToken))
+				await foreach (BundleNodeLocator locator in _outer.FindAliasesAsync(NamespaceId, alias, cancellationToken))
 				{
 					yield return new FlushedNodeHandle(TreeReader, locator);
 				}
@@ -300,7 +300,7 @@ namespace Horde.Server.Storage
 			}
 		}
 
-		class ExportInfo
+		class AliasInfo
 		{
 			[BsonElement("alias")]
 			public string Alias { get; set; } = String.Empty;
@@ -308,18 +308,22 @@ namespace Horde.Server.Storage
 			[BsonElement("hash")]
 			public IoHash Hash { get; set; }
 
+			[BsonElement("rank"), BsonIgnoreIfDefault]
+			public int Rank { get; set; }
+
 			[BsonElement("idx")]
 			public int Index { get; set; }
 
-			public ExportInfo()
+			public AliasInfo()
 			{
 			}
 
-			public ExportInfo(string alias, IoHash hash, int index)
+			public AliasInfo(string alias, IoHash hash, int index, int rank)
 			{
 				Alias = alias;
 				Hash = hash;
 				Index = index;
+				Rank = rank;
 			}
 		}
 
@@ -337,7 +341,7 @@ namespace Horde.Server.Storage
 			public List<ObjectId>? Imports { get; set; }
 
 			[BsonElement("exp"), BsonIgnoreIfNull]
-			public List<ExportInfo>? Exports { get; set; }
+			public List<AliasInfo>? Exports { get; set; }
 
 			[BsonIgnore]
 			public BundleLocator Locator => new BundleLocator(Path);
@@ -471,7 +475,7 @@ namespace Horde.Server.Storage
 			List<MongoIndex<BlobInfo>> blobIndexes = new List<MongoIndex<BlobInfo>>();
 			blobIndexes.Add(keys => keys.Ascending(x => x.Imports));
 			blobIndexes.Add(keys => keys.Ascending(x => x.NamespaceId).Ascending(x => x.Path), unique: true);
-			blobIndexes.Add(keys => keys.Ascending(x => x.NamespaceId).Ascending($"{nameof(BlobInfo.Exports)}.{nameof(ExportInfo.Alias)}"));
+			blobIndexes.Add(keys => keys.Ascending(x => x.NamespaceId).Ascending($"{nameof(BlobInfo.Exports)}.{nameof(AliasInfo.Alias)}"));
 			_blobCollection = mongoService.GetCollection<BlobInfo>("Storage.Blobs", blobIndexes);
 
 			List<MongoIndex<RefInfo>> refIndexes = new List<MongoIndex<RefInfo>>();
@@ -610,7 +614,7 @@ namespace Horde.Server.Storage
 		#region Blobs
 
 		/// <inheritdoc/>
-		async Task<BundleLocator> AddBlobAsync(NamespaceId namespaceId, Utf8String prefix = default, List<ExportInfo>? exports = null, CancellationToken cancellationToken = default)
+		async Task<BundleLocator> AddBlobAsync(NamespaceId namespaceId, Utf8String prefix = default, List<AliasInfo>? exports = null, CancellationToken cancellationToken = default)
 		{
 			ObjectId id = ObjectId.GenerateNewId(_clock.UtcNow);
 			BundleLocator blobId = (prefix.Length > 0) ? new BundleLocator($"{prefix}/{id}") : new BundleLocator(id.ToString());
@@ -704,9 +708,10 @@ namespace Horde.Server.Storage
 		/// <param name="namespaceId">Namespace to search</param>
 		/// <param name="alias">Alias for the node</param>
 		/// <param name="target">Target node for the alias</param>
+		/// <param name="rank">Rank for the alias. Higher ranked aliases are preferred by default.</param>
 		/// <param name="cancellationToken">Cancellation token for the operation</param>
 		/// <returns>Sequence of thandles</returns>
-		async Task AddAliasAsync(NamespaceId namespaceId, Utf8String alias, BundleNodeLocator target, CancellationToken cancellationToken = default)
+		async Task AddAliasAsync(NamespaceId namespaceId, Utf8String alias, BundleNodeLocator target, int rank, CancellationToken cancellationToken = default)
 		{
 			BlobInfo? blobInfo = await _blobCollection.Find(x => x.NamespaceId == namespaceId && x.Path == target.Blob.Path).FirstOrDefaultAsync(cancellationToken);
 			if (blobInfo == null)
@@ -720,7 +725,7 @@ namespace Horde.Server.Storage
 			}
 
 			FilterDefinition<BlobInfo> filter = Builders<BlobInfo>.Filter.Expr(x => x.NamespaceId == blobInfo.NamespaceId && x.Path == target.Blob.Path);
-			UpdateDefinition<BlobInfo> update = Builders<BlobInfo>.Update.Push(x => x.Exports, new ExportInfo(alias.ToString(), target.Hash, target.ExportIdx));
+			UpdateDefinition<BlobInfo> update = Builders<BlobInfo>.Update.Push(x => x.Exports, new AliasInfo(alias.ToString(), target.Hash, target.ExportIdx, rank));
 			await _blobCollection.UpdateOneAsync(filter, update, cancellationToken: cancellationToken);
 		}
 
@@ -744,21 +749,28 @@ namespace Horde.Server.Storage
 		/// <param name="alias">Alias for the node</param>
 		/// <param name="cancellationToken">Cancellation token for the operation</param>
 		/// <returns>Sequence of thandles</returns>
-		async IAsyncEnumerable<BundleNodeLocator> FindNodesAsync(NamespaceId namespaceId, Utf8String alias, [EnumeratorCancellation] CancellationToken cancellationToken = default)
+		async IAsyncEnumerable<BundleNodeLocator> FindAliasesAsync(NamespaceId namespaceId, Utf8String alias, [EnumeratorCancellation] CancellationToken cancellationToken = default)
 		{
+			List<(BundleNodeLocator Locator, int Rank)> locators = new List<(BundleNodeLocator, int)>();
+
 			await foreach (BlobInfo blobInfo in _blobCollection.Find(x => x.NamespaceId == namespaceId && x.Exports!.Any(y => y.Alias == alias)).ToAsyncEnumerable(cancellationToken))
 			{
 				if (blobInfo.Exports != null)
 				{
-					foreach (ExportInfo exportInfo in blobInfo.Exports)
+					foreach (AliasInfo aliasInfo in blobInfo.Exports)
 					{
-						if (exportInfo.Alias == alias)
+						if (aliasInfo.Alias == alias)
 						{
-							yield return new BundleNodeLocator(exportInfo.Hash, blobInfo.Locator, exportInfo.Index);
+							locators.Add((new BundleNodeLocator(aliasInfo.Hash, blobInfo.Locator, aliasInfo.Index), aliasInfo.Rank));
 						}
 					}
 				}
 			}
+
+			foreach ((BundleNodeLocator locator, _) in locators.OrderByDescending(x => x.Rank))
+			{
+				yield return locator;
+			}
 		}
 
 		#endregion
diff --git a/Engine/Source/Programs/Shared/EpicGames.Horde/Compute/AgentStorageClient.cs b/Engine/Source/Programs/Shared/EpicGames.Horde/Compute/AgentStorageClient.cs
--- a/Engine/Source/Programs/Shared/EpicGames.Horde/Compute/AgentStorageClient.cs
+++ b/Engine/Source/Programs/Shared/EpicGames.Horde/Compute/AgentStorageClient.cs
@@ -71,7 +71,7 @@ namespace EpicGames.Horde.Compute
 		#region Nodes
 
 		/// <inheritdoc/>
-		public override Task AddAliasAsync(Utf8String name, BundleNodeHandle locator, CancellationToken cancellationToken = default)
+		public override Task AddAliasAsync(Utf8String name, BundleNodeHandle locator, int rank, CancellationToken cancellationToken = default)
 		{
 			throw new NotSupportedException();
 		}
diff --git a/Engine/Source/Programs/Shared/EpicGames.Horde/Storage/Backends/FileStorageClient.cs b/Engine/Source/Programs/Shared/EpicGames.Horde/Storage/Backends/FileStorageClient.cs
--- a/Engine/Source/Programs/Shared/EpicGames.Horde/Storage/Backends/FileStorageClient.cs
+++ b/Engine/Source/Programs/Shared/EpicGames.Horde/Storage/Backends/FileStorageClient.cs
@@ -102,7 +102,7 @@ namespace EpicGames.Horde.Storage.Backends
 		#region Aliases
 
 		/// <inheritdoc/>
-		public override Task AddAliasAsync(Utf8String name, BundleNodeHandle locator, CancellationToken cancellationToken = default)
+		public override Task AddAliasAsync(Utf8String name, BundleNodeHandle locator, int rank = 0, CancellationToken cancellationToken = default)
 		{
 			throw new NotSupportedException("File storage client does not currently support aliases.");
 		}
diff --git a/Engine/Source/Programs/Shared/EpicGames.Horde/Storage/Backends/HttpStorageClient.cs b/Engine/Source/Programs/Shared/EpicGames.Horde/Storage/Backends/HttpStorageClient.cs
--- a/Engine/Source/Programs/Shared/EpicGames.Horde/Storage/Backends/HttpStorageClient.cs
+++ b/Engine/Source/Programs/Shared/EpicGames.Horde/Storage/Backends/HttpStorageClient.cs
@@ -197,7 +197,7 @@ namespace EpicGames.Horde.Storage.Backends
 		#region Nodes
 
 		/// <inheritdoc/>
-		public override Task AddAliasAsync(Utf8String name, BundleNodeHandle locator, CancellationToken cancellationToken = default)
+		public override Task AddAliasAsync(Utf8String name, BundleNodeHandle locator, int rank = 0, CancellationToken cancellationToken = default)
 		{
 			throw new NotSupportedException("Http storage client does not currently support aliases.");
 		}
diff --git a/Engine/Source/Programs/Shared/EpicGames.Horde/Storage/Backends/MemoryStorageClient.cs b/Engine/Source/Programs/Shared/EpicGames.Horde/Storage/Backends/MemoryStorageClient.cs
--- a/Engine/Source/Programs/Shared/EpicGames.Horde/Storage/Backends/MemoryStorageClient.cs
+++ b/Engine/Source/Programs/Shared/EpicGames.Horde/Storage/Backends/MemoryStorageClient.cs
@@ -21,7 +21,7 @@ namespace EpicGames.Horde.Storage.Backends
 	/// </summary>
 	public class MemoryStorageClient : BundleStorageClient
 	{
-		record class ExportEntry(BundleNodeHandle Handle, ExportEntry? Next);
+		record class ExportEntry(BundleNodeHandle Handle, int Rank, ExportEntry? Next);
 
 		/// <summary>
 		/// Map of blob id to blob data
@@ -85,9 +85,9 @@ namespace EpicGames.Horde.Storage.Backends
 		#region Aliases
 
 		/// <inheritdoc/>
-		public override Task AddAliasAsync(Utf8String name, BundleNodeHandle handle, CancellationToken cancellationToken = default)
+		public override Task AddAliasAsync(Utf8String name, BundleNodeHandle handle, int rank = 0, CancellationToken cancellationToken = default)
 		{
-			_exports.AddOrUpdate(name, _ => new ExportEntry(handle, null), (_, entry) => new ExportEntry(handle, entry));
+			_exports.AddOrUpdate(name, _ => new ExportEntry(handle, rank, null), (_, entry) => new ExportEntry(handle, rank, entry));
 			return Task.CompletedTask;
 		}
 
diff --git a/Engine/Source/Programs/Shared/EpicGames.Horde/Storage/IStorageClient.cs b/Engine/Source/Programs/Shared/EpicGames.Horde/Storage/IStorageClient.cs
--- a/Engine/Source/Programs/Shared/EpicGames.Horde/Storage/IStorageClient.cs
+++ b/Engine/Source/Programs/Shared/EpicGames.Horde/Storage/IStorageClient.cs
@@ -91,8 +91,9 @@ namespace EpicGames.Horde.Storage
 		/// </summary>
 		/// <param name="name">Alias for the node</param>
 		/// <param name="handle">Locator for the node</param>
+		/// <param name="rank">Rank for this alias. In situations where an alias has multiple mappings, the alias with the highest rank will be returned by default.</param>
 		/// <param name="cancellationToken">Cancellation token for the operation</param>
-		Task AddAliasAsync(Utf8String name, BlobHandle handle, CancellationToken cancellationToken = default);
+		Task AddAliasAsync(Utf8String name, BlobHandle handle, int rank = 0, CancellationToken cancellationToken = default);
 
 		/// <summary>
 		/// Removes an alias from a node
