Horde: Enable more DDC tests.

diff --git a/Engine/Source/Programs/Horde/Horde.Server.Tests/Ddc/Functional/Objects/ReferencesTests.cs b/Engine/Source/Programs/Horde/Horde.Server.Tests/Ddc/Functional/Objects/ReferencesTests.cs
--- a/Engine/Source/Programs/Horde/Horde.Server.Tests/Ddc/Functional/Objects/ReferencesTests.cs
+++ b/Engine/Source/Programs/Horde/Horde.Server.Tests/Ddc/Functional/Objects/ReferencesTests.cs
@@ -52,14 +52,6 @@ namespace Horde.Server.Tests.Ddc.FunctionalTests.References
 		Task<long> DeleteBucketAsync(NamespaceId ns, BucketId bucket);
 	}
 
-	static class CompileHacks
-	{
-		public static Task DeleteObjectAsync(this IBlobService blobService, NamespaceId namespaceId, BlobId blobId)
-		{
-			throw new NotImplementedException();
-		}
-	}
-
 	[TestClass]
     public class ReferencesTests : ControllerIntegrationTest
     {
@@ -1045,7 +1037,6 @@ namespace Horde.Server.Tests.Ddc.FunctionalTests.References
             }
         }
 
-		[Ignore("Horde does not support deleting objects manually (ie. BlobService.DeleteObjectAsync)")]
 		[TestMethod]
         public async Task GetMissingContentIdRecordAsync()
         {
@@ -1156,7 +1147,6 @@ namespace Horde.Server.Tests.Ddc.FunctionalTests.References
             }
         }
 
-		[Ignore("Horde does not support deleting objects manually (ie. BlobService.DeleteObjectAsync)")]
 		[TestMethod]
         public async Task GetMissingCompressedBufferAttachmentAsync()
         {
@@ -1261,7 +1251,6 @@ namespace Horde.Server.Tests.Ddc.FunctionalTests.References
             }
         }
 
-		[Ignore("Horde does not support deleting objects manually (ie. BlobService.DeleteObjectAsync)")]
 		[TestMethod]
         public async Task GetMissingBlobRecordAsync()
         {
diff --git a/Engine/Source/Programs/Horde/Horde.Server/Ddc/BlobService.cs b/Engine/Source/Programs/Horde/Horde.Server/Ddc/BlobService.cs
--- a/Engine/Source/Programs/Horde/Horde.Server/Ddc/BlobService.cs
+++ b/Engine/Source/Programs/Horde/Horde.Server/Ddc/BlobService.cs
@@ -6,6 +6,7 @@ using System.IO;
 using System.Linq;
 using System.Threading;
 using System.Threading.Tasks;
+using Amazon.S3.Model;
 using EpicGames.Core;
 using EpicGames.Horde.Storage;
 using Horde.Server.Storage;
@@ -67,6 +68,18 @@ namespace Horde.Server.Ddc
 			return await storageClient.FindAliasAsync(GetAlias(blob), cancellationToken).AnyAsync(cancellationToken);
 		}
 
+		public async Task DeleteObjectAsync(NamespaceId ns, BlobId blob, CancellationToken cancellationToken)
+		{
+			IStorageClient storageClient = await _storageService.GetClientAsync(ns, cancellationToken);
+			Utf8String alias = GetAlias(blob);
+
+			List<BlobHandle> handles = await storageClient.FindAliasAsync(alias, cancellationToken).ToListAsync(cancellationToken);
+			foreach (BlobHandle handle in handles)
+			{
+				await storageClient.RemoveAliasAsync(alias, handle);
+			}
+		}
+
 		public async Task<BlobId[]> FilterOutKnownBlobsAsync(NamespaceId ns, IEnumerable<BlobId> blobIds, CancellationToken cancellationToken)
 		{
 			IStorageClient storageClient = await _storageService.GetClientAsync(ns, cancellationToken);
diff --git a/Engine/Source/Programs/Horde/Horde.Server/Ddc/ContentIdService.cs b/Engine/Source/Programs/Horde/Horde.Server/Ddc/ContentIdService.cs
--- a/Engine/Source/Programs/Horde/Horde.Server/Ddc/ContentIdService.cs
+++ b/Engine/Source/Programs/Horde/Horde.Server/Ddc/ContentIdService.cs
@@ -52,7 +52,7 @@ namespace Horde.Server.Ddc
 				throw new BlobNotFoundException(ns, blobId);
 			}
 
-			await storageClient.AddAliasAsync(GetAlias(contentId), blobHandle, cancellationToken);
+			await storageClient.AddAliasAsync(GetAlias(contentId), blobHandle, -contentWeight, cancellationToken);
 		}
 	}
 }
diff --git a/Engine/Source/Programs/Horde/Horde.Server/Ddc/IBlobService.cs b/Engine/Source/Programs/Horde/Horde.Server/Ddc/IBlobService.cs
--- a/Engine/Source/Programs/Horde/Horde.Server/Ddc/IBlobService.cs
+++ b/Engine/Source/Programs/Horde/Horde.Server/Ddc/IBlobService.cs
@@ -34,6 +34,8 @@ namespace Horde.Server.Ddc
 
 		Task<bool> ExistsAsync(NamespaceId ns, BlobId blob, List<string>? storageLayers = null, CancellationToken cancellationToken = default);
 
+		Task DeleteObjectAsync(NamespaceId ns, BlobId blob, CancellationToken cancellationToken = default);
+
 		Task<BlobId[]> FilterOutKnownBlobsAsync(NamespaceId ns, IEnumerable<BlobId> blobs, CancellationToken cancellationToken = default);
 		Task<BlobContents> GetObjectsAsync(NamespaceId ns, BlobId[] refRequestBlobReferences, CancellationToken cancellationToken = default);
 	}
diff --git a/Engine/Source/Programs/Horde/Horde.Server/Ddc/IRefService.cs b/Engine/Source/Programs/Horde/Horde.Server/Ddc/IRefService.cs
--- a/Engine/Source/Programs/Horde/Horde.Server/Ddc/IRefService.cs
+++ b/Engine/Source/Programs/Horde/Horde.Server/Ddc/IRefService.cs
@@ -62,7 +62,7 @@ namespace Horde.Server.Ddc
 
 	public class ObjectHashMismatchException : Exception
 	{
-		public ObjectHashMismatchException(NamespaceId ns, BucketId bucket, IoHash name, BlobId suppliedHash, BlobId actualHash) : base($"Object {name} in bucket {bucket} and namespace {ns} did not reference hash {suppliedHash} was referencing {actualHash}")
+		public ObjectHashMismatchException(NamespaceId ns, BucketId bucket, RefId name, BlobId suppliedHash, BlobId actualHash) : base($"Object {name} in bucket {bucket} and namespace {ns} did not reference hash {suppliedHash} was referencing {actualHash}")
 		{
 		}
 	}
diff --git a/Engine/Source/Programs/Horde/Horde.Server/Storage/StorageService.cs b/Engine/Source/Programs/Horde/Horde.Server/Storage/StorageService.cs
--- a/Engine/Source/Programs/Horde/Horde.Server/Storage/StorageService.cs
+++ b/Engine/Source/Programs/Horde/Horde.Server/Storage/StorageService.cs
@@ -341,7 +341,7 @@ namespace Horde.Server.Storage
 			public List<ObjectId>? Imports { get; set; }
 
 			[BsonElement("exp"), BsonIgnoreIfNull]
-			public List<AliasInfo>? Exports { get; set; }
+			public List<AliasInfo>? Aliases { get; set; }
 
 			[BsonIgnore]
 			public BundleLocator Locator => new BundleLocator(Path);
@@ -475,7 +475,7 @@ namespace Horde.Server.Storage
 			List<MongoIndex<BlobInfo>> blobIndexes = new List<MongoIndex<BlobInfo>>();
 			blobIndexes.Add(keys => keys.Ascending(x => x.Imports));
 			blobIndexes.Add(keys => keys.Ascending(x => x.NamespaceId).Ascending(x => x.Path), unique: true);
-			blobIndexes.Add(keys => keys.Ascending(x => x.NamespaceId).Ascending($"{nameof(BlobInfo.Exports)}.{nameof(AliasInfo.Alias)}"));
+			blobIndexes.Add(keys => keys.Ascending(x => x.NamespaceId).Ascending($"{nameof(BlobInfo.Aliases)}.{nameof(AliasInfo.Alias)}"));
 			_blobCollection = mongoService.GetCollection<BlobInfo>("Storage.Blobs", blobIndexes);
 
 			List<MongoIndex<RefInfo>> refIndexes = new List<MongoIndex<RefInfo>>();
@@ -620,7 +620,7 @@ namespace Horde.Server.Storage
 			BundleLocator blobId = (prefix.Length > 0) ? new BundleLocator($"{prefix}/{id}") : new BundleLocator(id.ToString());
 
 			BlobInfo blobInfo = new BlobInfo(id, namespaceId, blobId);
-			blobInfo.Exports = exports;
+			blobInfo.Aliases = exports;
 			await _blobCollection.InsertOneAsync(blobInfo, new InsertOneOptions { }, cancellationToken);
 
 			return blobId;
@@ -719,13 +719,13 @@ namespace Horde.Server.Storage
 				throw new KeyNotFoundException($"Missing blob {target.Blob}");
 			}
 
-			if (blobInfo.Exports != null && blobInfo.Exports.Any(x => x.Alias == alias && x.Index == target.ExportIdx))
+			if (blobInfo.Aliases != null && blobInfo.Aliases.Any(x => x.Alias == alias && x.Index == target.ExportIdx))
 			{
 				return;
 			}
 
 			FilterDefinition<BlobInfo> filter = Builders<BlobInfo>.Filter.Expr(x => x.NamespaceId == blobInfo.NamespaceId && x.Path == target.Blob.Path);
-			UpdateDefinition<BlobInfo> update = Builders<BlobInfo>.Update.Push(x => x.Exports, new AliasInfo(alias.ToString(), target.Hash, target.ExportIdx, rank));
+			UpdateDefinition<BlobInfo> update = Builders<BlobInfo>.Update.Push(x => x.Aliases, new AliasInfo(alias.ToString(), target.Hash, target.ExportIdx, rank));
 			await _blobCollection.UpdateOneAsync(filter, update, cancellationToken: cancellationToken);
 		}
 
@@ -737,9 +737,11 @@ namespace Horde.Server.Storage
 		/// <param name="target">Target node for the alias</param>
 		/// <param name="cancellationToken">Cancellation token for the operation</param>
 		/// <returns>Sequence of thandles</returns>
-		Task RemoveAliasAsync(NamespaceId namespaceId, Utf8String alias, BundleNodeLocator target, CancellationToken cancellationToken = default)
+		async Task RemoveAliasAsync(NamespaceId namespaceId, Utf8String alias, BundleNodeLocator target, CancellationToken cancellationToken = default)
 		{
-			throw new NotSupportedException();
+			FilterDefinition<BlobInfo> filter = Builders<BlobInfo>.Filter.Expr(x => x.NamespaceId == namespaceId && x.Path == target.Blob.Path);
+			UpdateDefinition<BlobInfo> update = Builders<BlobInfo>.Update.PullFilter(x => x.Aliases, Builders<AliasInfo>.Filter.Expr(x => x.Alias == alias));
+			await _blobCollection.UpdateOneAsync(filter, update, cancellationToken: cancellationToken);
 		}
 
 		/// <summary>
@@ -753,11 +755,11 @@ namespace Horde.Server.Storage
 		{
 			List<(BundleNodeLocator Locator, int Rank)> locators = new List<(BundleNodeLocator, int)>();
 
-			await foreach (BlobInfo blobInfo in _blobCollection.Find(x => x.NamespaceId == namespaceId && x.Exports!.Any(y => y.Alias == alias)).ToAsyncEnumerable(cancellationToken))
+			await foreach (BlobInfo blobInfo in _blobCollection.Find(x => x.NamespaceId == namespaceId && x.Aliases!.Any(y => y.Alias == alias)).ToAsyncEnumerable(cancellationToken))
 			{
-				if (blobInfo.Exports != null)
+				if (blobInfo.Aliases != null)
 				{
-					foreach (AliasInfo aliasInfo in blobInfo.Exports)
+					foreach (AliasInfo aliasInfo in blobInfo.Aliases)
 					{
 						if (aliasInfo.Alias == alias)
 						{
