Horde: Missing file.

diff --git a/Engine/Source/Programs/Horde/Horde.Server/Ddc/RefService.cs b/Engine/Source/Programs/Horde/Horde.Server/Ddc/RefService.cs
--- a/Engine/Source/Programs/Horde/Horde.Server/Ddc/RefService.cs
+++ b/Engine/Source/Programs/Horde/Horde.Server/Ddc/RefService.cs
@@ -9,17 +9,20 @@ using EpicGames.Core;
 using EpicGames.Horde.Storage;
 using EpicGames.Serialization;
 using Horde.Server.Storage;
+using OpenTelemetry.Trace;
 
 namespace Horde.Server.Ddc
 {
 	class RefService : IRefService
 	{
 		readonly StorageService _storageService;
+		readonly IReferenceResolver _referenceResolver;
 		readonly IBlobService _blobService;
 
-		public RefService(StorageService storageService, IBlobService blobService)
+		public RefService(StorageService storageService, IReferenceResolver referenceResolver, IBlobService blobService)
 		{
 			_storageService = storageService;
+			_referenceResolver = referenceResolver;
 			_blobService = blobService;
 		}
 
@@ -40,45 +43,84 @@ namespace Horde.Server.Ddc
 
 		public async Task<(ContentId[], BlobId[])> FinalizeAsync(NamespaceId ns, BucketId bucket, RefId key, BlobId blobHash, CancellationToken cancellationToken)
 		{
-			List<BlobId> missingBlobIds = new List<BlobId>();
+			IStorageClient storageClient = await _storageService.GetClientAsync(ns, cancellationToken);
+
+			BlobHandle? blobHandle = await storageClient.FindAliasAsync(BlobService.GetAlias(blobHash), cancellationToken).FirstOrDefaultAsync(cancellationToken);
+			if (blobHandle == null)
+			{
+				throw new BlobNotFoundException(ns, blobHash);
+			}
+
+			BlobData blobContents = await blobHandle.ReadAsync(cancellationToken);
+			CbObject payload = new CbObject(blobContents.Data);
 
-			// Read all the blobs back in
-			List<BlobId> blobIds = new List<BlobId>();
-			blobIds.Add(blobHash);
+			BlobId[] referencedBlobs = Array.Empty<BlobId>();
 
-			for (int idx = 0; idx < blobIds.Count; idx++)
+			ContentId[] missingReferences = Array.Empty<ContentId>();
+			BlobId[] missingBlobs = Array.Empty<BlobId>();
+			bool hasReferences = HasAttachments(payload);
+			if (hasReferences)
 			{
 				try
 				{
-					using BlobContents contents = await _blobService.GetObjectAsync(ns, blobIds[idx], storageLayers: null, supportsRedirectUri: false, cancellationToken: cancellationToken);
-					byte[] data = await contents.Stream.ReadAllBytesAsync(cancellationToken);
-
-					CbObject obj = new CbObject(data);
-					obj.IterateAttachments(x => blobIds.Add(new BlobId(x.AsAttachment())));
+					referencedBlobs = await _referenceResolver.GetReferencedBlobsAsync(ns, payload).ToArrayAsync(cancellationToken);
+				}
+				catch (PartialReferenceResolveException e)
+				{
+					missingReferences = e.UnresolvedReferences.ToArray();
 				}
-				catch (BlobNotFoundException)
+				catch (ReferenceIsMissingBlobsException e)
 				{
-					missingBlobIds.Add(blobIds[idx]);
+					missingBlobs = e.MissingBlobs.ToArray();
 				}
 			}
 
-			if (missingBlobIds.Count == 0)
+			if (missingReferences.Length == 0 && missingBlobs.Length == 0)
 			{
-				IStorageClient storageClient = await _storageService.GetClientAsync(ns, cancellationToken);
+				// TODO: We resolved all these blobs above... Need to just have GetReferencedBlobs just return the appropriate handles directly.
+				DdcRefNode refNode = new DdcRefNode(blobHandle.Hash);
+				refNode.References.Add((blobHandle.Hash, blobHandle));
+				foreach (BlobId referencedBlob in referencedBlobs)
+				{
+					BlobHandle handle = await storageClient.FindAliasAsync(BlobService.GetAlias(referencedBlob), cancellationToken).FirstAsync(cancellationToken);
+					refNode.References.Add((referencedBlob.Hash, handle));
+				}
 
-				BlobHandle? handle = await storageClient.FindAliasAsync(BlobService.GetAlias(blobHash), cancellationToken).FirstOrDefaultAsync(cancellationToken);
-				if (handle == null)
+				await storageClient.WriteRefAsync(GetRefName(bucket, key), refNode, cancellationToken: cancellationToken);
+			}
+
+			return (missingReferences, missingBlobs);
+		}
+
+		private bool HasAttachments(CbObject payload)
+		{
+			bool FieldHasAttachments(CbField field)
+			{
+				if (field.IsObject())
 				{
-					throw new BlobNotFoundException(ns, blobHash);
+					bool hasAttachment = HasAttachments(field.AsObject());
+					if (hasAttachment)
+					{
+						return true;
+					}
 				}
 
-				DdcRefNode refNode = new DdcRefNode(blobHash.Hash);
-				refNode.References.Add((blobHash.Hash, handle));
+				if (field.IsArray())
+				{
+					foreach (CbField subField in field.AsArray())
+					{
+						bool hasAttachment = FieldHasAttachments(subField);
+						if (hasAttachment)
+						{
+							return true;
+						}
+					}
+				}
 
-				await storageClient.WriteRefAsync(GetRefName(bucket, key), refNode, cancellationToken: cancellationToken);
+				return field.IsAttachment();
 			}
 
-			return (Array.Empty<ContentId>(), missingBlobIds.ToArray());
+			return payload.Any(FieldHasAttachments);
 		}
 
 		public async Task<(RefRecord, BlobContents?)> GetAsync(NamespaceId ns, BucketId bucket, RefId key, string[] fields, bool doLastAccessTracking, CancellationToken cancellationToken)
