Hode: Use BinaryId as the base for UserId instances.

diff --git a/Engine/Source/Programs/Horde/Horde.Server.Tests/JobTaskSourceTests.cs b/Engine/Source/Programs/Horde/Horde.Server.Tests/JobTaskSourceTests.cs
--- a/Engine/Source/Programs/Horde/Horde.Server.Tests/JobTaskSourceTests.cs
+++ b/Engine/Source/Programs/Horde/Horde.Server.Tests/JobTaskSourceTests.cs
@@ -106,7 +106,7 @@ namespace Horde.Server.Tests
 
 			// update template with some step states
 			IStream Stream = await StreamCollection.GetAsync(fixture.StreamConfig!);
-			Stream = Deref(await StreamCollection.TryUpdateTemplateRefAsync(Stream, fixture.TemplateRefId1, new List<UpdateStepStateRequest>() { new UpdateStepStateRequest() { Name = "Paused Step", PausedByUserId = UserId.GenerateNewId().ToString() } }));
+			Stream = Deref(await StreamCollection.TryUpdateTemplateRefAsync(Stream, fixture.TemplateRefId1, new List<UpdateStepStateRequest>() { new UpdateStepStateRequest() { Name = "Paused Step", PausedByUserId = new UserId(BinaryIdUtils.CreateNew()).ToString() } }));
 
 			// create a new graph with the associated nodes
 			List<NewGroup> newGroups = new List<NewGroup>();
diff --git a/Engine/Source/Programs/Horde/Horde.Server.Tests/Stubs/Services/PerforceServiceStub.cs b/Engine/Source/Programs/Horde/Horde.Server.Tests/Stubs/Services/PerforceServiceStub.cs
--- a/Engine/Source/Programs/Horde/Horde.Server.Tests/Stubs/Services/PerforceServiceStub.cs
+++ b/Engine/Source/Programs/Horde/Horde.Server.Tests/Stubs/Services/PerforceServiceStub.cs
@@ -11,6 +11,7 @@ using EpicGames.Perforce;
 using Horde.Server.Perforce;
 using Horde.Server.Streams;
 using Horde.Server.Users;
+using Horde.Server.Utilities;
 using Microsoft.VisualStudio.TestTools.UnitTesting;
 
 namespace Horde.Server.Tests.Stubs.Services
@@ -26,7 +27,7 @@ namespace Horde.Server.Tests.Stubs.Services
 
 			public User(string login)
 			{
-				Id = UserId.GenerateNewId();
+				Id = new UserId(BinaryIdUtils.CreateNew());
 				Name = login.ToUpperInvariant();
 				Email = $"{login}@server";
 				Login = login;
diff --git a/Engine/Source/Programs/Horde/Horde.Server/Notifications/Sinks/SlackNotificationSink.cs b/Engine/Source/Programs/Horde/Horde.Server/Notifications/Sinks/SlackNotificationSink.cs
--- a/Engine/Source/Programs/Horde/Horde.Server/Notifications/Sinks/SlackNotificationSink.cs
+++ b/Engine/Source/Programs/Horde/Horde.Server/Notifications/Sinks/SlackNotificationSink.cs
@@ -41,6 +41,7 @@ using Microsoft.Extensions.Options;
 using MongoDB.Bson;
 using MongoDB.Bson.Serialization.Attributes;
 using MongoDB.Driver;
+using EpicGames.Horde;
 
 namespace Horde.Server.Notifications.Sinks
 {
@@ -2766,7 +2767,7 @@ namespace Horde.Server.Notifications.Sinks
 							{
 								int issueId = Int32.Parse(match.Groups[1].Value, NumberStyles.Integer, CultureInfo.InvariantCulture);
 								string verb = match.Groups[2].Value;
-								UserId userId = new UserId(ObjectId.Parse(match.Groups[3].Value));
+								UserId userId = new UserId(BinaryId.Parse(match.Groups[3].Value));
 								await HandleIssueDmResponseAsync(issueId, verb, userId, payload.User.Id);
 							}
 							else if (TryMatch(action.Value, @"^issue_(\d+)_([a-zA-Z]+)$", out match) && payload.TriggerId != null)
diff --git a/Engine/Source/Programs/Horde/Horde.Server/Users/UserCollection.cs b/Engine/Source/Programs/Horde/Horde.Server/Users/UserCollection.cs
--- a/Engine/Source/Programs/Horde/Horde.Server/Users/UserCollection.cs
+++ b/Engine/Source/Programs/Horde/Horde.Server/Users/UserCollection.cs
@@ -133,7 +133,7 @@ namespace Horde.Server.Users
 		public async Task<IUser> FindOrAddUserByLoginAsync(string login, string? name, string? email)
 		{
 			ClaimDocument newPrimaryClaim = new ClaimDocument(ClaimTypes.Name, login);
-			UpdateDefinition<UserDocument> update = Builders<UserDocument>.Update.SetOnInsert(x => x.Id, UserId.GenerateNewId());
+			UpdateDefinition<UserDocument> update = Builders<UserDocument>.Update.SetOnInsert(x => x.Id, new UserId(BinaryIdUtils.CreateNew()));
 			return await _users.FindOneAndUpdateAsync<UserDocument>(x => x.PrimaryClaim == newPrimaryClaim, update, new FindOneAndUpdateOptions<UserDocument> { IsUpsert = true, ReturnDocument = ReturnDocument.After });
 		}
 
diff --git a/Engine/Source/Programs/Horde/Horde.Server/Users/UserCollectionV2.cs b/Engine/Source/Programs/Horde/Horde.Server/Users/UserCollectionV2.cs
--- a/Engine/Source/Programs/Horde/Horde.Server/Users/UserCollectionV2.cs
+++ b/Engine/Source/Programs/Horde/Horde.Server/Users/UserCollectionV2.cs
@@ -299,7 +299,7 @@ namespace Horde.Server.Users
 		/// <inheritdoc/>
 		public async Task<IUser> FindOrAddUserByLoginAsync(string login, string? name, string? email)
 		{
-			UserId newUserId = UserId.GenerateNewId();
+			UserId newUserId = new UserId(BinaryIdUtils.CreateNew());
 			UpdateDefinition<UserDocument> update = Builders<UserDocument>.Update.SetOnInsert(x => x.Id, newUserId).SetOnInsert(x => x.Login, login).Unset(x => x.Hidden);
 
 			if (name == null)
diff --git a/Engine/Source/Programs/Horde/Horde.Server/Users/UserId.cs b/Engine/Source/Programs/Horde/Horde.Server/Users/UserId.cs
--- a/Engine/Source/Programs/Horde/Horde.Server/Users/UserId.cs
+++ b/Engine/Source/Programs/Horde/Horde.Server/Users/UserId.cs
@@ -2,6 +2,7 @@
 
 using System.ComponentModel;
 using EpicGames.Core;
+using EpicGames.Horde;
 using Horde.Server.Utilities;
 using MongoDB.Bson;
 
@@ -12,32 +13,27 @@ namespace Horde.Server.Users
 	/// </summary>
 	/// <param name="Id">Id to construct from</param>
 	[JsonSchemaString]
-	[TypeConverter(typeof(ObjectIdTypeConverter<UserId, UserIdConverter>))]
+	[TypeConverter(typeof(BinaryIdTypeConverter<UserId, UserIdConverter>))]
 	[ObjectIdConverter(typeof(UserIdConverter))]
-	public record struct UserId(ObjectId Id)
+	public record struct UserId(BinaryId Id)
 	{
 		/// <summary>
 		/// Constant value for empty user id
 		/// </summary>
-		public static UserId Empty { get; } = new UserId(ObjectId.Empty);
+		public static UserId Empty { get; } = default;
 
 		/// <summary>
 		/// Special user id for an anonymous administrator
 		/// </summary>
 		public static UserId Anonymous { get; } = UserId.Parse("63f7d3525119b9aa4c0f035a");
 
-		/// <summary>
-		/// Creates a new <see cref="UserId"/>
-		/// </summary>
-		public static UserId GenerateNewId() => new UserId(ObjectId.GenerateNewId());
-
 		/// <inheritdoc cref="ObjectId.Parse(System.String)"/>
-		public static UserId Parse(string text) => new UserId(ObjectId.Parse(text));
+		public static UserId Parse(string text) => new UserId(BinaryId.Parse(text));
 
 		/// <inheritdoc cref="ObjectId.TryParse(System.String, out ObjectId)"/>
 		public static bool TryParse(string text, out UserId id)
 		{
-			if (ObjectId.TryParse(text, out ObjectId objectId))
+			if (BinaryId.TryParse(text, out BinaryId objectId))
 			{
 				id = new UserId(objectId);
 				return true;
@@ -54,14 +50,14 @@ namespace Horde.Server.Users
 	}
 
 	/// <summary>
-	/// Converter to and from <see cref="ObjectId"/> instances.
+	/// Converter to and from <see cref="BinaryId"/> instances.
 	/// </summary>
-	class UserIdConverter : ObjectIdConverter<UserId>
+	class UserIdConverter : BinaryIdConverter<UserId>
 	{
 		/// <inheritdoc/>
-		public override UserId FromObjectId(ObjectId id) => new UserId(id);
+		public override UserId FromBinaryId(BinaryId id) => new UserId(id);
 
 		/// <inheritdoc/>
-		public override ObjectId ToObjectId(UserId value) => value.Id;
+		public override BinaryId ToBinaryId(UserId value) => value.Id;
 	}
 }
