Horde: Fix incorrect attribute on UserId.

diff --git a/Engine/Source/Programs/Horde/Horde.Server/Users/UserId.cs b/Engine/Source/Programs/Horde/Horde.Server/Users/UserId.cs
--- a/Engine/Source/Programs/Horde/Horde.Server/Users/UserId.cs
+++ b/Engine/Source/Programs/Horde/Horde.Server/Users/UserId.cs
@@ -14,7 +14,7 @@ namespace Horde.Server.Users
 	/// <param name="Id">Id to construct from</param>
 	[JsonSchemaString]
 	[TypeConverter(typeof(BinaryIdTypeConverter<UserId, UserIdConverter>))]
-	[ObjectIdConverter(typeof(UserIdConverter))]
+	[BinaryIdConverter(typeof(UserIdConverter))]
 	public record struct UserId(BinaryId Id)
 	{
 		/// <summary>
