Horde - Exposed parts of the AWS EC2 metadata as env. variables

diff --git a/Engine/Source/Programs/Horde/Horde.Agent/Execution/JobExecutor.cs b/Engine/Source/Programs/Horde/Horde.Agent/Execution/JobExecutor.cs
--- a/Engine/Source/Programs/Horde/Horde.Agent/Execution/JobExecutor.cs
+++ b/Engine/Source/Programs/Horde/Horde.Agent/Execution/JobExecutor.cs
@@ -1337,6 +1337,10 @@ namespace Horde.Agent.Execution
 			FileUtils.ForceDeleteDirectoryContents(testDataDir);
 			newEnvVars["UE_TESTDATA_DIR"] = testDataDir.FullName;
 
+			// TODO: These are AWS specific, this should be extended to handle more clouds or for licensees to be able to set these
+			newEnvVars["UE_HORDE_AVAILABILITY_ZONE"] = Amazon.Util.EC2InstanceMetadata.AvailabilityZone ?? "";
+			newEnvVars["UE_HORDE_REGION"] = Amazon.Util.EC2InstanceMetadata.Region?.DisplayName ?? "";
+
 			newEnvVars["UE_HORDE_JOBID"] = _jobId;
 			newEnvVars["UE_HORDE_BATCHID"] = _batchId;
 			newEnvVars["UE_HORDE_STEPID"] = step.StepId;
