Horde: Rename constants for artifact types to match what we use them for.

diff --git a/Engine/Source/Programs/Horde/Horde.Agent/Execution/JobExecutor.cs b/Engine/Source/Programs/Horde/Horde.Agent/Execution/JobExecutor.cs
--- a/Engine/Source/Programs/Horde/Horde.Agent/Execution/JobExecutor.cs
+++ b/Engine/Source/Programs/Horde/Horde.Agent/Execution/JobExecutor.cs
@@ -999,7 +999,7 @@ namespace Horde.Agent.Execution
 				CreateJobArtifactRequest artifactRequest = new CreateJobArtifactRequest();
 				artifactRequest.JobId = _jobId;
 				artifactRequest.StepId = step.StepId;
-				artifactRequest.Type = JobArtifactType.Output;
+				artifactRequest.Type = JobArtifactType.TempStorage;
 				artifactRequest.NamespaceId = _namespaceId.ToString();
 				artifactRequest.RefName = refName.ToString();
 
diff --git a/Engine/Source/Programs/Horde/Horde.Server.Tests/ArtifactTests.cs b/Engine/Source/Programs/Horde/Horde.Server.Tests/ArtifactTests.cs
--- a/Engine/Source/Programs/Horde/Horde.Server.Tests/ArtifactTests.cs
+++ b/Engine/Source/Programs/Horde/Horde.Server.Tests/ArtifactTests.cs
@@ -24,7 +24,7 @@ namespace Horde.Server.Tests
 		public async Task CreateArtifactAsync()
 		{
 			IArtifactCollection artifactCollection = ServiceProvider.GetRequiredService<IArtifactCollection>();
-			IArtifact artifact = await artifactCollection.AddAsync(new ArtifactId(BinaryIdUtils.CreateNew()), ArtifactType.StepOutput, new string[] { "test1", "test2" }, Namespace.Artifacts, new RefName("test"), null, AclScopeName.Root);
+			IArtifact artifact = await artifactCollection.AddAsync(new ArtifactId(BinaryIdUtils.CreateNew()), ArtifactType.StepTempStorage, new string[] { "test1", "test2" }, Namespace.Artifacts, new RefName("test"), null, AclScopeName.Root);
 
 			{
 				List<IArtifact> artifacts = await artifactCollection.FindAsync(keys: new[] { "test1" }).ToListAsync();
@@ -53,7 +53,7 @@ namespace Horde.Server.Tests
 			await expirationService.StartAsync(CancellationToken.None);
 			
 			IArtifactCollection artifactCollection = ServiceProvider.GetRequiredService<IArtifactCollection>();
-			IArtifact artifact = await artifactCollection.AddAsync(new ArtifactId(BinaryIdUtils.CreateNew()), ArtifactType.StepOutput, new string[] { "test1", "test2" }, Namespace.Artifacts, new RefName("test"), clock.UtcNow + TimeSpan.FromHours(1.0), AclScopeName.Root);
+			IArtifact artifact = await artifactCollection.AddAsync(new ArtifactId(BinaryIdUtils.CreateNew()), ArtifactType.StepTempStorage, new string[] { "test1", "test2" }, Namespace.Artifacts, new RefName("test"), clock.UtcNow + TimeSpan.FromHours(1.0), AclScopeName.Root);
 
 			{
 				List<IArtifact> artifacts = await artifactCollection.FindAsync(keys: new[] { "test1" }).ToListAsync();
diff --git a/Engine/Source/Programs/Horde/Horde.Server/Jobs/JobRpcService.cs b/Engine/Source/Programs/Horde/Horde.Server/Jobs/JobRpcService.cs
--- a/Engine/Source/Programs/Horde/Horde.Server/Jobs/JobRpcService.cs
+++ b/Engine/Source/Programs/Horde/Horde.Server/Jobs/JobRpcService.cs
@@ -77,9 +77,9 @@ namespace Horde.Server.Jobs
 
 			ArtifactType type = request.Type switch
 			{
-				JobArtifactType.Output => ArtifactType.StepOutput,
+				JobArtifactType.TempStorage => ArtifactType.StepTempStorage,
 				JobArtifactType.Saved => ArtifactType.StepSaved,
-				JobArtifactType.Trace => ArtifactType.StepTrace,
+				JobArtifactType.Trace => ArtifactType.StepTraces,
 				JobArtifactType.TestData => ArtifactType.StepTestData,
 				_ => throw new StructuredRpcException(StatusCode.InvalidArgument, "Invalid artifact type")
 			};
diff --git a/Engine/Source/Programs/Shared/EpicGames.Horde/Api/ArtifactType.cs b/Engine/Source/Programs/Shared/EpicGames.Horde/Api/ArtifactType.cs
--- a/Engine/Source/Programs/Shared/EpicGames.Horde/Api/ArtifactType.cs
+++ b/Engine/Source/Programs/Shared/EpicGames.Horde/Api/ArtifactType.cs
@@ -22,7 +22,7 @@ namespace EpicGames.Horde.Api
 		/// <summary>
 		/// Output from a build step
 		/// </summary>
-		public static ArtifactType StepOutput { get; } = new ArtifactType("step-output");
+		public static ArtifactType StepTempStorage { get; } = new ArtifactType("step-output");
 
 		/// <summary>
 		/// Captured state from the machine executing a build step
@@ -32,7 +32,7 @@ namespace EpicGames.Horde.Api
 		/// <summary>
 		/// Traces from the machine executing a build step
 		/// </summary>
-		public static ArtifactType StepTrace { get; } = new ArtifactType("step-trace");
+		public static ArtifactType StepTraces { get; } = new ArtifactType("step-trace");
 
 		/// <summary>
 		/// Test data generated by a machine executing a build step
