Horde: Upload exported.json and preprocessed scripts to artifact storage.

diff --git a/Engine/Source/Programs/Horde/Horde.Agent/Execution/JobExecutor.cs b/Engine/Source/Programs/Horde/Horde.Agent/Execution/JobExecutor.cs
--- a/Engine/Source/Programs/Horde/Horde.Agent/Execution/JobExecutor.cs
+++ b/Engine/Source/Programs/Horde/Horde.Agent/Execution/JobExecutor.cs
@@ -464,15 +464,20 @@ namespace Horde.Agent.Execution
 				arguments.AppendArgument(additionalArgument);
 			}
 
+			List<FileReference> additionalArtifactFiles = new List<FileReference>();
+			additionalArtifactFiles.Add(definitionFile);
+
 			FileReference? preprocessedScriptFile = null;
 			FileReference? preprocessedSchemaFile = null;
 			if (_preprocessScript)
 			{
 				preprocessedScriptFile = FileReference.Combine(workspaceDir, PreprocessedScript);
 				arguments.AppendArgument("-Preprocess=", preprocessedScriptFile.FullName);
+				additionalArtifactFiles.Add(preprocessedScriptFile);
 
 				preprocessedSchemaFile = FileReference.Combine(workspaceDir, PreprocessedSchema);
 				arguments.AppendArgument("-Schema=", preprocessedSchemaFile.FullName);
+				additionalArtifactFiles.Add(preprocessedSchemaFile);
 			}
 			if (sharedStorageDir != null && !_preprocessScript) // Do not precompile when preprocessing the script; other agents may have a different view of UAT
 			{
@@ -480,7 +485,7 @@ namespace Horde.Agent.Execution
 				arguments.Append($" CopyUAT -WithLauncher -TargetDir=\"{buildDir}\"");
 			}
 
-			int result = await ExecuteAutomationToolAsync(step, workspaceDir, sharedStorageDir, arguments.ToString(), useP4, logger, cancellationToken);
+			int result = await ExecuteAutomationToolAsync(step, workspaceDir, sharedStorageDir, arguments.ToString(), useP4, additionalArtifactFiles, logger, cancellationToken);
 			if (result != 0)
 			{
 				return false;
@@ -745,7 +750,7 @@ namespace Horde.Agent.Execution
 					arguments.AppendArgument("-SharedStorageDir=", sharedStorageDir.FullName);
 				}
 				
-				bool result = await ExecuteAutomationToolAsync(step, workspaceDir, sharedStorageDir, arguments.ToString(), useP4, logger, cancellationToken) == 0;
+				bool result = await ExecuteAutomationToolAsync(step, workspaceDir, sharedStorageDir, arguments.ToString(), useP4, null, logger, cancellationToken) == 0;
 				await UploadXgeMonitorFilesAsync(step, logger, cancellationToken);
 				return result;
 			}
@@ -847,7 +852,7 @@ namespace Horde.Agent.Execution
 			}
 
 			// Run UAT
-			if (await ExecuteAutomationToolAsync(step, workspaceDir, null, arguments, useP4, logger, cancellationToken) != 0)
+			if (await ExecuteAutomationToolAsync(step, workspaceDir, null, arguments, useP4, null, logger, cancellationToken) != 0)
 			{
 				return false;
 			}
@@ -1010,7 +1015,7 @@ namespace Horde.Agent.Execution
 			return true;
 		}
 
-		protected async Task<int> ExecuteAutomationToolAsync(BeginStepResponse step, DirectoryReference workspaceDir, DirectoryReference? sharedStorageDir, string? arguments, bool? useP4, ILogger logger, CancellationToken cancellationToken)
+		protected async Task<int> ExecuteAutomationToolAsync(BeginStepResponse step, DirectoryReference workspaceDir, DirectoryReference? sharedStorageDir, string? arguments, bool? useP4, IEnumerable<FileReference>? additionalArtifactFiles, ILogger logger, CancellationToken cancellationToken)
 		{
 			int result;
 			using IScope scope = GlobalTracer.Instance.BuildSpan("BuildGraph").StartActive();
@@ -1027,7 +1032,7 @@ namespace Horde.Agent.Execution
 
 			if (RuntimeInformation.IsOSPlatform(OSPlatform.Windows))
 			{
-				result = await ExecuteCommandAsync(step, workspaceDir, sharedStorageDir, Environment.GetEnvironmentVariable("COMSPEC") ?? "cmd.exe", $"/C \"\"{workspaceDir}\\Engine\\Build\\BatchFiles\\RunUAT.bat\" {arguments}\"", logger, cancellationToken);
+				result = await ExecuteCommandAsync(step, workspaceDir, sharedStorageDir, Environment.GetEnvironmentVariable("COMSPEC") ?? "cmd.exe", $"/C \"\"{workspaceDir}\\Engine\\Build\\BatchFiles\\RunUAT.bat\" {arguments}\"", additionalArtifactFiles, logger, cancellationToken);
 			}
 			else if (RuntimeInformation.IsOSPlatform(OSPlatform.Linux))
 			{
@@ -1038,11 +1043,11 @@ namespace Horde.Agent.Execution
 					args = $"\"{workspaceDir}/Engine/Build/BatchFiles/RunWineUAT.sh\" {arguments}";
 				}
 				
-				result = await ExecuteCommandAsync(step, workspaceDir, sharedStorageDir, "/bin/bash", args, logger, cancellationToken);
+				result = await ExecuteCommandAsync(step, workspaceDir, sharedStorageDir, "/bin/bash", args, additionalArtifactFiles, logger, cancellationToken);
 			}
 			else if (RuntimeInformation.IsOSPlatform(OSPlatform.OSX))
 			{
-				result = await ExecuteCommandAsync(step, workspaceDir, sharedStorageDir, "/bin/sh", $"\"{workspaceDir}/Engine/Build/BatchFiles/RunUAT.sh\" {arguments}", logger, cancellationToken);
+				result = await ExecuteCommandAsync(step, workspaceDir, sharedStorageDir, "/bin/sh", $"\"{workspaceDir}/Engine/Build/BatchFiles/RunUAT.sh\" {arguments}", additionalArtifactFiles, logger, cancellationToken);
 			}
 			else
 			{
@@ -1301,7 +1306,7 @@ namespace Horde.Agent.Execution
 			return dirs;
 		}
 
-		async Task<int> ExecuteCommandAsync(BeginStepResponse step, DirectoryReference workspaceDir, DirectoryReference? sharedStorageDir, string fileName, string arguments, ILogger jobLogger, CancellationToken cancellationToken)
+		async Task<int> ExecuteCommandAsync(BeginStepResponse step, DirectoryReference workspaceDir, DirectoryReference? sharedStorageDir, string fileName, string arguments, IEnumerable<FileReference>? additionalArtifactFiles, ILogger jobLogger, CancellationToken cancellationToken)
 		{
 			// Method for expanding environment variable properties related to this step
 			Dictionary<string, string> properties = new Dictionary<string, string>(StringComparer.OrdinalIgnoreCase);
@@ -1477,9 +1482,18 @@ namespace Horde.Agent.Execution
 				await CreateArtifactsAsync(step.StepId, JobArtifactType.TestData, workspaceDir, testDataFiles, jobLogger, cancellationToken);
 			}
 
+			List<FileReference> artifactFiles = new List<FileReference>();
+			if (additionalArtifactFiles != null)
+			{
+				artifactFiles.AddRange(additionalArtifactFiles);
+			}
 			if (DirectoryReference.Exists(logDir))
 			{
-				List<FileReference> artifactFiles = DirectoryReference.EnumerateFiles(logDir, "*", SearchOption.AllDirectories).ToList();
+				artifactFiles.AddRange(DirectoryReference.EnumerateFiles(logDir, "*", SearchOption.AllDirectories));
+			}
+
+			if (artifactFiles.Count > 0)
+			{
 				if (_jobOptions.UseNewTempStorage ?? false)
 				{
 					await CreateArtifactAsync(step.StepId, JobArtifactType.Saved, workspaceDir, artifactFiles, jobLogger, cancellationToken);
