Horde: Revert change to rename artifact types; dashboard relies on names.

diff --git a/Engine/Source/Programs/Horde/Horde.Server.Tests/ArtifactTests.cs b/Engine/Source/Programs/Horde/Horde.Server.Tests/ArtifactTests.cs
--- a/Engine/Source/Programs/Horde/Horde.Server.Tests/ArtifactTests.cs
+++ b/Engine/Source/Programs/Horde/Horde.Server.Tests/ArtifactTests.cs
@@ -24,7 +24,7 @@ namespace Horde.Server.Tests
 		public async Task CreateArtifactAsync()
 		{
 			IArtifactCollection artifactCollection = ServiceProvider.GetRequiredService<IArtifactCollection>();
-			IArtifact artifact = await artifactCollection.AddAsync(new ArtifactId(BinaryIdUtils.CreateNew()), ArtifactType.StepTempStorage, new string[] { "test1", "test2" }, Namespace.Artifacts, new RefName("test"), null, AclScopeName.Root);
+			IArtifact artifact = await artifactCollection.AddAsync(new ArtifactId(BinaryIdUtils.CreateNew()), ArtifactType.StepOutput, new string[] { "test1", "test2" }, Namespace.Artifacts, new RefName("test"), null, AclScopeName.Root);
 
 			{
 				List<IArtifact> artifacts = await artifactCollection.FindAsync(keys: new[] { "test1" }).ToListAsync();
@@ -53,7 +53,7 @@ namespace Horde.Server.Tests
 			await expirationService.StartAsync(CancellationToken.None);
 			
 			IArtifactCollection artifactCollection = ServiceProvider.GetRequiredService<IArtifactCollection>();
-			IArtifact artifact = await artifactCollection.AddAsync(new ArtifactId(BinaryIdUtils.CreateNew()), ArtifactType.StepTempStorage, new string[] { "test1", "test2" }, Namespace.Artifacts, new RefName("test"), clock.UtcNow + TimeSpan.FromHours(1.0), AclScopeName.Root);
+			IArtifact artifact = await artifactCollection.AddAsync(new ArtifactId(BinaryIdUtils.CreateNew()), ArtifactType.StepOutput, new string[] { "test1", "test2" }, Namespace.Artifacts, new RefName("test"), clock.UtcNow + TimeSpan.FromHours(1.0), AclScopeName.Root);
 
 			{
 				List<IArtifact> artifacts = await artifactCollection.FindAsync(keys: new[] { "test1" }).ToListAsync();
diff --git a/Engine/Source/Programs/Horde/Horde.Server/Jobs/JobRpcService.cs b/Engine/Source/Programs/Horde/Horde.Server/Jobs/JobRpcService.cs
--- a/Engine/Source/Programs/Horde/Horde.Server/Jobs/JobRpcService.cs
+++ b/Engine/Source/Programs/Horde/Horde.Server/Jobs/JobRpcService.cs
@@ -77,9 +77,9 @@ namespace Horde.Server.Jobs
 
 			ArtifactType type = request.Type switch
 			{
-				JobArtifactType.TempStorage => ArtifactType.StepTempStorage,
+				JobArtifactType.TempStorage => ArtifactType.StepOutput,
 				JobArtifactType.Saved => ArtifactType.StepSaved,
-				JobArtifactType.Trace => ArtifactType.StepTraces,
+				JobArtifactType.Trace => ArtifactType.StepTrace,
 				JobArtifactType.TestData => ArtifactType.StepTestData,
 				_ => throw new StructuredRpcException(StatusCode.InvalidArgument, "Invalid artifact type")
 			};
diff --git a/Engine/Source/Programs/Shared/EpicGames.Horde/Api/ArtifactType.cs b/Engine/Source/Programs/Shared/EpicGames.Horde/Api/ArtifactType.cs
--- a/Engine/Source/Programs/Shared/EpicGames.Horde/Api/ArtifactType.cs
+++ b/Engine/Source/Programs/Shared/EpicGames.Horde/Api/ArtifactType.cs
@@ -22,7 +22,7 @@ namespace EpicGames.Horde.Api
 		/// <summary>
 		/// Output from a build step
 		/// </summary>
-		public static ArtifactType StepTempStorage { get; } = new ArtifactType("step-output");
+		public static ArtifactType StepOutput { get; } = new ArtifactType("step-output");
 
 		/// <summary>
 		/// Captured state from the machine executing a build step
@@ -32,7 +32,7 @@ namespace EpicGames.Horde.Api
 		/// <summary>
 		/// Traces from the machine executing a build step
 		/// </summary>
-		public static ArtifactType StepTraces { get; } = new ArtifactType("step-trace");
+		public static ArtifactType StepTrace { get; } = new ArtifactType("step-trace");
 
 		/// <summary>
 		/// Test data generated by a machine executing a build step
