Horde: Add commands to execute individual compute jobs through command line tool.

diff --git a/Engine/Source/Programs/Horde/Horde.Agent/Commands/Compute/ComputeCommand.cs b/Engine/Source/Programs/Horde/Horde.Agent/Commands/Compute/ComputeCommand.cs
--- a/Engine/Source/Programs/Horde/Horde.Agent/Commands/Compute/ComputeCommand.cs
+++ b/Engine/Source/Programs/Horde/Horde.Agent/Commands/Compute/ComputeCommand.cs
@@ -92,18 +92,18 @@ namespace Horde.Agent.Commands.Compute
 			}
 			else
 			{
-				return new ServerComputeClient(CreateHttpClient, logger);
+				return new ServerComputeClient(CreateHttpClientAsync, logger);
 			}
 		}
 
-		HttpClient CreateHttpClient()
+		Task<HttpClient> CreateHttpClientAsync(CancellationToken cancellationToken)
 		{
 			ServerProfile profile = _settings.Value.GetCurrentServerProfile();
 
 			HttpClient client = _httpClientFactory.CreateClient();
 			client.BaseAddress = profile.Url;
 			client.DefaultRequestHeaders.Authorization = new AuthenticationHeaderValue("Bearer", profile.Token);
-			return client;
+			return Task.FromResult(client);
 		}
 	}
 
diff --git a/Engine/Source/Programs/Horde/Horde/Commands/Artifacts/ArtifactFind.cs b/Engine/Source/Programs/Horde/Horde/Commands/Artifacts/ArtifactFind.cs
--- a/Engine/Source/Programs/Horde/Horde/Commands/Artifacts/ArtifactFind.cs
+++ b/Engine/Source/Programs/Horde/Horde/Commands/Artifacts/ArtifactFind.cs
@@ -15,7 +15,7 @@ namespace Horde.Commands.Bundles
 
 		public override async Task<int> ExecuteAsync(ILogger logger)
 		{
-			using HordeHttpClient horde = await Settings.GetHttpCientAsync(logger);
+			using HordeHttpClient horde = await Settings.GetHttpClientAsync(logger);
 
 			List<GetArtifactResponse> artifacts = await horde.FindArtifactsAsync(keys: Keys);
 			foreach (GetArtifactResponse artifact in artifacts)
diff --git a/Engine/Source/Programs/Horde/Horde/Commands/Compute/ComputeRun.cs b/Engine/Source/Programs/Horde/Horde/Commands/Compute/ComputeRun.cs
new file mode 100644
--- /dev/null
+++ b/Engine/Source/Programs/Horde/Horde/Commands/Compute/ComputeRun.cs
@@ -0,0 +1,141 @@
+// Copyright Epic Games, Inc. All Rights Reserved.
+
+using System.Net.Http.Headers;
+using System.Reflection;
+using System.Text.Json;
+using EpicGames.Core;
+using EpicGames.Horde.Common;
+using EpicGames.Horde.Compute;
+using EpicGames.Horde.Compute.Clients;
+using EpicGames.Horde.Storage.Backends;
+using EpicGames.Horde.Storage.Bundles;
+using EpicGames.Horde.Storage.Nodes;
+using Microsoft.Extensions.Logging;
+
+namespace Horde.Commands.Compute
+{
+	/// <summary>
+	/// Executes a compute command using a task definition and sandbox taken from the local machine
+	/// </summary>
+	[Command("compute", "run", "Executes a command through the compute API")]
+	class ComputeRun : Command
+	{
+		class JsonComputeTask
+		{
+			public string Executable { get; set; } = null!;
+			public List<string> Arguments { get; set; } = new List<string>();
+			public string WorkingDir { get; set; } = String.Empty;
+			public Dictionary<string, string?> EnvVars { get; set; } = new Dictionary<string, string?>();
+			public List<string> OutputPaths { get; set; } = new List<string>();
+		}
+
+		[CommandLine("-Cluster")]
+		public string ClusterId { get; set; } = "default";
+
+		[CommandLine("-Requirements=", Description = "Match the agent to run on")]
+		public string? Requirements { get; set; }
+
+		[CommandLine("-Local")]
+		public bool Local { get; set; }
+
+		[CommandLine("-Loopback")]
+		public bool Loopback { get; set; }
+
+		[CommandLine("-InProc")]
+		public bool InProc { get; set; }
+
+		[CommandLine("-Sandbox=")]
+		public DirectoryReference SandboxDir { get; set; } = DirectoryReference.Combine(DirectoryReference.GetSpecialFolder(Environment.SpecialFolder.LocalApplicationData)!, "Horde", "Sandbox");
+
+		[CommandLine("-Task=", Required = true)]
+		FileReference TaskFile { get; set; } = null!;
+
+		/// <inheritdoc/>
+		public override async Task<int> ExecuteAsync(ILogger logger)
+		{
+			await using IComputeClient client = CreateClient(logger);
+
+			Requirements? requirements = null;
+			if (Requirements != null)
+			{
+				requirements = new Requirements(Condition.Parse(Requirements));
+			}
+
+			await using IComputeLease? lease = await client.TryAssignWorkerAsync(new ClusterId(ClusterId), requirements, CancellationToken.None);
+			if (lease == null)
+			{
+				throw new Exception("Unable to create lease");
+			}
+
+			bool result = await HandleRequestAsync(lease, logger, CancellationToken.None);
+			return result ? 0 : 1;
+		}
+
+		IComputeClient CreateClient(ILogger logger)
+		{
+			if (Local)
+			{
+				return new LocalComputeClient(2000, SandboxDir, InProc, logger);
+			}
+			else if (Loopback)
+			{
+				return new AgentComputeClient(Assembly.GetExecutingAssembly().Location, 2000, logger);
+			}
+			else
+			{
+				return new ServerComputeClient(async ctx => await CreateHttpClientAsync(logger, ctx), logger);
+			}
+		}
+
+		static async Task<HttpClient> CreateHttpClientAsync(ILogger logger, CancellationToken cancellationToken)
+		{
+			HttpClient client = new HttpClient();
+			client.BaseAddress = await Settings.GetServerAsync(cancellationToken);
+			client.DefaultRequestHeaders.Authorization = new AuthenticationHeaderValue("Bearer", await Settings.GetAccessTokenAsync(logger, cancellationToken));
+			return client;
+		}
+
+		/// <inheritdoc/>
+		async Task<bool> HandleRequestAsync(IComputeLease lease, ILogger logger, CancellationToken cancellationToken)
+		{
+			const int ControlChannelId = 0;
+
+			// Read the task definition
+			byte[] data = await FileReference.ReadAllBytesAsync(TaskFile, cancellationToken);
+			JsonComputeTask jsonComputeTask = JsonSerializer.Deserialize<JsonComputeTask>(data, new JsonSerializerOptions { AllowTrailingCommas = true, PropertyNameCaseInsensitive = true, PropertyNamingPolicy = JsonNamingPolicy.CamelCase })!;
+
+			// Create a sandbox from the data to be uploaded
+			MemoryStorageClient storage = new MemoryStorageClient();
+			BundleNodeLocator sandbox = await CreateSandboxAsync(TaskFile, storage, cancellationToken);
+
+			// Open a socket and upload the sandbox
+			using (AgentMessageChannel channel = lease.Socket.CreateAgentMessageChannel(ControlChannelId, 4 * 1024 * 1024, logger))
+			{
+				await channel.WaitForAttachAsync(cancellationToken);
+				await channel.UploadFilesAsync("", sandbox, storage, cancellationToken);
+
+				await using (AgentManagedProcess process = await channel.ExecuteAsync(jsonComputeTask.Executable, jsonComputeTask.Arguments, jsonComputeTask.WorkingDir, jsonComputeTask.EnvVars, ExecuteProcessFlags.None, cancellationToken))
+				{
+					string? line;
+					while ((line = await process.ReadLineAsync(cancellationToken)) != null)
+					{
+						logger.LogInformation("Child Process: {Text}", line);
+					}
+				}
+			}
+
+			return true;
+		}
+
+		static async Task<BundleNodeLocator> CreateSandboxAsync(FileReference taskFile, BundleStorageClient storage, CancellationToken cancellationToken)
+		{
+			await using BundleWriter writer = storage.CreateWriter();
+
+			DirectoryNode sandbox = new DirectoryNode();
+			await sandbox.CopyFromDirectoryAsync(taskFile.Directory.ToDirectoryInfo(), new ChunkingOptions(), writer, null, cancellationToken);
+
+			BundleNodeHandle handle = await writer.FlushAsync(sandbox, cancellationToken);
+			return handle.GetLocator();
+		}
+	}
+}
diff --git a/Engine/Source/Programs/Horde/Horde/Commands/Compute/ComputeWorker.cs b/Engine/Source/Programs/Horde/Horde/Commands/Compute/ComputeWorker.cs
new file mode 100644
--- /dev/null
+++ b/Engine/Source/Programs/Horde/Horde/Commands/Compute/ComputeWorker.cs
@@ -0,0 +1,56 @@
+// Copyright Epic Games, Inc. All Rights Reserved.
+
+using System.Net;
+using System.Net.Sockets;
+using EpicGames.Core;
+using EpicGames.Horde.Compute;
+using EpicGames.Horde.Compute.Transports;
+using Microsoft.Extensions.Caching.Memory;
+using Microsoft.Extensions.Logging;
+
+namespace Horde.Commands.Compute
+{
+	/// <summary>
+	/// Helper command for hosting a local compute worker in a separate process
+	/// </summary>
+	[Command("compute", "worker", "Runs the agent as a local compute host, accepting incoming connections on the loopback adapter with a given port")]
+	class ComputeWorkerCommand : Command
+	{
+		readonly IMemoryCache _memoryCache;
+
+		[CommandLine("-Port=")]
+		int Port { get; set; } = 2000;
+
+		public ComputeWorkerCommand(IMemoryCache memoryCache)
+		{
+			_memoryCache = memoryCache;
+		}
+
+		public override async Task<int> ExecuteAsync(ILogger logger)
+		{
+			logger.LogInformation("** WORKER **");
+
+			using Socket tcpSocket = new Socket(SocketType.Stream, ProtocolType.IP);
+			await tcpSocket.ConnectAsync(IPAddress.Loopback, Port);
+
+			await using (RemoteComputeSocket socket = new RemoteComputeSocket(new TcpTransport(tcpSocket), ComputeSocketEndpoint.Remote, logger))
+			{
+				logger.LogInformation("Running worker...");
+				await RunWorkerAsync(socket, _memoryCache, logger, CancellationToken.None);
+				logger.LogInformation("Worker complete");
+				await socket.CloseAsync(CancellationToken.None);
+			}
+
+			logger.LogInformation("Stopping");
+			return 0;
+		}
+
+		public static async Task RunWorkerAsync(ComputeSocket socket, IMemoryCache memoryCache, ILogger logger, CancellationToken cancellationToken)
+		{
+			DirectoryReference sandboxDir = DirectoryReference.Combine(DirectoryReference.GetSpecialFolder(Environment.SpecialFolder.LocalApplicationData)!, "Horde", "Sandbox");
+
+			AgentMessageHandler worker = new AgentMessageHandler(sandboxDir, memoryCache, null, false, null, logger);
+			await worker.RunAsync(socket, cancellationToken);
+		}
+	}
+}
diff --git a/Engine/Source/Programs/Horde/Horde/Settings.cs b/Engine/Source/Programs/Horde/Horde/Settings.cs
--- a/Engine/Source/Programs/Horde/Horde/Settings.cs
+++ b/Engine/Source/Programs/Horde/Horde/Settings.cs
@@ -136,7 +136,7 @@ namespace Horde
 			return result.AccessToken;
 		}
 
-		public static async Task<HordeHttpClient> GetHttpCientAsync(ILogger logger, CancellationToken cancellationToken = default)
+		public static async Task<HordeHttpClient> GetHttpClientAsync(ILogger logger, CancellationToken cancellationToken = default)
 		{
 			Uri? server = await GetServerAsync(cancellationToken);
 			if (server == null)
diff --git a/Engine/Source/Programs/Shared/EpicGames.Horde/Compute/Clients/ServerComputeClient.cs b/Engine/Source/Programs/Shared/EpicGames.Horde/Compute/Clients/ServerComputeClient.cs
--- a/Engine/Source/Programs/Shared/EpicGames.Horde/Compute/Clients/ServerComputeClient.cs
+++ b/Engine/Source/Programs/Shared/EpicGames.Horde/Compute/Clients/ServerComputeClient.cs
@@ -69,7 +69,7 @@ namespace EpicGames.Horde.Compute.Clients
 		}
 
 		readonly HttpClient? _defaultHttpClient;
-		readonly Func<HttpClient> _createHttpClient;
+		readonly Func<CancellationToken, Task<HttpClient>> _createHttpClientAsync;
 		readonly CancellationTokenSource _cancellationSource = new CancellationTokenSource();
 		readonly ILogger _logger;
 
@@ -91,18 +91,18 @@ namespace EpicGames.Horde.Compute.Clients
 			_defaultHttpClient.DefaultRequestHeaders.Authorization = authHeader;
 #pragma warning restore CA2000 // Dispose objects before losing scope
 
-			_createHttpClient = GetDefaultHttpClient;
+			_createHttpClientAsync = GetDefaultHttpClientAsync;
 			_logger = logger;
 		}
 
 		/// <summary>
 		/// Constructor
 		/// </summary>
-		/// <param name="createHttpClient">Creates an HTTP client with the correct base address for the server</param>
+		/// <param name="createHttpClientAsync">Creates an HTTP client with the correct base address for the server</param>
 		/// <param name="logger">Logger for diagnostic messages</param>
-		public ServerComputeClient(Func<HttpClient> createHttpClient, ILogger logger)
+		public ServerComputeClient(Func<CancellationToken, Task<HttpClient>> createHttpClientAsync, ILogger logger)
 		{
-			_createHttpClient = createHttpClient;
+			_createHttpClientAsync = createHttpClientAsync;
 			_logger = logger;
 		}
 
@@ -110,7 +110,7 @@ namespace EpicGames.Horde.Compute.Clients
 		/// Gets the default http client
 		/// </summary>
 		/// <returns></returns>
-		HttpClient GetDefaultHttpClient() => _defaultHttpClient!;
+		Task<HttpClient> GetDefaultHttpClientAsync(CancellationToken cancellationToken) => Task.FromResult(_defaultHttpClient!);
 
 		/// <inheritdoc/>
 		public ValueTask DisposeAsync()
@@ -144,7 +144,7 @@ namespace EpicGames.Horde.Compute.Clients
 			_logger.LogDebug("Requesting compute resource");
 
 			// Assign a compute worker
-			HttpClient client = _createHttpClient();
+			HttpClient client = await _createHttpClientAsync(cancellationToken);
 
 			AssignComputeRequest request = new AssignComputeRequest();
 			request.Requirements = requirements;
