Horde: Copy preprocessed files to/from temp storage.

diff --git a/Engine/Source/Programs/Horde/Horde.Agent/Execution/JobExecutor.cs b/Engine/Source/Programs/Horde/Horde.Agent/Execution/JobExecutor.cs
--- a/Engine/Source/Programs/Horde/Horde.Agent/Execution/JobExecutor.cs
+++ b/Engine/Source/Programs/Horde/Horde.Agent/Execution/JobExecutor.cs
@@ -29,6 +29,7 @@ using HordeCommon;
 using HordeCommon.Rpc;
 using HordeCommon.Rpc.Messages;
 using HordeCommon.Rpc.Tasks;
+using Microsoft.CodeAnalysis.CSharp.Syntax;
 using Microsoft.Extensions.Caching.Memory;
 using Microsoft.Extensions.Logging;
 using OpenTracing;
@@ -447,6 +448,8 @@ namespace Horde.Agent.Execution
 
 		protected abstract Task<bool> ExecuteAsync(BeginStepResponse step, ILogger logger, CancellationToken cancellationToken);
 
+		const string BuildGraphTempStorageDir = "BuildGraph";
+
 		protected virtual async Task<bool> SetupAsync(BeginStepResponse step, DirectoryReference workspaceDir, DirectoryReference? sharedStorageDir, bool? useP4, ILogger logger, CancellationToken cancellationToken)
 		{
 			FileReference definitionFile = FileReference.Combine(workspaceDir, "Engine", "Saved", "Horde", "Exported.json");
@@ -464,20 +467,15 @@ namespace Horde.Agent.Execution
 				arguments.AppendArgument(additionalArgument);
 			}
 
-			List<FileReference> additionalArtifactFiles = new List<FileReference>();
-			additionalArtifactFiles.Add(definitionFile);
-
 			FileReference? preprocessedScriptFile = null;
 			FileReference? preprocessedSchemaFile = null;
 			if (_preprocessScript)
 			{
 				preprocessedScriptFile = FileReference.Combine(workspaceDir, PreprocessedScript);
 				arguments.AppendArgument("-Preprocess=", preprocessedScriptFile.FullName);
-				additionalArtifactFiles.Add(preprocessedScriptFile);
 
 				preprocessedSchemaFile = FileReference.Combine(workspaceDir, PreprocessedSchema);
 				arguments.AppendArgument("-Schema=", preprocessedSchemaFile.FullName);
-				additionalArtifactFiles.Add(preprocessedSchemaFile);
 			}
 			if (sharedStorageDir != null && !_preprocessScript) // Do not precompile when preprocessing the script; other agents may have a different view of UAT
 			{
@@ -485,15 +483,66 @@ namespace Horde.Agent.Execution
 				arguments.Append($" CopyUAT -WithLauncher -TargetDir=\"{buildDir}\"");
 			}
 
-			int result = await ExecuteAutomationToolAsync(step, workspaceDir, sharedStorageDir, arguments.ToString(), useP4, additionalArtifactFiles, logger, cancellationToken);
+			int result = await ExecuteAutomationToolAsync(step, workspaceDir, sharedStorageDir, arguments.ToString(), useP4, logger, cancellationToken);
 			if (result != 0)
 			{
 				return false;
 			}
-			
-			await ArtifactUploader.UploadAsync(RpcConnection, _jobId, _batchId, step.StepId, definitionFile.GetFileName(), definitionFile, logger, cancellationToken);
-			await StorePreprocessedFileAsync(preprocessedScriptFile, step.StepId, sharedStorageDir, logger, cancellationToken);
-			await StorePreprocessedFileAsync(preprocessedSchemaFile, step.StepId, sharedStorageDir, logger, cancellationToken);
+
+			if (_jobOptions.UseNewTempStorage ?? false)
+			{
+				List<FileReference> buildGraphFiles = new List<FileReference>();
+				buildGraphFiles.Add(definitionFile);
+				if (preprocessedScriptFile != null)
+				{
+					buildGraphFiles.Add(preprocessedScriptFile);
+				}
+				if (preprocessedSchemaFile != null)
+				{
+					buildGraphFiles.Add(preprocessedSchemaFile);
+				}
+
+				using (GlobalTracer.Instance.BuildSpan("TempStorage").WithTag("resource", "Write").StartActive())
+				{
+					IStorageClient storage = _storageFactory.CreateStorageClient(_session, _namespaceId, _token);
+
+					Stopwatch timer = Stopwatch.StartNew();
+
+					RefName refName = TempStorage.GetRefNameForNode(_storagePrefix, step.Name);
+					await using (IStorageWriter treeWriter = storage.CreateWriter(refName))
+					{
+						DirectoryNode buildGraphNode = new DirectoryNode();
+						await buildGraphNode.AddFilesAsync(workspaceDir, buildGraphFiles, new ChunkingOptions(), treeWriter, null, cancellationToken);
+						NodeRef<DirectoryNode> outputNodeRef = await treeWriter.WriteNodeAsync(buildGraphNode, cancellationToken);
+
+						DirectoryNode rootNode = new DirectoryNode();
+						rootNode.AddDirectory(new DirectoryEntry(BuildGraphTempStorageDir, buildGraphNode.Length, outputNodeRef));
+
+						await storage.WriteRefAsync(refName, rootNode, refOptions: new RefOptions(), cancellationToken: cancellationToken);
+					}
+
+					logger.LogInformation("Upload took {Time:n1}s", timer.Elapsed.TotalSeconds);
+
+					// Create the artifact
+					using IRpcClientRef<JobRpc.JobRpcClient> jobRpc = await RpcConnection.GetClientRefAsync<JobRpc.JobRpcClient>(cancellationToken);
+
+					CreateJobArtifactRequest artifactRequest = new CreateJobArtifactRequest();
+					artifactRequest.JobId = _jobId;
+					artifactRequest.StepId = step.StepId;
+					artifactRequest.Type = JobArtifactType.TempStorage;
+					artifactRequest.NamespaceId = _namespaceId.ToString();
+					artifactRequest.RefName = refName.ToString();
+
+					CreateJobArtifactResponse artifact = await jobRpc.Client.CreateArtifactAsync(artifactRequest, cancellationToken: cancellationToken);
+					logger.LogInformation("Created output artifact {ArtifactId} with ref {RefName} in namespace {Namespace}", artifact.Id, artifact.RefName, artifact.NamespaceId);
+				}
+			}
+			else
+			{
+				await ArtifactUploader.UploadAsync(RpcConnection, _jobId, _batchId, step.StepId, definitionFile.GetFileName(), definitionFile, logger, cancellationToken);
+				await StorePreprocessedFileAsync(preprocessedScriptFile, step.StepId, sharedStorageDir, logger, cancellationToken);
+				await StorePreprocessedFileAsync(preprocessedSchemaFile, step.StepId, sharedStorageDir, logger, cancellationToken);
+			}
 
 			JsonSerializerOptions options = new JsonSerializerOptions();
 			options.PropertyNameCaseInsensitive = true;
@@ -710,12 +759,29 @@ namespace Horde.Agent.Execution
 			if (_preprocessScript)
 			{
 				FileReference localPreprocessedScript = FileReference.Combine(workspaceDir, PreprocessedScript);
-				FetchPreprocessedFile(localPreprocessedScript, sharedStorageDir, logger);
 				arguments.AppendArgument(ScriptArgumentPrefix, localPreprocessedScript.FullName);
 
 				FileReference localPreprocessedSchema = FileReference.Combine(workspaceDir, PreprocessedSchema);
-				FetchPreprocessedFile(localPreprocessedSchema, sharedStorageDir, logger);
 				arguments.AppendArgument("-ImportSchema=", localPreprocessedSchema.FullName);
+
+				if (_jobOptions.UseNewTempStorage ?? false)
+				{
+					IStorageClient storage = _storageFactory.CreateStorageClient(_session, _namespaceId, _token);
+
+					RefName refName = TempStorage.GetRefNameForNode(_storagePrefix, step.Name);
+
+					DirectoryNode node = await storage.ReadRefAsync<DirectoryNode>(refName, cancellationToken: cancellationToken);
+					DirectoryNode? buildGraphDir = await node.FindDirectoryAsync(BuildGraphTempStorageDir, cancellationToken);
+					if (buildGraphDir != null)
+					{
+						await buildGraphDir.CopyToDirectoryAsync(new DirectoryInfo(workspaceDir.FullName), logger, cancellationToken);
+					}
+				}
+				else
+				{
+					FetchPreprocessedFile(localPreprocessedScript, sharedStorageDir, logger);
+					FetchPreprocessedFile(localPreprocessedSchema, sharedStorageDir, logger);
+				}
 			}
 			else if (_scriptFileName != null)
 			{
@@ -750,7 +816,7 @@ namespace Horde.Agent.Execution
 					arguments.AppendArgument("-SharedStorageDir=", sharedStorageDir.FullName);
 				}
 				
-				bool result = await ExecuteAutomationToolAsync(step, workspaceDir, sharedStorageDir, arguments.ToString(), useP4, null, logger, cancellationToken) == 0;
+				bool result = await ExecuteAutomationToolAsync(step, workspaceDir, sharedStorageDir, arguments.ToString(), useP4, logger, cancellationToken) == 0;
 				await UploadXgeMonitorFilesAsync(step, logger, cancellationToken);
 				return result;
 			}
@@ -795,7 +861,6 @@ namespace Horde.Agent.Execution
 		{
 			DirectoryReference manifestDir = DirectoryReference.Combine(workspaceDir, "Engine", "Saved", "BuildGraph");
 
-			using MemoryCache cache = new MemoryCache(new MemoryCacheOptions { });
 			IStorageClient storage = _storageFactory.CreateStorageClient(_session, _namespaceId, _token);
 
 			// Create the mapping of tag names to file sets
@@ -852,7 +917,7 @@ namespace Horde.Agent.Execution
 			}
 
 			// Run UAT
-			if (await ExecuteAutomationToolAsync(step, workspaceDir, null, arguments, useP4, null, logger, cancellationToken) != 0)
+			if (await ExecuteAutomationToolAsync(step, workspaceDir, null, arguments, useP4, logger, cancellationToken) != 0)
 			{
 				return false;
 			}
@@ -1015,7 +1080,7 @@ namespace Horde.Agent.Execution
 			return true;
 		}
 
-		protected async Task<int> ExecuteAutomationToolAsync(BeginStepResponse step, DirectoryReference workspaceDir, DirectoryReference? sharedStorageDir, string? arguments, bool? useP4, IEnumerable<FileReference>? additionalArtifactFiles, ILogger logger, CancellationToken cancellationToken)
+		protected async Task<int> ExecuteAutomationToolAsync(BeginStepResponse step, DirectoryReference workspaceDir, DirectoryReference? sharedStorageDir, string? arguments, bool? useP4, ILogger logger, CancellationToken cancellationToken)
 		{
 			int result;
 			using IScope scope = GlobalTracer.Instance.BuildSpan("BuildGraph").StartActive();
@@ -1032,7 +1097,7 @@ namespace Horde.Agent.Execution
 
 			if (RuntimeInformation.IsOSPlatform(OSPlatform.Windows))
 			{
-				result = await ExecuteCommandAsync(step, workspaceDir, sharedStorageDir, Environment.GetEnvironmentVariable("COMSPEC") ?? "cmd.exe", $"/C \"\"{workspaceDir}\\Engine\\Build\\BatchFiles\\RunUAT.bat\" {arguments}\"", additionalArtifactFiles, logger, cancellationToken);
+				result = await ExecuteCommandAsync(step, workspaceDir, sharedStorageDir, Environment.GetEnvironmentVariable("COMSPEC") ?? "cmd.exe", $"/C \"\"{workspaceDir}\\Engine\\Build\\BatchFiles\\RunUAT.bat\" {arguments}\"", logger, cancellationToken);
 			}
 			else if (RuntimeInformation.IsOSPlatform(OSPlatform.Linux))
 			{
@@ -1043,11 +1108,11 @@ namespace Horde.Agent.Execution
 					args = $"\"{workspaceDir}/Engine/Build/BatchFiles/RunWineUAT.sh\" {arguments}";
 				}
 				
-				result = await ExecuteCommandAsync(step, workspaceDir, sharedStorageDir, "/bin/bash", args, additionalArtifactFiles, logger, cancellationToken);
+				result = await ExecuteCommandAsync(step, workspaceDir, sharedStorageDir, "/bin/bash", args, logger, cancellationToken);
 			}
 			else if (RuntimeInformation.IsOSPlatform(OSPlatform.OSX))
 			{
-				result = await ExecuteCommandAsync(step, workspaceDir, sharedStorageDir, "/bin/sh", $"\"{workspaceDir}/Engine/Build/BatchFiles/RunUAT.sh\" {arguments}", additionalArtifactFiles, logger, cancellationToken);
+				result = await ExecuteCommandAsync(step, workspaceDir, sharedStorageDir, "/bin/sh", $"\"{workspaceDir}/Engine/Build/BatchFiles/RunUAT.sh\" {arguments}", logger, cancellationToken);
 			}
 			else
 			{
@@ -1306,7 +1371,7 @@ namespace Horde.Agent.Execution
 			return dirs;
 		}
 
-		async Task<int> ExecuteCommandAsync(BeginStepResponse step, DirectoryReference workspaceDir, DirectoryReference? sharedStorageDir, string fileName, string arguments, IEnumerable<FileReference>? additionalArtifactFiles, ILogger jobLogger, CancellationToken cancellationToken)
+		async Task<int> ExecuteCommandAsync(BeginStepResponse step, DirectoryReference workspaceDir, DirectoryReference? sharedStorageDir, string fileName, string arguments, ILogger jobLogger, CancellationToken cancellationToken)
 		{
 			// Method for expanding environment variable properties related to this step
 			Dictionary<string, string> properties = new Dictionary<string, string>(StringComparer.OrdinalIgnoreCase);
@@ -1482,18 +1547,9 @@ namespace Horde.Agent.Execution
 				await CreateArtifactsAsync(step.StepId, JobArtifactType.TestData, workspaceDir, testDataFiles, jobLogger, cancellationToken);
 			}
 
-			List<FileReference> artifactFiles = new List<FileReference>();
-			if (additionalArtifactFiles != null)
-			{
-				artifactFiles.AddRange(additionalArtifactFiles);
-			}
 			if (DirectoryReference.Exists(logDir))
 			{
-				artifactFiles.AddRange(DirectoryReference.EnumerateFiles(logDir, "*", SearchOption.AllDirectories));
-			}
-
-			if (artifactFiles.Count > 0)
-			{
+				List<FileReference> artifactFiles = DirectoryReference.EnumerateFiles(logDir, "*", SearchOption.AllDirectories).ToList();
 				if (_jobOptions.UseNewTempStorage ?? false)
 				{
 					await CreateArtifactAsync(step.StepId, JobArtifactType.Saved, workspaceDir, artifactFiles, jobLogger, cancellationToken);
