Horde: Fix copying preprocessed files from temp storage.

diff --git a/Engine/Source/Programs/Horde/Horde.Agent/Execution/JobExecutor.cs b/Engine/Source/Programs/Horde/Horde.Agent/Execution/JobExecutor.cs
--- a/Engine/Source/Programs/Horde/Horde.Agent/Execution/JobExecutor.cs
+++ b/Engine/Source/Programs/Horde/Horde.Agent/Execution/JobExecutor.cs
@@ -5,8 +5,6 @@ using System.Collections.Generic;
 using System.Diagnostics;
 using System.IO;
 using System.Linq;
-using System.Net.Http;
-using System.Net.Http.Headers;
 using System.Runtime.InteropServices;
 using System.Security.Principal;
 using System.Text;
@@ -16,21 +14,17 @@ using System.Threading;
 using System.Threading.Tasks;
 using EpicGames.Core;
 using EpicGames.Horde.Storage;
-using EpicGames.Horde.Storage.Backends;
 using EpicGames.Horde.Storage.Nodes;
 using Grpc.Core;
 using Horde.Agent.Parser;
 using Horde.Agent.Services;
 using Horde.Agent.Utility;
-using Horde.Common;
 using Horde.Common.Rpc;
 using Horde.Storage.Utility;
 using HordeCommon;
 using HordeCommon.Rpc;
 using HordeCommon.Rpc.Messages;
 using HordeCommon.Rpc.Tasks;
-using Microsoft.CodeAnalysis.CSharp.Syntax;
-using Microsoft.Extensions.Caching.Memory;
 using Microsoft.Extensions.Logging;
 using OpenTracing;
 using OpenTracing.Util;
@@ -448,6 +442,7 @@ namespace Horde.Agent.Execution
 
 		protected abstract Task<bool> ExecuteAsync(BeginStepResponse step, ILogger logger, CancellationToken cancellationToken);
 
+		const string SetupStepName = "setup";
 		const string BuildGraphTempStorageDir = "BuildGraph";
 
 		protected virtual async Task<bool> SetupAsync(BeginStepResponse step, DirectoryReference workspaceDir, DirectoryReference? sharedStorageDir, bool? useP4, ILogger logger, CancellationToken cancellationToken)
@@ -508,7 +503,7 @@ namespace Horde.Agent.Execution
 
 					Stopwatch timer = Stopwatch.StartNew();
 
-					RefName refName = TempStorage.GetRefNameForNode(_storagePrefix, step.Name);
+					RefName refName = TempStorage.GetRefNameForNode(_storagePrefix, SetupStepName);
 					await using (IStorageWriter treeWriter = storage.CreateWriter(refName))
 					{
 						DirectoryNode buildGraphNode = new DirectoryNode();
@@ -768,7 +763,7 @@ namespace Horde.Agent.Execution
 				{
 					IStorageClient storage = _storageFactory.CreateStorageClient(_session, _namespaceId, _token);
 
-					RefName refName = TempStorage.GetRefNameForNode(_storagePrefix, step.Name);
+					RefName refName = TempStorage.GetRefNameForNode(_storagePrefix, SetupStepName);
 
 					DirectoryNode node = await storage.ReadRefAsync<DirectoryNode>(refName, cancellationToken: cancellationToken);
 					DirectoryNode? buildGraphDir = await node.FindDirectoryAsync(BuildGraphTempStorageDir, cancellationToken);
