Horde: Add initial batch/step id to bisection task to avoid some expensive graph operations
#rnx

diff --git a/Engine/Source/Programs/Horde/Horde.Server.Tests/BisectTests.cs b/Engine/Source/Programs/Horde/Horde.Server.Tests/BisectTests.cs
--- a/Engine/Source/Programs/Horde/Horde.Server.Tests/BisectTests.cs
+++ b/Engine/Source/Programs/Horde/Horde.Server.Tests/BisectTests.cs
@@ -101,7 +101,7 @@ namespace Horde.Server.Tests
 			List<GetBisectTaskResponse> jobTasks = Deref(await BisectTasksController!.GetJobBisectTasksAsync(failedJob.Id));
 			Assert.AreEqual(1, jobTasks.Count);
 			Assert.AreEqual(failedJob.Id, jobTasks[0].InitialJobId);
-			Assert.AreEqual(3, jobTasks[0].Steps.Count);
+			Assert.AreEqual(4, jobTasks[0].Steps.Count);
 			Assert.AreEqual(BisectTaskState.Succeeded, jobTasks[0].State);
 
 		}
diff --git a/Engine/Source/Programs/Horde/Horde.Server/Jobs/Bisect/BisectController.cs b/Engine/Source/Programs/Horde/Horde.Server/Jobs/Bisect/BisectController.cs
--- a/Engine/Source/Programs/Horde/Horde.Server/Jobs/Bisect/BisectController.cs
+++ b/Engine/Source/Programs/Horde/Horde.Server/Jobs/Bisect/BisectController.cs
@@ -94,6 +94,12 @@ namespace Horde.Server.Jobs.Bisect
 		/// <inheritdoc cref="IBisectTask.InitialJobId"/>
 		public JobId InitialJobId => _bisectTask.InitialJobId;
 
+		/// <inheritdoc cref="IBisectTask.InitialStepId"/>
+		public SubResourceId InitialBatchId => _bisectTask.InitialStepId;
+
+		/// <inheritdoc cref="IBisectTask.InitialStepId"/>
+		public SubResourceId InitialStepId => _bisectTask.InitialStepId;
+
 		/// <inheritdoc cref="IBisectTask.InitialChange"/>
 		public int InitialChange => _bisectTask.InitialChange;
 
@@ -159,7 +165,6 @@ namespace Horde.Server.Jobs.Bisect
 		public List<JobId> ExcludeJobs { get; set; } = new List<JobId>();
 	}
 
-
 	/// <summary>
 	/// Controller for the /api/v1/bisect endpoint
 	/// </summary>
@@ -242,12 +247,27 @@ namespace Horde.Server.Jobs.Bisect
 				return BadRequest("Step has not failed");
 			}
 
+			IJobStepBatch? initialBatch = null;
+			foreach (IJobStepBatch batch in job.Batches)
+			{
+				if (batch.Steps.FirstOrDefault(x => x.Id == jobStep.Id) != null)
+				{
+					initialBatch = batch;
+					break;
+				}
+			}
+
+			if (initialBatch == null)
+			{
+				return BadRequest("Unable to find batch");
+			}
+
 			CreateBisectTaskOptions options = new CreateBisectTaskOptions();
 			options.CommitTags = create.CommitTags;
 			options.IgnoreChanges = create.IgnoreChanges;
 			options.IgnoreJobs = create.IgnoreJobs;
 
-			IBisectTask bisectTask = await _bisectTaskCollection.CreateAsync(job, create.NodeName, jobStep.Outcome, User.GetUserId() ?? UserId.Empty, options, cancellationToken);
+			IBisectTask bisectTask = await _bisectTaskCollection.CreateAsync(job, initialBatch.Id, jobStep.Id, create.NodeName, jobStep.Outcome, User.GetUserId() ?? UserId.Empty, options, cancellationToken);
 			return new CreateBisectTaskResponse(bisectTask);
 		}
 
@@ -342,7 +362,14 @@ namespace Horde.Server.Jobs.Bisect
 		async Task<GetBisectTaskResponse> CreateBisectTaskResponseAsync(IJob initialJob, IBisectTask task, CancellationToken cancellationToken = default)
 		{
 			IUser? user = await _userCollection.GetCachedUserAsync(task.OwnerId);
+
 			List<IJobStepRef> steps = await _jobStepRefs.GetStepsForNodeAsync(initialJob.StreamId, initialJob.TemplateId, task.NodeName, null, true, 1024, task.Id, cancellationToken);
+			IJobStepRef? initialStep = await _jobStepRefs.FindAsync(task.InitialJobId, task.InitialBatchId, task.InitialStepId);
+			if (initialStep != null) 
+			{
+				steps.Add(initialStep);
+			}
+
 			IJob? nextJob = task.State == BisectTaskState.Running ? await _jobCollection.FindBisectTaskJobsAsync(task.Id, true, cancellationToken).FirstOrDefaultAsync(cancellationToken) : null;
 			return new GetBisectTaskResponse(task, new GetThinUserInfoResponse(user), steps, nextJob);
 		}
@@ -421,7 +448,7 @@ namespace Horde.Server.Jobs.Bisect
 				IJob? job = jobs.FirstOrDefault(x => x.Id == task.InitialJobId);
 				if (job != null)
 				{
-					responses.Add(await CreateBisectTaskResponseAsync(job, task, cancellationToken));				
+					responses.Add(await CreateBisectTaskResponseAsync(job, task, cancellationToken));
 				}
 			}
 
diff --git a/Engine/Source/Programs/Horde/Horde.Server/Jobs/Bisect/BisectTaskCollection.cs b/Engine/Source/Programs/Horde/Horde.Server/Jobs/Bisect/BisectTaskCollection.cs
--- a/Engine/Source/Programs/Horde/Horde.Server/Jobs/Bisect/BisectTaskCollection.cs
+++ b/Engine/Source/Programs/Horde/Horde.Server/Jobs/Bisect/BisectTaskCollection.cs
@@ -51,6 +51,12 @@ namespace Horde.Server.Jobs.Bisect
 			[BsonElement("job")]
 			public JobId InitialJobId { get; set; }
 
+			[BsonElement("batchid")]
+			public SubResourceId InitialBatchId { get; set; }
+
+			[BsonElement("stepid")]
+			public SubResourceId InitialStepId { get; set; }
+
 			[BsonElement("chg")]
 			public int InitialChange { get; set; }
 
@@ -87,7 +93,7 @@ namespace Horde.Server.Jobs.Bisect
 		}
 
 		/// <inheritdoc/>
-		public async Task<IBisectTask> CreateAsync(IJob job, string nodeName, JobStepOutcome outcome, UserId ownerId, CreateBisectTaskOptions? options, CancellationToken cancellationToken = default)
+		public async Task<IBisectTask> CreateAsync(IJob job, SubResourceId batchId, SubResourceId stepId, string nodeName, JobStepOutcome outcome, UserId ownerId, CreateBisectTaskOptions? options, CancellationToken cancellationToken = default)
 		{
 			BisectTaskDoc bisectTaskDoc = new BisectTaskDoc();
 			bisectTaskDoc.Id = BisectTaskId.GenerateNewId();
@@ -98,6 +104,8 @@ namespace Horde.Server.Jobs.Bisect
 			bisectTaskDoc.NodeName = nodeName;
 			bisectTaskDoc.Outcome = outcome;
 			bisectTaskDoc.InitialJobId = job.Id;
+			bisectTaskDoc.InitialBatchId = batchId;
+			bisectTaskDoc.InitialStepId = stepId;
 			bisectTaskDoc.InitialChange = job.Change;
 			bisectTaskDoc.CurrentJobId = job.Id;
 			bisectTaskDoc.CurrentChange = job.Change;
diff --git a/Engine/Source/Programs/Horde/Horde.Server/Jobs/Bisect/IBisectTask.cs b/Engine/Source/Programs/Horde/Horde.Server/Jobs/Bisect/IBisectTask.cs
--- a/Engine/Source/Programs/Horde/Horde.Server/Jobs/Bisect/IBisectTask.cs
+++ b/Engine/Source/Programs/Horde/Horde.Server/Jobs/Bisect/IBisectTask.cs
@@ -94,6 +94,16 @@ namespace Horde.Server.Jobs.Bisect
 		/// </summary>
 		public JobId InitialJobId { get; }
 
+		/// <summary>
+		/// Starting batch id for the bisection
+		/// </summary>
+		public SubResourceId InitialBatchId { get; }
+
+		/// <summary>
+		/// Starting step id for the bisection
+		/// </summary>
+		public SubResourceId InitialStepId { get; }
+
 		/// <summary>
 		/// Starting change for the bisection
 		/// </summary>
diff --git a/Engine/Source/Programs/Horde/Horde.Server/Jobs/Bisect/IBisectTaskCollection.cs b/Engine/Source/Programs/Horde/Horde.Server/Jobs/Bisect/IBisectTaskCollection.cs
--- a/Engine/Source/Programs/Horde/Horde.Server/Jobs/Bisect/IBisectTaskCollection.cs
+++ b/Engine/Source/Programs/Horde/Horde.Server/Jobs/Bisect/IBisectTaskCollection.cs
@@ -19,13 +19,15 @@ namespace Horde.Server.Jobs.Bisect
 		/// Creates a new bisect task
 		/// </summary>
 		/// <param name="job">First job to search back from</param>
+		/// <param name="batchId">Starting batch id for the bisection</param>
+		/// <param name="stepId">Starting batch id for the bisection</param>
 		/// <param name="nodeName">Name of the node to search for</param>
 		/// <param name="outcome">Outcome of the step to search for</param>
 		/// <param name="ownerId">User that initiated the search</param>
 		/// <param name="options">Options for the bisection</param>
 		/// <param name="cancellationToken">Cancellation token for the operation</param>
 		/// <returns>The new bisect task instance</returns>
-		public Task<IBisectTask> CreateAsync(IJob job, string nodeName, JobStepOutcome outcome, UserId ownerId, CreateBisectTaskOptions? options = null, CancellationToken cancellationToken = default);
+		public Task<IBisectTask> CreateAsync(IJob job, SubResourceId batchId, SubResourceId stepId, string nodeName, JobStepOutcome outcome, UserId ownerId, CreateBisectTaskOptions? options = null, CancellationToken cancellationToken = default);
 
 		/// <summary>
 		/// Finds all the active bisect tasks
