Horde: Add a name for the root artifact ref in order to conform to the same API as regular storage endpoints.

diff --git a/Engine/Source/Programs/Horde/Horde.Server/Artifacts/ArtifactsController.cs b/Engine/Source/Programs/Horde/Horde.Server/Artifacts/ArtifactsController.cs
--- a/Engine/Source/Programs/Horde/Horde.Server/Artifacts/ArtifactsController.cs
+++ b/Engine/Source/Programs/Horde/Horde.Server/Artifacts/ArtifactsController.cs
@@ -108,7 +108,7 @@ namespace Horde.Server.Artifacts
 		/// <param name="cancellationToken">Cancellation token for the operation</param>
 		/// <returns>Information about all the artifacts</returns>
 		[HttpGet]
-		[Route("/api/v2/artifacts/{id}/ref")]
+		[Route("/api/v2/artifacts/{id}/ref/default")]
 		public async Task<ActionResult<object>> ReadArtifactRefAsync(ArtifactId id, CancellationToken cancellationToken = default)
 		{
 			IArtifact? artifact = await _artifactCollection.GetAsync(id, cancellationToken);
