Horde: Fix serialization error when querying refs.

diff --git a/Engine/Source/Programs/Horde/Horde.Server/Artifacts/ArtifactsController.cs b/Engine/Source/Programs/Horde/Horde.Server/Artifacts/ArtifactsController.cs
--- a/Engine/Source/Programs/Horde/Horde.Server/Artifacts/ArtifactsController.cs
+++ b/Engine/Source/Programs/Horde/Horde.Server/Artifacts/ArtifactsController.cs
@@ -121,7 +121,7 @@ namespace Horde.Server.Artifacts
 				return Forbid(ArtifactAclAction.ReadArtifact, artifact.AclScope);
 			}
 
-			return StorageController.ReadRefInternalAsync(_storageService, artifact.NamespaceId, artifact.RefName, Request.Headers, cancellationToken);
+			return await StorageController.ReadRefInternalAsync(_storageService, artifact.NamespaceId, artifact.RefName, Request.Headers, cancellationToken);
 		}
 
 		/// <summary>
