Horde: Fix ref responses from artifact controller being wrapped in another ActionResult object.

diff --git a/Engine/Source/Programs/Horde/Horde.Server/Artifacts/ArtifactsController.cs b/Engine/Source/Programs/Horde/Horde.Server/Artifacts/ArtifactsController.cs
--- a/Engine/Source/Programs/Horde/Horde.Server/Artifacts/ArtifactsController.cs
+++ b/Engine/Source/Programs/Horde/Horde.Server/Artifacts/ArtifactsController.cs
@@ -109,7 +109,7 @@ namespace Horde.Server.Artifacts
 		/// <returns>Information about all the artifacts</returns>
 		[HttpGet]
 		[Route("/api/v2/artifacts/{id}/refs/default")]
-		public async Task<ActionResult<object>> ReadArtifactRefAsync(ArtifactId id, CancellationToken cancellationToken = default)
+		public async Task<ActionResult<ReadRefResponse>> ReadArtifactRefAsync(ArtifactId id, CancellationToken cancellationToken = default)
 		{
 			IArtifact? artifact = await _artifactCollection.GetAsync(id, cancellationToken);
 			if (artifact == null)
