Horde: Fix return type from querying artifact blobs.

diff --git a/Engine/Source/Programs/Horde/Horde.Server/Artifacts/ArtifactsController.cs b/Engine/Source/Programs/Horde/Horde.Server/Artifacts/ArtifactsController.cs
--- a/Engine/Source/Programs/Horde/Horde.Server/Artifacts/ArtifactsController.cs
+++ b/Engine/Source/Programs/Horde/Horde.Server/Artifacts/ArtifactsController.cs
@@ -81,7 +81,7 @@ namespace Horde.Server.Artifacts
 		[HttpGet]
 		[Route("/api/v2/artifacts/{id}/blobs/{*locator}")]
 		[Route("/api/v2/artifacts/{id}/bundles/{*locator}")]
-		public async Task<ActionResult<object>> ReadArtifactBlobAsync(ArtifactId id, BundleLocator locator, [FromQuery] int? offset = null, [FromQuery] int? length = null, CancellationToken cancellationToken = default)
+		public async Task<ActionResult> ReadArtifactBlobAsync(ArtifactId id, BundleLocator locator, [FromQuery] int? offset = null, [FromQuery] int? length = null, CancellationToken cancellationToken = default)
 		{
 			IArtifact? artifact = await _artifactCollection.GetAsync(id, cancellationToken);
 			if (artifact == null)
@@ -98,7 +98,7 @@ namespace Horde.Server.Artifacts
 			}
 
 			StorageClient storageClient = await _storageService.GetClientAsync(artifact.NamespaceId, cancellationToken);
-			return StorageController.ReadBlobInternalAsync(storageClient, locator, offset, length, cancellationToken);
+			return await StorageController.ReadBlobInternalAsync(storageClient, locator, offset, length, cancellationToken);
 		}
 
 		/// <summary>
