Horde: Add missing call to refreshing cached commit in Perforce trigger callback

diff --git a/Engine/Source/Programs/Horde/Horde.Server/Perforce/PerforceController.cs b/Engine/Source/Programs/Horde/Horde.Server/Perforce/PerforceController.cs
--- a/Engine/Source/Programs/Horde/Horde.Server/Perforce/PerforceController.cs
+++ b/Engine/Source/Programs/Horde/Horde.Server/Perforce/PerforceController.cs
@@ -11,6 +11,7 @@ using Microsoft.AspNetCore.Authorization;
 using Microsoft.AspNetCore.Mvc;
 using Microsoft.Extensions.Logging;
 using Microsoft.Extensions.Options;
+using OpenTelemetry.Trace;
 
 namespace Horde.Server.Perforce
 {
@@ -254,14 +255,16 @@ namespace Horde.Server.Perforce
 	public class PublicPerforceController : ControllerBase
 	{
 		private readonly IPerforceService _perforceService;
+		private readonly Tracer _tracer;
 		private readonly ILogger<PerforceController> _logger;
 
 		/// <summary>
 		/// Constructor
 		/// </summary>
-		public PublicPerforceController(IPerforceService perforceService, ILogger<PerforceController> logger)
+		public PublicPerforceController(IPerforceService perforceService, Tracer tracer, ILogger<PerforceController> logger)
 		{
 			_perforceService = perforceService;
+			_tracer = tracer;
 			_logger = logger;
 		}
 
@@ -276,12 +279,25 @@ namespace Horde.Server.Perforce
 			_logger.LogDebug("Received Perforce trigger callback. Type={Type} CL={Changelist} User={User} Root={Root}", 
 				trigger.TriggerType, trigger.ChangeNumber, trigger.User, trigger.ChangeRoot);
 			
-			// For "form-save" triggers, change number can be -1 due to variable "formname" is set to "default" (non-submitted changelist)
-			if (trigger.TriggerType == "change-commit" && trigger.ChangeNumber != -1)
+			using TelemetrySpan span = _tracer.StartActiveSpan($"{nameof(PublicPerforceController)}.{nameof(TriggerCallback)}");
+			span.SetAttribute("type", trigger.TriggerType);
+			span.SetAttribute("cl", trigger.ChangeNumber);
+			span.SetAttribute("user", trigger.User);
+			span.SetAttribute("root", trigger.ChangeRoot);
+			span.SetAttribute("client", trigger.Client);
+
+			switch (trigger.TriggerType)
 			{
-				// Not implemented yet
-				await Task.Delay(0); // Avoid await warnings
-				// await _perforceService.RefreshCachedCommitAsync(cluster, trigger.ChangeNumber);
+				case "change-commit":
+				case "shelve-commit":
+				case "form-save":
+					// For "form-save" triggers, change number can be -1 due to variable "formname" is set to "default" (non-submitted changelist)
+					if (trigger.ChangeNumber != -1)
+					{
+						await _perforceService.RefreshCachedCommitAsync(cluster, trigger.ChangeNumber);
+						span.SetAttribute("commitRefreshed", true);
+					}
+					break;
 			}
 
 			string content = "{\"message\": \"Trigger received\"}";
