Horde: Fix bisection response returning subresource object
#rnx

diff --git a/Engine/Source/Programs/Horde/Horde.Server/Jobs/Bisect/BisectController.cs b/Engine/Source/Programs/Horde/Horde.Server/Jobs/Bisect/BisectController.cs
--- a/Engine/Source/Programs/Horde/Horde.Server/Jobs/Bisect/BisectController.cs
+++ b/Engine/Source/Programs/Horde/Horde.Server/Jobs/Bisect/BisectController.cs
@@ -95,11 +95,11 @@ namespace Horde.Server.Jobs.Bisect
 		/// <inheritdoc cref="IBisectTask.InitialJobId"/>
 		public JobId InitialJobId => _bisectTask.InitialJobId;
 
-		/// <inheritdoc cref="IBisectTask.InitialStepId"/>
-		public SubResourceId InitialBatchId => _bisectTask.InitialStepId;
+		/// <inheritdoc cref="IBisectTask.InitialBatchId"/>
+		public string InitialBatchId => _bisectTask.InitialStepId.ToString();
 
 		/// <inheritdoc cref="IBisectTask.InitialStepId"/>
-		public SubResourceId InitialStepId => _bisectTask.InitialStepId;
+		public string InitialStepId => _bisectTask.InitialStepId.ToString();
 
 		/// <inheritdoc cref="IBisectTask.InitialChange"/>
 		public int InitialChange => _bisectTask.InitialChange;
