Horde: Add a debug message when the storage backend returns more data than it should.

diff --git a/Engine/Source/Programs/Horde/Horde.Server/Storage/StorageService.cs b/Engine/Source/Programs/Horde/Horde.Server/Storage/StorageService.cs
--- a/Engine/Source/Programs/Horde/Horde.Server/Storage/StorageService.cs
+++ b/Engine/Source/Programs/Horde/Horde.Server/Storage/StorageService.cs
@@ -121,12 +121,14 @@ namespace Horde.Server.Storage
 			readonly StorageService _outer;
 			readonly string _prefix;
 			readonly Tracer _tracer;
+			readonly ILogger _logger;
 
 			public StorageClientImpl(StorageService outer, NamespaceConfig config, IStorageBackend backend, IMemoryCache? memoryCache, Tracer tracer, ILogger logger)
 				: base(config, backend, memoryCache, logger)
 			{
 				_outer = outer;
 				_tracer = tracer;
+				_logger = logger;
 
 				_prefix = config.Prefix;
 				if (_prefix.Length > 0 && !_prefix.EndsWith("/", StringComparison.Ordinal))
@@ -163,7 +165,14 @@ namespace Horde.Server.Storage
 
 				string path = GetBlobPath(locator);
 				await using Stream stream = await Backend.ReadAsync(path, offset, length, cancellationToken);
-				return await stream.ReadAllBytesAsync(cancellationToken);
+
+				ReadOnlyMemory<byte> data = await stream.ReadAllBytesAsync(cancellationToken);
+				if (data.Length > length)
+				{
+					_logger.LogDebug("Storage backend returned more data than requested; truncating response (wanted {Length}, got {GotLength})", length, data.Length);
+					data = data.Slice(0, length);
+				}
+				return data;
 			}
 
 			/// <inheritdoc/>
