Horde: Allow cancellation on any lease types.

diff --git a/Engine/Source/Programs/Horde/Horde.Server/Agents/Leases/LeasesController.cs b/Engine/Source/Programs/Horde/Horde.Server/Agents/Leases/LeasesController.cs
--- a/Engine/Source/Programs/Horde/Horde.Server/Agents/Leases/LeasesController.cs
+++ b/Engine/Source/Programs/Horde/Horde.Server/Agents/Leases/LeasesController.cs
@@ -174,17 +174,11 @@ namespace Horde.Server.Agents.Leases
 			}
 
 			AgentLease? agentLease = agent.Leases.FirstOrDefault(x => x.Id == leaseId);
-
 			if (agentLease == null)
 			{
 				return NotFound(agent.Id, leaseId);
 			}
 
-			if (!agentLease.IsConformLease())
-			{
-				return BadRequest("Lease abort only supported on conform leases for now, {LeaseId}", leaseId);
-			}
-
 			await _agentService.CancelLeaseAsync(agent, leaseId);
 			return Ok();
 		}
