Horde: Support ranged GETs through redirect URIs in storage backends.

diff --git a/Engine/Source/Programs/Horde/Horde.Server/Storage/Backends/AwsStorageBackend.cs b/Engine/Source/Programs/Horde/Horde.Server/Storage/Backends/AwsStorageBackend.cs
--- a/Engine/Source/Programs/Horde/Horde.Server/Storage/Backends/AwsStorageBackend.cs
+++ b/Engine/Source/Programs/Horde/Horde.Server/Storage/Backends/AwsStorageBackend.cs
@@ -6,6 +6,7 @@ using System.Data.SqlTypes;
 using System.Globalization;
 using System.IO;
 using System.Runtime.CompilerServices;
+using System.Text;
 using System.Threading;
 using System.Threading.Tasks;
 using Amazon;
@@ -313,7 +314,7 @@ namespace Horde.Server.Storage.Backends
 		}
 
 		/// <inheritdoc/>
-		public ValueTask<Uri?> TryGetReadRedirectAsync(string path, CancellationToken cancellationToken = default) => new ValueTask<Uri?>(GetPresignedUrl(path, HttpVerb.GET));
+		public ValueTask<Uri?> TryGetReadRedirectAsync(string path, int? offset, int? length, CancellationToken cancellationToken = default) => new ValueTask<Uri?>(GetPresignedUrl(path, HttpVerb.GET, offset, length));
 
 		/// <inheritdoc/>
 		public ValueTask<Uri?> TryGetWriteRedirectAsync(string path, CancellationToken cancellationToken = default) => new ValueTask<Uri?>(GetPresignedUrl(path, HttpVerb.PUT));
@@ -321,7 +322,7 @@ namespace Horde.Server.Storage.Backends
 		/// <summary>
 		/// Helper method to generate a presigned URL for a request
 		/// </summary>
-		Uri? GetPresignedUrl(string path, HttpVerb verb)
+		Uri? GetPresignedUrl(string path, HttpVerb verb, int? offset = null, int? length = null)
 		{
 			using TelemetrySpan span = OpenTelemetryTracers.Horde.StartActiveSpan($"{nameof(AwsStorageBackend)}.{nameof(GetPresignedUrl)}");
 			span.SetAttribute("path", path);
@@ -334,6 +335,20 @@ namespace Horde.Server.Storage.Backends
 				newGetRequest.BucketName = _options.AwsBucketName;
 				newGetRequest.Key = fullPath;
 				newGetRequest.Verb = verb;
+				if (offset != null || length != null)
+				{
+					StringBuilder range = new StringBuilder();
+					if (offset != null)
+					{
+						range.Append(offset.Value);
+					}
+					range.Append("-");
+					if (length != null)
+					{
+						range.Append(length.Value);
+					}
+					newGetRequest.Headers["Range"] = range.ToString();
+				}
 				newGetRequest.Expires = DateTime.UtcNow.AddHours(3.0);
 				newGetRequest.ResponseHeaderOverrides.CacheControl = "private, max-age=2592000, immutable"; // 30 days
 
diff --git a/Engine/Source/Programs/Horde/Horde.Server/Storage/Backends/FileStorageBackend.cs b/Engine/Source/Programs/Horde/Horde.Server/Storage/Backends/FileStorageBackend.cs
--- a/Engine/Source/Programs/Horde/Horde.Server/Storage/Backends/FileStorageBackend.cs
+++ b/Engine/Source/Programs/Horde/Horde.Server/Storage/Backends/FileStorageBackend.cs
@@ -180,7 +180,7 @@ namespace Horde.Server.Storage.Backends
 		}
 
 		/// <inheritdoc/>
-		public ValueTask<Uri?> TryGetReadRedirectAsync(string path, CancellationToken cancellationToken = default) => default;
+		public ValueTask<Uri?> TryGetReadRedirectAsync(string path, int? offset = null, int? length = null, CancellationToken cancellationToken = default) => default;
 
 		/// <inheritdoc/>
 		public ValueTask<Uri?> TryGetWriteRedirectAsync(string path, CancellationToken cancellationToken = default) => default;
diff --git a/Engine/Source/Programs/Horde/Horde.Server/Storage/Backends/MemoryStorageBackend.cs b/Engine/Source/Programs/Horde/Horde.Server/Storage/Backends/MemoryStorageBackend.cs
--- a/Engine/Source/Programs/Horde/Horde.Server/Storage/Backends/MemoryStorageBackend.cs
+++ b/Engine/Source/Programs/Horde/Horde.Server/Storage/Backends/MemoryStorageBackend.cs
@@ -78,7 +78,7 @@ namespace Horde.Server.Storage.Backends
 		}
 
 		/// <inheritdoc/>
-		public ValueTask<Uri?> TryGetReadRedirectAsync(string path, CancellationToken cancellationToken = default) => default;
+		public ValueTask<Uri?> TryGetReadRedirectAsync(string path, int? offset = null, int? length = null, CancellationToken cancellationToken = default) => default;
 
 		/// <inheritdoc/>
 		public ValueTask<Uri?> TryGetWriteRedirectAsync(string path, CancellationToken cancellationToken = default) => default;
diff --git a/Engine/Source/Programs/Horde/Horde.Server/Storage/IStorageBackend.cs b/Engine/Source/Programs/Horde/Horde.Server/Storage/IStorageBackend.cs
--- a/Engine/Source/Programs/Horde/Horde.Server/Storage/IStorageBackend.cs
+++ b/Engine/Source/Programs/Horde/Horde.Server/Storage/IStorageBackend.cs
@@ -72,9 +72,11 @@ namespace Horde.Server.Storage
 		/// Gets a HTTP redirect for a read request
 		/// </summary>
 		/// <param name="path">Path to read from</param>
+		/// <param name="offset">Offset to start reading from</param>
+		/// <param name="length">Length of data to read</param>
 		/// <param name="cancellationToken">Cancellation token for the operation</param>
 		/// <returns>Path to upload the data to</returns>
-		ValueTask<Uri?> TryGetReadRedirectAsync(string path, CancellationToken cancellationToken = default);
+		ValueTask<Uri?> TryGetReadRedirectAsync(string path, int? offset, int? length, CancellationToken cancellationToken = default);
 
 		/// <summary>
 		/// Gets a HTTP redirect for a write request
@@ -137,7 +139,7 @@ namespace Horde.Server.Storage
 			public IAsyncEnumerable<string> EnumerateAsync(CancellationToken cancellationToken = default) => _inner.EnumerateAsync(cancellationToken);
 
 			/// <inheritdoc/>
-			public ValueTask<Uri?> TryGetReadRedirectAsync(string path, CancellationToken cancellationToken = default) => _inner.TryGetReadRedirectAsync(path, cancellationToken);
+			public ValueTask<Uri?> TryGetReadRedirectAsync(string path, int? offset = null, int? length = null, CancellationToken cancellationToken = default) => _inner.TryGetReadRedirectAsync(path, offset, length, cancellationToken);
 
 			/// <inheritdoc/>
 			public ValueTask<Uri?> TryGetWriteRedirectAsync(string path, CancellationToken cancellationToken = default) => _inner.TryGetWriteRedirectAsync(path, cancellationToken);
diff --git a/Engine/Source/Programs/Horde/Horde.Server/Storage/StorageBackendProvider.cs b/Engine/Source/Programs/Horde/Horde.Server/Storage/StorageBackendProvider.cs
--- a/Engine/Source/Programs/Horde/Horde.Server/Storage/StorageBackendProvider.cs
+++ b/Engine/Source/Programs/Horde/Horde.Server/Storage/StorageBackendProvider.cs
@@ -86,7 +86,7 @@ namespace Horde.Server.Storage
 			public Task WriteAsync(string path, Stream stream, CancellationToken cancellationToken = default) => _backend.WriteAsync(path, stream, cancellationToken);
 
 			/// <inheritdoc/>
-			public ValueTask<Uri?> TryGetReadRedirectAsync(string path, CancellationToken cancellationToken = default) => _backend.TryGetReadRedirectAsync(path, cancellationToken);
+			public ValueTask<Uri?> TryGetReadRedirectAsync(string path, int? offset = null, int? length = null, CancellationToken cancellationToken = default) => _backend.TryGetReadRedirectAsync(path, offset, length, cancellationToken);
 
 			/// <inheritdoc/>
 			public ValueTask<Uri?> TryGetWriteRedirectAsync(string path, CancellationToken cancellationToken = default) => _backend.TryGetWriteRedirectAsync(path, cancellationToken);
diff --git a/Engine/Source/Programs/Horde/Horde.Server/Storage/StorageController.cs b/Engine/Source/Programs/Horde/Horde.Server/Storage/StorageController.cs
--- a/Engine/Source/Programs/Horde/Horde.Server/Storage/StorageController.cs
+++ b/Engine/Source/Programs/Horde/Horde.Server/Storage/StorageController.cs
@@ -276,7 +276,7 @@ namespace Horde.Server.Storage
 		{
 			if (storageClient is StorageClient storageClientImpl)
 			{
-				Uri? redirectUrl = await storageClientImpl.GetReadRedirectAsync(locator, cancellationToken);
+				Uri? redirectUrl = await storageClientImpl.GetReadRedirectAsync(locator, offset, length, cancellationToken);
 				if (redirectUrl != null)
 				{
 					return new RedirectResult(redirectUrl.ToString());
diff --git a/Engine/Source/Programs/Horde/Horde.Server/Storage/StorageService.cs b/Engine/Source/Programs/Horde/Horde.Server/Storage/StorageService.cs
--- a/Engine/Source/Programs/Horde/Horde.Server/Storage/StorageService.cs
+++ b/Engine/Source/Programs/Horde/Horde.Server/Storage/StorageService.cs
@@ -98,9 +98,11 @@ namespace Horde.Server.Storage
 		/// Gets a redirect for a read request
 		/// </summary>
 		/// <param name="locator">Locator for the blob</param>
+		/// <param name="offset">Offset of the data to return</param>
+		/// <param name="length">Length of the data to return</param>
 		/// <param name="cancellationToken">Cancellation token for the operation</param>
 		/// <returns>Path to upload the data to</returns>
-		public abstract ValueTask<Uri?> GetReadRedirectAsync(BundleLocator locator, CancellationToken cancellationToken = default);
+		public abstract ValueTask<Uri?> GetReadRedirectAsync(BundleLocator locator, int? offset, int? length, CancellationToken cancellationToken = default);
 
 		/// <summary>
 		/// Gets a redirect for a write request
@@ -153,7 +155,7 @@ namespace Horde.Server.Storage
 			}
 
 			/// <inheritdoc/>
-			public override ValueTask<Uri?> GetReadRedirectAsync(BundleLocator locator, CancellationToken cancellationToken = default) => Backend.TryGetReadRedirectAsync(GetBlobPath(locator), cancellationToken);
+			public override ValueTask<Uri?> GetReadRedirectAsync(BundleLocator locator, int? offset = null, int? length = null, CancellationToken cancellationToken = default) => Backend.TryGetReadRedirectAsync(GetBlobPath(locator), offset, length, cancellationToken);
 
 			/// <inheritdoc/>
 			public override async Task<ReadOnlyMemory<byte>> ReadBundleRangeAsync(BundleLocator locator, int offset, int length, CancellationToken cancellationToken = default)
