Horde: Fix copy paste with batch id on bisection response
#rnx

diff --git a/Engine/Source/Programs/Horde/Horde.Server/Jobs/Bisect/BisectController.cs b/Engine/Source/Programs/Horde/Horde.Server/Jobs/Bisect/BisectController.cs
--- a/Engine/Source/Programs/Horde/Horde.Server/Jobs/Bisect/BisectController.cs
+++ b/Engine/Source/Programs/Horde/Horde.Server/Jobs/Bisect/BisectController.cs
@@ -96,7 +96,7 @@ namespace Horde.Server.Jobs.Bisect
 		public JobId InitialJobId => _bisectTask.InitialJobId;
 
 		/// <inheritdoc cref="IBisectTask.InitialBatchId"/>
-		public string InitialBatchId => _bisectTask.InitialStepId.ToString();
+		public string InitialBatchId => _bisectTask.InitialBatchId.ToString();
 
 		/// <inheritdoc cref="IBisectTask.InitialStepId"/>
 		public string InitialStepId => _bisectTask.InitialStepId.ToString();
