Horde: Fix syntax of AWS range request.

diff --git a/Engine/Source/Programs/Horde/Horde.Server/Storage/Backends/AwsStorageBackend.cs b/Engine/Source/Programs/Horde/Horde.Server/Storage/Backends/AwsStorageBackend.cs
--- a/Engine/Source/Programs/Horde/Horde.Server/Storage/Backends/AwsStorageBackend.cs
+++ b/Engine/Source/Programs/Horde/Horde.Server/Storage/Backends/AwsStorageBackend.cs
@@ -337,15 +337,12 @@ namespace Horde.Server.Storage.Backends
 				newGetRequest.Verb = verb;
 				if (offset != null || length != null)
 				{
-					StringBuilder range = new StringBuilder();
-					if (offset != null)
-					{
-						range.Append(offset.Value);
-					}
-					range.Append("-");
+					StringBuilder range = new StringBuilder("bytes=");
+					range.Append(offset ?? 0);
+					range.Append('-');
 					if (length != null)
 					{
-						range.Append(length.Value);
+						range.Append(offset + (length.Value - 1));
 					}
 					newGetRequest.Headers["Range"] = range.ToString();
 				}
