Horde: Use HTTP range headers for partial object fetches. This works better with redirect URLs, because the parameters are resent to the new location.

diff --git a/Engine/Source/Programs/Horde/Horde.Server/Artifacts/ArtifactsController.cs b/Engine/Source/Programs/Horde/Horde.Server/Artifacts/ArtifactsController.cs
--- a/Engine/Source/Programs/Horde/Horde.Server/Artifacts/ArtifactsController.cs
+++ b/Engine/Source/Programs/Horde/Horde.Server/Artifacts/ArtifactsController.cs
@@ -74,14 +74,12 @@ namespace Horde.Server.Artifacts
 		/// </summary>
 		/// <param name="id">Identifier of the artifact to retrieve</param>
 		/// <param name="locator">The blob locator</param>
-		/// <param name="length">Length of data to return</param>
-		/// <param name="offset">Offset of the data to return</param>
 		/// <param name="cancellationToken">Cancellation token for the operation</param>
 		/// <returns>Information about all the artifacts</returns>
 		[HttpGet]
 		[Route("/api/v2/artifacts/{id}/blobs/{*locator}")]
 		[Route("/api/v2/artifacts/{id}/bundles/{*locator}")]
-		public async Task<ActionResult> ReadArtifactBlobAsync(ArtifactId id, BundleLocator locator, [FromQuery] int? offset = null, [FromQuery] int? length = null, CancellationToken cancellationToken = default)
+		public async Task<ActionResult> ReadArtifactBlobAsync(ArtifactId id, BundleLocator locator, CancellationToken cancellationToken = default)
 		{
 			IArtifact? artifact = await _artifactCollection.GetAsync(id, cancellationToken);
 			if (artifact == null)
@@ -98,7 +96,7 @@ namespace Horde.Server.Artifacts
 			}
 
 			StorageClient storageClient = await _storageService.GetClientAsync(artifact.NamespaceId, cancellationToken);
-			return await StorageController.ReadBlobInternalAsync(storageClient, locator, offset, length, cancellationToken);
+			return await StorageController.ReadBlobInternalAsync(storageClient, locator, Request.Headers, cancellationToken);
 		}
 
 		/// <summary>
diff --git a/Engine/Source/Programs/Horde/Horde.Server/Storage/Backends/AwsStorageBackend.cs b/Engine/Source/Programs/Horde/Horde.Server/Storage/Backends/AwsStorageBackend.cs
--- a/Engine/Source/Programs/Horde/Horde.Server/Storage/Backends/AwsStorageBackend.cs
+++ b/Engine/Source/Programs/Horde/Horde.Server/Storage/Backends/AwsStorageBackend.cs
@@ -270,9 +270,18 @@ namespace Horde.Server.Storage.Backends
 		}
 
 		/// <inheritdoc/>
-		public Task<Stream> ReadAsync(string path, int offset, int length, CancellationToken cancellationToken)
+		public Task<Stream> ReadAsync(string path, int offset, int? length, CancellationToken cancellationToken)
 		{
-			return ReadAsync(path, new ByteRange(offset, offset + length), cancellationToken);
+			string range;
+			if (length == null)
+			{
+				range = $"bytes={offset}-";
+			}
+			else
+			{
+				range = $"bytes={offset}-{offset + length.Value - 1}";
+			}
+			return ReadAsync(path, new ByteRange(range), cancellationToken);
 		}
 
 		async Task<Stream> ReadAsync(string path, ByteRange? byteRange, CancellationToken cancellationToken)
@@ -314,7 +323,7 @@ namespace Horde.Server.Storage.Backends
 		}
 
 		/// <inheritdoc/>
-		public ValueTask<Uri?> TryGetReadRedirectAsync(string path, int? offset, int? length, CancellationToken cancellationToken = default) => new ValueTask<Uri?>(GetPresignedUrl(path, HttpVerb.GET, offset, length));
+		public ValueTask<Uri?> TryGetReadRedirectAsync(string path, CancellationToken cancellationToken = default) => new ValueTask<Uri?>(GetPresignedUrl(path, HttpVerb.GET));
 
 		/// <inheritdoc/>
 		public ValueTask<Uri?> TryGetWriteRedirectAsync(string path, CancellationToken cancellationToken = default) => new ValueTask<Uri?>(GetPresignedUrl(path, HttpVerb.PUT));
@@ -322,7 +331,7 @@ namespace Horde.Server.Storage.Backends
 		/// <summary>
 		/// Helper method to generate a presigned URL for a request
 		/// </summary>
-		Uri? GetPresignedUrl(string path, HttpVerb verb, int? offset = null, int? length = null)
+		Uri? GetPresignedUrl(string path, HttpVerb verb)
 		{
 			using TelemetrySpan span = OpenTelemetryTracers.Horde.StartActiveSpan($"{nameof(AwsStorageBackend)}.{nameof(GetPresignedUrl)}");
 			span.SetAttribute("path", path);
@@ -335,17 +344,6 @@ namespace Horde.Server.Storage.Backends
 				newGetRequest.BucketName = _options.AwsBucketName;
 				newGetRequest.Key = fullPath;
 				newGetRequest.Verb = verb;
-				if (offset != null || length != null)
-				{
-					StringBuilder range = new StringBuilder("bytes=");
-					range.Append(offset ?? 0);
-					range.Append('-');
-					if (length != null)
-					{
-						range.Append(offset + (length.Value - 1));
-					}
-					newGetRequest.Headers["Range"] = range.ToString();
-				}
 				newGetRequest.Expires = DateTime.UtcNow.AddHours(3.0);
 				newGetRequest.ResponseHeaderOverrides.CacheControl = "private, max-age=2592000, immutable"; // 30 days
 
diff --git a/Engine/Source/Programs/Horde/Horde.Server/Storage/Backends/FileStorageBackend.cs b/Engine/Source/Programs/Horde/Horde.Server/Storage/Backends/FileStorageBackend.cs
--- a/Engine/Source/Programs/Horde/Horde.Server/Storage/Backends/FileStorageBackend.cs
+++ b/Engine/Source/Programs/Horde/Horde.Server/Storage/Backends/FileStorageBackend.cs
@@ -73,7 +73,7 @@ namespace Horde.Server.Storage.Backends
 		}
 
 		/// <inheritdoc/>
-		public async Task<Stream> ReadAsync(string path, int offset, int length, CancellationToken cancellationToken)
+		public async Task<Stream> ReadAsync(string path, int offset, int? length, CancellationToken cancellationToken)
 		{
 			Stream stream = await ReadAsync(path, cancellationToken);
 			stream.Seek(offset, SeekOrigin.Begin);
@@ -180,7 +180,7 @@ namespace Horde.Server.Storage.Backends
 		}
 
 		/// <inheritdoc/>
-		public ValueTask<Uri?> TryGetReadRedirectAsync(string path, int? offset = null, int? length = null, CancellationToken cancellationToken = default) => default;
+		public ValueTask<Uri?> TryGetReadRedirectAsync(string path, CancellationToken cancellationToken = default) => default;
 
 		/// <inheritdoc/>
 		public ValueTask<Uri?> TryGetWriteRedirectAsync(string path, CancellationToken cancellationToken = default) => default;
diff --git a/Engine/Source/Programs/Horde/Horde.Server/Storage/Backends/MemoryStorageBackend.cs b/Engine/Source/Programs/Horde/Horde.Server/Storage/Backends/MemoryStorageBackend.cs
--- a/Engine/Source/Programs/Horde/Horde.Server/Storage/Backends/MemoryStorageBackend.cs
+++ b/Engine/Source/Programs/Horde/Horde.Server/Storage/Backends/MemoryStorageBackend.cs
@@ -7,6 +7,7 @@ using System.IO;
 using System.Runtime.CompilerServices;
 using System.Threading;
 using System.Threading.Tasks;
+using EpicGames.Core;
 
 namespace Horde.Server.Storage.Backends
 {
@@ -29,18 +30,17 @@ namespace Horde.Server.Storage.Backends
 		}
 
 		/// <inheritdoc/>
-		public Task<Stream> ReadAsync(string path, CancellationToken cancellationToken)
-		{
-			byte[] data = _pathToData[path];
-			return Task.FromResult<Stream>(new MemoryStream(data, false));
-		}
+		public Task<Stream> ReadAsync(string path, CancellationToken cancellationToken) => ReadAsync(path, 0, null, cancellationToken);
 
 		/// <inheritdoc/>
-		public async Task<Stream> ReadAsync(string path, int offset, int length, CancellationToken cancellationToken)
+		public Task<Stream> ReadAsync(string path, int offset, int? length, CancellationToken cancellationToken)
 		{
-			Stream stream = await ReadAsync(path, cancellationToken);
-			stream.Seek(offset, SeekOrigin.Begin);
-			return stream;
+			ReadOnlyMemory<byte> data = _pathToData[path].AsMemory(offset);
+			if (length != null)
+			{
+				data = data.Slice(0, length.Value);
+			}
+			return Task.FromResult<Stream>(new ReadOnlyMemoryStream(data));
 		}
 
 		/// <inheritdoc/>
@@ -78,7 +78,7 @@ namespace Horde.Server.Storage.Backends
 		}
 
 		/// <inheritdoc/>
-		public ValueTask<Uri?> TryGetReadRedirectAsync(string path, int? offset = null, int? length = null, CancellationToken cancellationToken = default) => default;
+		public ValueTask<Uri?> TryGetReadRedirectAsync(string path, CancellationToken cancellationToken = default) => default;
 
 		/// <inheritdoc/>
 		public ValueTask<Uri?> TryGetWriteRedirectAsync(string path, CancellationToken cancellationToken = default) => default;
diff --git a/Engine/Source/Programs/Horde/Horde.Server/Storage/IStorageBackend.cs b/Engine/Source/Programs/Horde/Horde.Server/Storage/IStorageBackend.cs
--- a/Engine/Source/Programs/Horde/Horde.Server/Storage/IStorageBackend.cs
+++ b/Engine/Source/Programs/Horde/Horde.Server/Storage/IStorageBackend.cs
@@ -19,14 +19,6 @@ namespace Horde.Server.Storage
 		/// </summary>
 		bool SupportsRedirects { get; }
 
-		/// <summary>
-		/// Attempts to open a read stream for the given path.
-		/// </summary>
-		/// <param name="path">Relative path within the bucket</param>
-		/// <param name="cancellationToken">Cancellation token for the operation</param>
-		/// <returns></returns>
-		Task<Stream> ReadAsync(string path, CancellationToken cancellationToken = default);
-
 		/// <summary>
 		/// Attempts to open a read stream for the given path.
 		/// </summary>
@@ -35,7 +27,7 @@ namespace Horde.Server.Storage
 		/// <param name="length">Length of data to read</param>
 		/// <param name="cancellationToken">Cancellation token for the operation</param>
 		/// <returns></returns>
-		Task<Stream> ReadAsync(string path, int offset, int length, CancellationToken cancellationToken = default);
+		Task<Stream> ReadAsync(string path, int offset, int? length, CancellationToken cancellationToken = default);
 
 		/// <summary>
 		/// Writes a stream to the given path. If the stream throws an exception during read, the write will be aborted.
@@ -72,11 +64,9 @@ namespace Horde.Server.Storage
 		/// Gets a HTTP redirect for a read request
 		/// </summary>
 		/// <param name="path">Path to read from</param>
-		/// <param name="offset">Offset to start reading from</param>
-		/// <param name="length">Length of data to read</param>
 		/// <param name="cancellationToken">Cancellation token for the operation</param>
 		/// <returns>Path to upload the data to</returns>
-		ValueTask<Uri?> TryGetReadRedirectAsync(string path, int? offset, int? length, CancellationToken cancellationToken = default);
+		ValueTask<Uri?> TryGetReadRedirectAsync(string path, CancellationToken cancellationToken = default);
 
 		/// <summary>
 		/// Gets a HTTP redirect for a write request
@@ -121,10 +111,7 @@ namespace Horde.Server.Storage
 			public void Dispose() => _inner.Dispose();
 
 			/// <inheritdoc/>
-			public Task<Stream> ReadAsync(string path, CancellationToken cancellationToken) => _inner.ReadAsync(path, cancellationToken);
-
-			/// <inheritdoc/>
-			public Task<Stream> ReadAsync(string path, int offset, int length, CancellationToken cancellationToken) => _inner.ReadAsync(path, offset, length, cancellationToken);
+			public Task<Stream> ReadAsync(string path, int offset, int? length, CancellationToken cancellationToken) => _inner.ReadAsync(path, offset, length, cancellationToken);
 
 			/// <inheritdoc/>
 			public Task WriteAsync(string path, Stream stream, CancellationToken cancellationToken) => _inner.WriteAsync(path, stream, cancellationToken);
@@ -139,7 +126,7 @@ namespace Horde.Server.Storage
 			public IAsyncEnumerable<string> EnumerateAsync(CancellationToken cancellationToken = default) => _inner.EnumerateAsync(cancellationToken);
 
 			/// <inheritdoc/>
-			public ValueTask<Uri?> TryGetReadRedirectAsync(string path, int? offset = null, int? length = null, CancellationToken cancellationToken = default) => _inner.TryGetReadRedirectAsync(path, offset, length, cancellationToken);
+			public ValueTask<Uri?> TryGetReadRedirectAsync(string path, CancellationToken cancellationToken = default) => _inner.TryGetReadRedirectAsync(path, cancellationToken);
 
 			/// <inheritdoc/>
 			public ValueTask<Uri?> TryGetWriteRedirectAsync(string path, CancellationToken cancellationToken = default) => _inner.TryGetWriteRedirectAsync(path, cancellationToken);
@@ -203,6 +190,15 @@ namespace Horde.Server.Storage
 			return new TypedStorageBackend<T>(backend);
 		}
 
+		/// <summary>
+		/// Attempts to open a read stream for the given path.
+		/// </summary>
+		/// <param name="storageBackend"></param>
+		/// <param name="path">Relative path within the bucket</param>
+		/// <param name="cancellationToken">Cancellation token for the operation</param>
+		/// <returns></returns>
+		public static Task<Stream> ReadAsync(this IStorageBackend storageBackend, string path, CancellationToken cancellationToken = default) => storageBackend.ReadAsync(path, 0, null, cancellationToken);
+
 		/// <summary>
 		/// Writes a block of memory to storage
 		/// </summary>
diff --git a/Engine/Source/Programs/Horde/Horde.Server/Storage/StorageBackendProvider.cs b/Engine/Source/Programs/Horde/Horde.Server/Storage/StorageBackendProvider.cs
--- a/Engine/Source/Programs/Horde/Horde.Server/Storage/StorageBackendProvider.cs
+++ b/Engine/Source/Programs/Horde/Horde.Server/Storage/StorageBackendProvider.cs
@@ -77,16 +77,13 @@ namespace Horde.Server.Storage
 			public IAsyncEnumerable<string> EnumerateAsync(CancellationToken cancellationToken = default) => _backend.EnumerateAsync(cancellationToken);
 
 			/// <inheritdoc/>
-			public Task<Stream> ReadAsync(string path, CancellationToken cancellationToken = default) => _backend.ReadAsync(path, cancellationToken);
-
-			/// <inheritdoc/>
-			public Task<Stream> ReadAsync(string path, int offset, int length, CancellationToken cancellationToken = default) => _backend.ReadAsync(path, offset, length, cancellationToken);
+			public Task<Stream> ReadAsync(string path, int offset, int? length, CancellationToken cancellationToken = default) => _backend.ReadAsync(path, offset, length, cancellationToken);
 
 			/// <inheritdoc/>
 			public Task WriteAsync(string path, Stream stream, CancellationToken cancellationToken = default) => _backend.WriteAsync(path, stream, cancellationToken);
 
 			/// <inheritdoc/>
-			public ValueTask<Uri?> TryGetReadRedirectAsync(string path, int? offset = null, int? length = null, CancellationToken cancellationToken = default) => _backend.TryGetReadRedirectAsync(path, offset, length, cancellationToken);
+			public ValueTask<Uri?> TryGetReadRedirectAsync(string path, CancellationToken cancellationToken = default) => _backend.TryGetReadRedirectAsync(path, cancellationToken);
 
 			/// <inheritdoc/>
 			public ValueTask<Uri?> TryGetWriteRedirectAsync(string path, CancellationToken cancellationToken = default) => _backend.TryGetWriteRedirectAsync(path, cancellationToken);
diff --git a/Engine/Source/Programs/Horde/Horde.Server/Storage/StorageController.cs b/Engine/Source/Programs/Horde/Horde.Server/Storage/StorageController.cs
--- a/Engine/Source/Programs/Horde/Horde.Server/Storage/StorageController.cs
+++ b/Engine/Source/Programs/Horde/Horde.Server/Storage/StorageController.cs
@@ -266,39 +266,71 @@ namespace Horde.Server.Storage
 				return Forbid(StorageAclAction.ReadBlobs, namespaceId);
 			}
 
-			return await ReadBlobInternalAsync(client, locator, offset, length, cancellationToken);
+			return await ReadBlobInternalAsync(client, locator, Request.Headers, cancellationToken);
 		}
 
 		/// <summary>
 		/// Reads a blob from storage, without performing namespace access checks.
 		/// </summary>
-		internal static async Task<ActionResult> ReadBlobInternalAsync(BundleStorageClient storageClient, BundleLocator locator, int? offset, int? length, CancellationToken cancellationToken)
+		internal static async Task<ActionResult> ReadBlobInternalAsync(BundleStorageClient storageClient, BundleLocator locator, IHeaderDictionary headers, CancellationToken cancellationToken)
 		{
 			if (storageClient is StorageClient storageClientImpl)
 			{
-				Uri? redirectUrl = await storageClientImpl.GetReadRedirectAsync(locator, offset, length, cancellationToken);
+				Uri? redirectUrl = await storageClientImpl.GetReadRedirectAsync(locator, cancellationToken);
 				if (redirectUrl != null)
 				{
 					return new RedirectResult(redirectUrl.ToString());
 				}
 			}
 
+			// Parse the range header
+			int offset = 0;
+			int? length = null;
+
+			if (headers.Range.Count > 0)
+			{
+				if (headers.Range.Count > 1)
+				{
+					return new BadRequestObjectResult(LogEvent.Create(LogLevel.Error, "Unsupported range header; only one range is allowed"));
+				}
+
+				string value = headers.Range[0];
+
+				Match match = Regex.Match(value, @"^\s*bytes\s*=\s*(\d*)-(\d*)$");
+				if (!match.Success)
+				{
+					return new BadRequestObjectResult(LogEvent.Create(LogLevel.Error, "Unsupported range header syntax; cannot parse {Value}", value));
+				}
+
+				if (match.Groups[1].Length > 0 && !Int32.TryParse(match.Groups[1].Value, out offset))
+				{
+					return new BadRequestObjectResult(LogEvent.Create(LogLevel.Error, "Unable to parse start for range: {Value}", value));
+				}
+				if (match.Groups[2].Length > 0)
+				{
+					int end;
+					if (Int32.TryParse(match.Groups[2].Value, out end) && end > offset)
+					{
+						length = (end + 1) - offset;
+					}
+					else
+					{
+						return new BadRequestObjectResult(LogEvent.Create(LogLevel.Error, "Unable to parse end for range: {Value}", value));
+					}
+				}
+			}
+
 #pragma warning disable CA2000 // Dispose objects before losing scope
-			// TODO: would be better to use the range header here, but seems to require a lot of plumbing to convert unseekable AWS streams into a format that works with range processing.
 			Stream stream;
-			if (offset == null && length == null)
+			if (offset == 0 && length == null)
 			{
 				Bundle bundle = await storageClient.ReadBundleAsync(locator, cancellationToken);
 				stream = new ReadOnlySequenceStream(bundle.AsSequence());
 			}
-			else if (offset != null && length != null)
-			{
-				ReadOnlyMemory<byte> memory = await storageClient.ReadBundleRangeAsync(locator, offset.Value, length.Value, cancellationToken);
-				stream = new ReadOnlyMemoryStream(memory);
-			}
 			else
 			{
-				return new BadRequestObjectResult("Offset and length must both be specified as query parameters for ranged reads");
+				ReadOnlyMemory<byte> memory = await storageClient.ReadBundleRangeAsync(locator, offset, length, cancellationToken);
+				stream = new ReadOnlyMemoryStream(memory);
 			}
 			return new FileStreamResult(stream, "application/octet-stream");
 #pragma warning restore CA2000 // Dispose objects before losing scope
diff --git a/Engine/Source/Programs/Horde/Horde.Server/Storage/StorageService.cs b/Engine/Source/Programs/Horde/Horde.Server/Storage/StorageService.cs
--- a/Engine/Source/Programs/Horde/Horde.Server/Storage/StorageService.cs
+++ b/Engine/Source/Programs/Horde/Horde.Server/Storage/StorageService.cs
@@ -98,11 +98,9 @@ namespace Horde.Server.Storage
 		/// Gets a redirect for a read request
 		/// </summary>
 		/// <param name="locator">Locator for the blob</param>
-		/// <param name="offset">Offset of the data to return</param>
-		/// <param name="length">Length of the data to return</param>
 		/// <param name="cancellationToken">Cancellation token for the operation</param>
 		/// <returns>Path to upload the data to</returns>
-		public abstract ValueTask<Uri?> GetReadRedirectAsync(BundleLocator locator, int? offset, int? length, CancellationToken cancellationToken = default);
+		public abstract ValueTask<Uri?> GetReadRedirectAsync(BundleLocator locator, CancellationToken cancellationToken = default);
 
 		/// <summary>
 		/// Gets a redirect for a write request
@@ -155,10 +153,10 @@ namespace Horde.Server.Storage
 			}
 
 			/// <inheritdoc/>
-			public override ValueTask<Uri?> GetReadRedirectAsync(BundleLocator locator, int? offset = null, int? length = null, CancellationToken cancellationToken = default) => Backend.TryGetReadRedirectAsync(GetBlobPath(locator), offset, length, cancellationToken);
+			public override ValueTask<Uri?> GetReadRedirectAsync(BundleLocator locator, CancellationToken cancellationToken = default) => Backend.TryGetReadRedirectAsync(GetBlobPath(locator), cancellationToken);
 
 			/// <inheritdoc/>
-			public override async Task<ReadOnlyMemory<byte>> ReadBundleRangeAsync(BundleLocator locator, int offset, int length, CancellationToken cancellationToken = default)
+			public override async Task<ReadOnlyMemory<byte>> ReadBundleRangeAsync(BundleLocator locator, int offset, int? length, CancellationToken cancellationToken = default)
 			{
 				using TelemetrySpan span = _tracer.StartActiveSpan($"{nameof(StorageService)}.{nameof(StorageClientImpl)}.{nameof(ReadBundleRangeAsync)}");
 				span.SetAttribute("locator", locator.ToString());
@@ -168,13 +166,7 @@ namespace Horde.Server.Storage
 				string path = GetBlobPath(locator);
 				await using Stream stream = await Backend.ReadAsync(path, offset, length, cancellationToken);
 
-				ReadOnlyMemory<byte> data = await stream.ReadAllBytesAsync(cancellationToken);
-				if (data.Length > length)
-				{
-					_logger.LogDebug("Storage backend returned more data than requested; truncating response (wanted {Length}, got {GotLength})", length, data.Length);
-					data = data.Slice(0, length);
-				}
-				return data;
+				return await stream.ReadAllBytesAsync(cancellationToken);
 			}
 
 			/// <inheritdoc/>
diff --git a/Engine/Source/Programs/Horde/Horde.Server/Tools/ToolsController.cs b/Engine/Source/Programs/Horde/Horde.Server/Tools/ToolsController.cs
--- a/Engine/Source/Programs/Horde/Horde.Server/Tools/ToolsController.cs
+++ b/Engine/Source/Programs/Horde/Horde.Server/Tools/ToolsController.cs
@@ -542,13 +542,11 @@ namespace Horde.Server.Tools
 		/// </summary>
 		/// <param name="id">Identifier of the tool to retrieve</param>
 		/// <param name="locator">The blob locator</param>
-		/// <param name="length">Length of data to return</param>
-		/// <param name="offset">Offset of the data to return</param>
 		/// <param name="cancellationToken">Cancellation token for the operation</param>
 		/// <returns>Information about all the artifacts</returns>
 		[HttpGet]
 		[Route("/api/v1/tools/{id}/blobs/{*locator}")]
-		public async Task<ActionResult<object>> ReadToolBlobAsync(ToolId id, BundleLocator locator, [FromQuery] int? offset = null, [FromQuery] int? length = null, CancellationToken cancellationToken = default)
+		public async Task<ActionResult<object>> ReadToolBlobAsync(ToolId id, BundleLocator locator, CancellationToken cancellationToken = default)
 		{
 			ITool? tool = await _toolCollection.GetAsync(id, _globalConfig.Value);
 			if (tool == null)
@@ -566,7 +564,7 @@ namespace Horde.Server.Tools
 			}
 
 			BundleStorageClient storageClient = await _toolCollection.GetStorageClientAsync(tool, cancellationToken);
-			return StorageController.ReadBlobInternalAsync(storageClient, locator, offset, length, cancellationToken);
+			return StorageController.ReadBlobInternalAsync(storageClient, locator, Request.Headers, cancellationToken);
 		}
 
 		bool AuthorizeDownload(ToolConfig toolConfig)
diff --git a/Engine/Source/Programs/Shared/EpicGames.Horde/Compute/AgentMessage.cs b/Engine/Source/Programs/Shared/EpicGames.Horde/Compute/AgentMessage.cs
--- a/Engine/Source/Programs/Shared/EpicGames.Horde/Compute/AgentMessage.cs
+++ b/Engine/Source/Programs/Shared/EpicGames.Horde/Compute/AgentMessage.cs
@@ -578,13 +578,13 @@ namespace EpicGames.Horde.Compute
 		/// <param name="length">Length of data to return</param>
 		/// <param name="cancellationToken">Cancellation token for the operation</param>
 		/// <returns>Stream containing the blob data</returns>
-		public static async Task<ReadOnlyMemory<byte>> ReadBlobAsync(this AgentMessageChannel channel, BundleLocator locator, int offset, int length, CancellationToken cancellationToken = default)
+		public static async Task<ReadOnlyMemory<byte>> ReadBlobAsync(this AgentMessageChannel channel, BundleLocator locator, int offset, int? length, CancellationToken cancellationToken = default)
 		{
 			using (IAgentMessageBuilder request = await channel.CreateMessageAsync(AgentMessageType.ReadBlob, cancellationToken))
 			{
 				request.WriteBlobLocator(locator);
 				request.WriteUnsignedVarInt(offset);
-				request.WriteUnsignedVarInt(length);
+				request.WriteUnsignedVarInt(length ?? Int32.MaxValue);
 				request.Send();
 			}
 
diff --git a/Engine/Source/Programs/Shared/EpicGames.Horde/Compute/AgentStorageClient.cs b/Engine/Source/Programs/Shared/EpicGames.Horde/Compute/AgentStorageClient.cs
--- a/Engine/Source/Programs/Shared/EpicGames.Horde/Compute/AgentStorageClient.cs
+++ b/Engine/Source/Programs/Shared/EpicGames.Horde/Compute/AgentStorageClient.cs
@@ -47,7 +47,7 @@ namespace EpicGames.Horde.Compute
 		}
 
 		/// <inheritdoc/>
-		public override async Task<ReadOnlyMemory<byte>> ReadBundleRangeAsync(BundleLocator locator, int offset, int length, CancellationToken cancellationToken = default)
+		public override async Task<ReadOnlyMemory<byte>> ReadBundleRangeAsync(BundleLocator locator, int offset, int? length, CancellationToken cancellationToken = default)
 		{
 			await _semaphore.WaitAsync(cancellationToken);
 			try
diff --git a/Engine/Source/Programs/Shared/EpicGames.Horde/Storage/Backends/FileStorageClient.cs b/Engine/Source/Programs/Shared/EpicGames.Horde/Storage/Backends/FileStorageClient.cs
--- a/Engine/Source/Programs/Shared/EpicGames.Horde/Storage/Backends/FileStorageClient.cs
+++ b/Engine/Source/Programs/Shared/EpicGames.Horde/Storage/Backends/FileStorageClient.cs
@@ -65,14 +65,14 @@ namespace EpicGames.Horde.Storage.Backends
 		}
 
 		/// <inheritdoc/>
-		public override async Task<ReadOnlyMemory<byte>> ReadBundleRangeAsync(BundleLocator id, int offset, int length, CancellationToken cancellationToken = default)
+		public override async Task<ReadOnlyMemory<byte>> ReadBundleRangeAsync(BundleLocator id, int offset, int? length, CancellationToken cancellationToken = default)
 		{
 			Bundle bundle = await ReadBundleAsync(id, cancellationToken);
 
 			ReadOnlySequence<byte> sequence = bundle.AsSequence().Slice(offset);
-			if (sequence.Length > length)
+			if (length.HasValue && sequence.Length > length)
 			{
-				sequence = sequence.Slice(0, length);
+				sequence = sequence.Slice(0, length.Value);
 			}
 
 			return sequence.AsSingleSegment();
diff --git a/Engine/Source/Programs/Shared/EpicGames.Horde/Storage/Backends/HttpStorageClient.cs b/Engine/Source/Programs/Shared/EpicGames.Horde/Storage/Backends/HttpStorageClient.cs
--- a/Engine/Source/Programs/Shared/EpicGames.Horde/Storage/Backends/HttpStorageClient.cs
+++ b/Engine/Source/Programs/Shared/EpicGames.Horde/Storage/Backends/HttpStorageClient.cs
@@ -120,15 +120,18 @@ namespace EpicGames.Horde.Storage.Backends
 		}
 
 		/// <inheritdoc/>
-		public override async Task<ReadOnlyMemory<byte>> ReadBundleRangeAsync(BundleLocator locator, int offset, int length, CancellationToken cancellationToken = default)
+		public override async Task<ReadOnlyMemory<byte>> ReadBundleRangeAsync(BundleLocator locator, int offset, int? length, CancellationToken cancellationToken = default)
 		{
 			_logger.LogDebug("Reading {Locator} ({Offset}+{Length})", locator, offset, length);
 			using (HttpClient httpClient = _createClient())
 			{
-				using (HttpRequestMessage request = new HttpRequestMessage(HttpMethod.Get, $"blobs/{locator}?offset={offset}&length={length}"))
+				using (HttpRequestMessage request = new HttpRequestMessage(HttpMethod.Get, $"blobs/{locator}"))
 				{
+					request.Headers.Range = new RangeHeaderValue(offset, (length == null)? null : (offset + (length.Value - 1)));
+
 					HttpResponseMessage response = await httpClient.SendAsync(request, cancellationToken);
 					response.EnsureSuccessStatusCode();
+
 					return await response.Content.ReadAsByteArrayAsync(cancellationToken);
 				}
 			}
diff --git a/Engine/Source/Programs/Shared/EpicGames.Horde/Storage/Backends/MemoryStorageClient.cs b/Engine/Source/Programs/Shared/EpicGames.Horde/Storage/Backends/MemoryStorageClient.cs
--- a/Engine/Source/Programs/Shared/EpicGames.Horde/Storage/Backends/MemoryStorageClient.cs
+++ b/Engine/Source/Programs/Shared/EpicGames.Horde/Storage/Backends/MemoryStorageClient.cs
@@ -62,12 +62,12 @@ namespace EpicGames.Horde.Storage.Backends
 		}
 
 		/// <inheritdoc/>
-		public override Task<ReadOnlyMemory<byte>> ReadBundleRangeAsync(BundleLocator locator, int offset, int length, CancellationToken cancellationToken = default)
+		public override Task<ReadOnlyMemory<byte>> ReadBundleRangeAsync(BundleLocator locator, int offset, int? length, CancellationToken cancellationToken = default)
 		{
 			ReadOnlySequence<byte> sequence = _bundles[locator].AsSequence().Slice(offset);
-			if (sequence.Length > length)
+			if (length != null && sequence.Length > length)
 			{
-				sequence = sequence.Slice(0, length);
+				sequence = sequence.Slice(0, length.Value);
 			}
 			return Task.FromResult(sequence.AsSingleSegment());
 		}
diff --git a/Engine/Source/Programs/Shared/EpicGames.Horde/Storage/Bundles/BundleStorageClient.cs b/Engine/Source/Programs/Shared/EpicGames.Horde/Storage/Bundles/BundleStorageClient.cs
--- a/Engine/Source/Programs/Shared/EpicGames.Horde/Storage/Bundles/BundleStorageClient.cs
+++ b/Engine/Source/Programs/Shared/EpicGames.Horde/Storage/Bundles/BundleStorageClient.cs
@@ -36,7 +36,7 @@ namespace EpicGames.Horde.Storage.Bundles
 		public abstract Task<Bundle> ReadBundleAsync(BundleLocator locator, CancellationToken cancellationToken = default);
 
 		/// <inheritdoc/>
-		public abstract Task<ReadOnlyMemory<byte>> ReadBundleRangeAsync(BundleLocator locator, int offset, int length, CancellationToken cancellationToken = default);
+		public abstract Task<ReadOnlyMemory<byte>> ReadBundleRangeAsync(BundleLocator locator, int offset, int? length, CancellationToken cancellationToken = default);
 
 		/// <summary>
 		/// Utility method to read a blob into a buffer
