Horde: Fix test failure.

diff --git a/Engine/Source/Programs/Horde/Horde.Server/Storage/Backends/MemoryStorageBackend.cs b/Engine/Source/Programs/Horde/Horde.Server/Storage/Backends/MemoryStorageBackend.cs
--- a/Engine/Source/Programs/Horde/Horde.Server/Storage/Backends/MemoryStorageBackend.cs
+++ b/Engine/Source/Programs/Horde/Horde.Server/Storage/Backends/MemoryStorageBackend.cs
@@ -36,7 +36,7 @@ namespace Horde.Server.Storage.Backends
 		public Task<Stream> ReadAsync(string path, int offset, int? length, CancellationToken cancellationToken)
 		{
 			ReadOnlyMemory<byte> data = _pathToData[path].AsMemory(offset);
-			if (length != null)
+			if (length != null && length.Value < data.Length)
 			{
 				data = data.Slice(0, length.Value);
 			}
