Horde: Export all the possible ACL values when generating schema documentation.

diff --git a/Engine/Source/Programs/Horde/Horde.Server/Agents/Pools/PoolAclAction.cs b/Engine/Source/Programs/Horde/Horde.Server/Agents/Pools/PoolAclAction.cs
--- a/Engine/Source/Programs/Horde/Horde.Server/Agents/Pools/PoolAclAction.cs
+++ b/Engine/Source/Programs/Horde/Horde.Server/Agents/Pools/PoolAclAction.cs
@@ -12,26 +12,26 @@ namespace Horde.Server.Agents.Pools
 		/// <summary>
 		/// Create a global pool of agents
 		/// </summary>
-		public static readonly AclAction CreatePool = new AclAction("CreatePool");
+		public static AclAction CreatePool { get; } = new AclAction("CreatePool");
 
 		/// <summary>
 		/// Modify an agent pool
 		/// </summary>
-		public static readonly AclAction UpdatePool = new AclAction("UpdatePool");
+		public static AclAction UpdatePool { get; } = new AclAction("UpdatePool");
 
 		/// <summary>
 		/// Delete an agent pool
 		/// </summary>
-		public static readonly AclAction DeletePool = new AclAction("DeletePool");
+		public static AclAction DeletePool { get; } = new AclAction("DeletePool");
 
 		/// <summary>
 		/// Ability to view a pool
 		/// </summary>
-		public static readonly AclAction ViewPool = new AclAction("ViewPool");
+		public static AclAction ViewPool { get; } = new AclAction("ViewPool");
 
 		/// <summary>
 		/// View all the available agent pools
 		/// </summary>
-		public static readonly AclAction ListPools = new AclAction("ListPools");
+		public static AclAction ListPools { get; } = new AclAction("ListPools");
 	}
 }
diff --git a/Engine/Source/Programs/Horde/Horde.Server/Agents/Sessions/SessionAclAction.cs b/Engine/Source/Programs/Horde/Horde.Server/Agents/Sessions/SessionAclAction.cs
--- a/Engine/Source/Programs/Horde/Horde.Server/Agents/Sessions/SessionAclAction.cs
+++ b/Engine/Source/Programs/Horde/Horde.Server/Agents/Sessions/SessionAclAction.cs
@@ -12,11 +12,11 @@ namespace Horde.Server.Agents.Sessions
 		/// <summary>
 		/// Granted to agents to call CreateSession, which returns a bearer token identifying themselves valid to call UpdateSesssion via gRPC.
 		/// </summary>
-		public static readonly AclAction CreateSession = new AclAction("CreateSession");
+		public static AclAction CreateSession { get; } = new AclAction("CreateSession");
 
 		/// <summary>
 		/// Allows viewing information about an agent session
 		/// </summary>
-		public static readonly AclAction ViewSession = new AclAction("ViewSession");
+		public static AclAction ViewSession { get; } = new AclAction("ViewSession");
 	}
 }
diff --git a/Engine/Source/Programs/Horde/Horde.Server/Commands/Config/DocsCommand.cs b/Engine/Source/Programs/Horde/Horde.Server/Commands/Config/DocsCommand.cs
--- a/Engine/Source/Programs/Horde/Horde.Server/Commands/Config/DocsCommand.cs
+++ b/Engine/Source/Programs/Horde/Horde.Server/Commands/Config/DocsCommand.cs
@@ -8,7 +8,9 @@ using System.Linq;
 using System.Reflection;
 using System.Text.RegularExpressions;
 using System.Threading.Tasks;
+using System.Xml;
 using EpicGames.Core;
+using Horde.Server.Acls;
 using Horde.Server.Projects;
 using Horde.Server.Server;
 using Horde.Server.Streams;
@@ -44,10 +46,10 @@ namespace Horde.Server.Commands.Config
 
 			Dictionary<JsonSchemaType, PageInfo> typeToPageInfo = new Dictionary<JsonSchemaType, PageInfo>
 			{
-				[serverSchema.RootType] = new PageInfo("appsettings.json (Server)", "Horde (../Home.md) > Deployment (../Deployment.md) > Server (Server.md)", "Deployment/ServerSettings.md"),
-				[globalSchema.RootType] = new PageInfo("Globals.json", "Horde (../Home.md) > Configuration (../Config.md)", "Config/Schema/Globals.md"),
-				[projectSchema.RootType] = new PageInfo("*.project.json", "Horde (../Home.md) > Configuration (../Config.md)", "Config/Schema/Projects.md"),
-				[streamSchema.RootType] = new PageInfo("*.stream.json", "Horde (../Home.md) > Configuration (../Config.md)", "Config/Schema/Streams.md"),
+				[serverSchema.RootType] = new PageInfo("appsettings.json (Server)", "[Horde](../Home.md) > [Deployment](../Deployment.md) > Server (Server.md)", "Deployment/ServerSettings.md"),
+				[globalSchema.RootType] = new PageInfo("Globals.json", "[Horde](../Home.md) > [Configuration](../Config.md)", "Config/Schema/Globals.md"),
+				[projectSchema.RootType] = new PageInfo("*.project.json", "[Horde](../Home.md) > [Configuration](../Config.md)", "Config/Schema/Projects.md"),
+				[streamSchema.RootType] = new PageInfo("*.stream.json", "[Horde](../Home.md) > [Configuration](../Config.md)", "Config/Schema/Streams.md"),
 			};
 
 			if (Agent == null)
@@ -69,9 +71,99 @@ namespace Horde.Server.Commands.Config
 				await WriteDocAsync(type, pageInfo.Title, pageInfo.FileName, pageInfo.LinkRail, typeNameToPageName, logger);
 			}
 
+			await WriteAclDocAsync(Assembly.GetExecutingAssembly(), logger);
 			return 0;
 		}
 
+		async Task WriteAclDocAsync(Assembly assembly, ILogger logger)
+		{
+			FileReference file = FileReference.Combine(OutputDir, "Config/Schema/AclActions.md");
+			DirectoryReference.CreateDirectory(file.Directory);
+
+			using (FileStream stream = FileReference.Open(file, FileMode.Create, FileAccess.Write))
+			{
+				using (StreamWriter writer = new StreamWriter(stream))
+				{
+					await writer.WriteLineAsync($"[Horde](../Home.md) > [Configuration](../Config.md) > ACL Actions");
+					await writer.WriteLineAsync();
+					await writer.WriteLineAsync($"# ACL Actions");
+					await writer.WriteLineAsync();
+
+					Dictionary<string, List<PropertyInfo>> categoryToProperties = new Dictionary<string, List<PropertyInfo>>();
+
+					List<PropertyInfo> newProperties = new List<PropertyInfo>();
+					foreach (Type type in assembly.GetTypes())
+					{
+						if (type.IsClass)
+						{
+							foreach (PropertyInfo propertyInfo in type.GetProperties(BindingFlags.Static | BindingFlags.Public))
+							{
+								if (propertyInfo.PropertyType == typeof(AclAction))
+								{
+									newProperties.Add(propertyInfo);
+								}
+							}
+							if (newProperties.Count > 0)
+							{
+								string category = type.Namespace ?? "Default";
+
+								int topNamespaceIdx = category.LastIndexOf('.');
+								if (topNamespaceIdx != -1)
+								{
+									category = category.Substring(topNamespaceIdx + 1);
+								}
+
+								List<PropertyInfo>? existingProperties;
+								if (categoryToProperties.TryGetValue(category, out existingProperties))
+								{
+									existingProperties.AddRange(newProperties);
+									newProperties.Clear();
+								}
+								else
+								{
+									categoryToProperties.Add(category, newProperties);
+									newProperties = new List<PropertyInfo>();
+								}
+							}
+						}
+					}
+
+					XmlDocument documentation = new XmlDocument();
+
+					FileReference inputDocumentationFile = new FileReference(assembly.Location).ChangeExtension(".xml");
+					if (FileReference.Exists(inputDocumentationFile))
+					{
+						documentation.Load(inputDocumentationFile.FullName);
+					}
+
+					foreach ((string category, List<PropertyInfo> properties) in categoryToProperties.OrderBy(x => x.Key))
+					{
+						await writer.WriteLineAsync();
+						await writer.WriteLineAsync($"## {category}");
+						await writer.WriteLineAsync();
+						await writer.WriteLineAsync("| Name | Description |");
+						await writer.WriteLineAsync("| ---- | ----------- |");
+						foreach (PropertyInfo property in properties)
+						{
+							AclAction action = (AclAction)property.GetValue(null)!;
+							string description = String.Empty;
+
+							string selector = $"//member[@name='P:{property.DeclaringType!.FullName}.{property.Name}']/summary";
+							XmlNode? node = documentation.SelectSingleNode(selector);
+							if (node != null)
+							{
+								description = node.InnerText.Trim().Replace("\r\n", "\n", StringComparison.Ordinal);
+							}
+
+							await writer.WriteLineAsync($"| `{action.Name}` | {description} |");
+						}
+					}
+				}
+			}
+
+			logger.LogInformation("Written {File}", file);
+		}
+
 		class ObjectQueue
 		{
 			readonly Stack<JsonSchemaType> _stack = new Stack<JsonSchemaType>();
diff --git a/Engine/Source/Programs/Horde/Horde.Server/Jobs/Bisect/BisectTaskAclAction.cs b/Engine/Source/Programs/Horde/Horde.Server/Jobs/Bisect/BisectTaskAclAction.cs
--- a/Engine/Source/Programs/Horde/Horde.Server/Jobs/Bisect/BisectTaskAclAction.cs
+++ b/Engine/Source/Programs/Horde/Horde.Server/Jobs/Bisect/BisectTaskAclAction.cs
@@ -12,16 +12,16 @@ namespace Horde.Server.Jobs.Bisect
 		/// <summary>
 		/// Ability to start new bisect tasks
 		/// </summary>
-		public static readonly AclAction CreateBisectTask = new AclAction("CreateBisectTask");
+		public static AclAction CreateBisectTask { get; } = new AclAction("CreateBisectTask");
 
 		/// <summary>
 		/// Ability to update a bisect task
 		/// </summary>
-		public static readonly AclAction UpdateBisectTask = new AclAction("UpdateBisectTask");
+		public static AclAction UpdateBisectTask { get; } = new AclAction("UpdateBisectTask");
 
 		/// <summary>
 		/// Ability to view a bisect task
 		/// </summary>
-		public static readonly AclAction ViewBisectTask = new AclAction("ViewBisectTask");
+		public static AclAction ViewBisectTask { get; } = new AclAction("ViewBisectTask");
 	}
 }
diff --git a/Engine/Source/Programs/Horde/Horde.Server/Jobs/JobAclAction.cs b/Engine/Source/Programs/Horde/Horde.Server/Jobs/JobAclAction.cs
--- a/Engine/Source/Programs/Horde/Horde.Server/Jobs/JobAclAction.cs
+++ b/Engine/Source/Programs/Horde/Horde.Server/Jobs/JobAclAction.cs
@@ -12,31 +12,31 @@ namespace Horde.Server.Jobs
 		/// <summary>
 		/// Ability to start new jobs
 		/// </summary>
-		public static readonly AclAction CreateJob = new AclAction("CreateJob");
+		public static AclAction CreateJob { get; } = new AclAction("CreateJob");
 
 		/// <summary>
 		/// Rename a job, modify its priority, etc...
 		/// </summary>
-		public static readonly AclAction UpdateJob = new AclAction("UpdateJob");
+		public static AclAction UpdateJob { get; } = new AclAction("UpdateJob");
 
 		/// <summary>
 		/// Delete a job properties
 		/// </summary>
-		public static readonly AclAction DeleteJob = new AclAction("DeleteJob");
+		public static AclAction DeleteJob { get; } = new AclAction("DeleteJob");
 
 		/// <summary>
 		/// Allows updating a job metadata (name, changelist number, step properties, new groups, job states, etc...). Typically granted to agents. Not user facing.
 		/// </summary>
-		public static readonly AclAction ExecuteJob = new AclAction("ExecuteJob");
+		public static AclAction ExecuteJob { get; } = new AclAction("ExecuteJob");
 
 		/// <summary>
 		/// Ability to retry a failed job step
 		/// </summary>
-		public static readonly AclAction RetryJobStep = new AclAction("RetryJobStep");
+		public static AclAction RetryJobStep { get; } = new AclAction("RetryJobStep");
 
 		/// <summary>
 		/// Ability to view a job
 		/// </summary>
-		public static readonly AclAction ViewJob = new AclAction("ViewJob");
+		public static AclAction ViewJob { get; } = new AclAction("ViewJob");
 	}
 }
diff --git a/Engine/Source/Programs/Horde/Horde.Server/Notifications/NotificationAclAction.cs b/Engine/Source/Programs/Horde/Horde.Server/Notifications/NotificationAclAction.cs
--- a/Engine/Source/Programs/Horde/Horde.Server/Notifications/NotificationAclAction.cs
+++ b/Engine/Source/Programs/Horde/Horde.Server/Notifications/NotificationAclAction.cs
@@ -12,6 +12,6 @@ namespace Horde.Server.Notifications
 		/// <summary>
 		/// Ability to subscribe to notifications
 		/// </summary>
-		public static readonly AclAction CreateSubscription = new AclAction("CreateSubscription");
+		public static AclAction CreateSubscription { get; } = new AclAction("CreateSubscription");
 	}
 }
diff --git a/Engine/Source/Programs/Horde/Horde.Server/Storage/StorageAclAction.cs b/Engine/Source/Programs/Horde/Horde.Server/Storage/StorageAclAction.cs
--- a/Engine/Source/Programs/Horde/Horde.Server/Storage/StorageAclAction.cs
+++ b/Engine/Source/Programs/Horde/Horde.Server/Storage/StorageAclAction.cs
@@ -12,26 +12,26 @@ namespace Horde.Server.Storage
 		/// <summary>
 		/// Ability to read blobs from the storage service
 		/// </summary>
-		public static readonly AclAction ReadBlobs = new AclAction("ReadBlobs");
+		public static AclAction ReadBlobs { get; } = new AclAction("ReadBlobs");
 
 		/// <summary>
 		/// Ability to write blobs to the storage service
 		/// </summary>
-		public static readonly AclAction WriteBlobs = new AclAction("WriteBlobs");
+		public static AclAction WriteBlobs { get; } = new AclAction("WriteBlobs");
 
 		/// <summary>
 		/// Ability to read refs from the storage service
 		/// </summary>
-		public static readonly AclAction ReadRefs = new AclAction("ReadRefs");
+		public static AclAction ReadRefs { get; } = new AclAction("ReadRefs");
 
 		/// <summary>
 		/// Ability to write refs to the storage service
 		/// </summary>
-		public static readonly AclAction WriteRefs = new AclAction("WriteRefs");
+		public static AclAction WriteRefs { get; } = new AclAction("WriteRefs");
 
 		/// <summary>
 		/// Ability to delete refs
 		/// </summary>
-		public static readonly AclAction DeleteRefs = new AclAction("DeleteRefs");
+		public static AclAction DeleteRefs { get; } = new AclAction("DeleteRefs");
 	}
 }
