Horde: Improve documentation output. Add some missing comments, support reading XML documentation for assemblies other than the initial assembly, fix formatting of multi-line messages.

diff --git a/Engine/Source/Programs/Horde/Horde.Server/Agents/Fleet/IFleetManager.cs b/Engine/Source/Programs/Horde/Horde.Server/Agents/Fleet/IFleetManager.cs
--- a/Engine/Source/Programs/Horde/Horde.Server/Agents/Fleet/IFleetManager.cs
+++ b/Engine/Source/Programs/Horde/Horde.Server/Agents/Fleet/IFleetManager.cs
@@ -18,29 +18,29 @@ namespace Horde.Server.Agents.Fleet
 		/// Default fleet manager
 		/// </summary>
 		Default,
-		
+
 		/// <summary>
-		/// <see cref="NoOpFleetManager" />
+		/// No-op fleet manager.
 		/// </summary>
 		NoOp,
-		
+
 		/// <summary>
-		/// <see cref="AwsFleetManager" />
+		/// Fleet manager for handling AWS EC2 instances. Will create and/or terminate instances from scratch.
 		/// </summary>
 		Aws,
-		
+
 		/// <summary>
-		/// <see cref="AwsReuseFleetManager" />
+		/// Fleet manager for handling AWS EC2 instances. Will start already existing but stopped instances to reuse existing EBS disks.
 		/// </summary>
 		AwsReuse,
-		
+
 		/// <summary>
-		/// <see cref="AwsRecyclingFleetManager" />
+		/// Fleet manager for handling AWS EC2 instances. Will start already existing but stopped instances to reuse existing EBS disks.
 		/// </summary>
 		AwsRecycle,
-		
+
 		/// <summary>
-		/// <see cref="AwsAsgFleetManager" />
+		/// Fleet manager for handling AWS EC2 instances. Uses an EC2 auto-scaling group for controlling the number of running instances.
 		/// </summary>
 		AwsAsg
 	}
diff --git a/Engine/Source/Programs/Horde/Horde.Server/Commands/Config/DocsCommand.cs b/Engine/Source/Programs/Horde/Horde.Server/Commands/Config/DocsCommand.cs
--- a/Engine/Source/Programs/Horde/Horde.Server/Commands/Config/DocsCommand.cs
+++ b/Engine/Source/Programs/Horde/Horde.Server/Commands/Config/DocsCommand.cs
@@ -28,25 +28,24 @@
 		[CommandLine]
 		public FileReference? Agent { get; set; }
 
-		class AppSettings
-		{
-			public ServerSettings Horde { get; set; } = new ServerSettings();
-		}
+		record class PageInfo(string Title, string LinkRail, string FileName, string? Introduction = null);
 
-		record class PageInfo(string Title, string LinkRail, string FileName);
+		const string AppSettingsIntro = "All Horde-specific settings are stored in a root object called `horde`. Other .NET functionality may be configured using properties in the root of this file.";
 
 		public override async Task<int> ExecuteAsync(ILogger logger)
 		{
 			DirectoryReference.CreateDirectory(OutputDir);
 
-			JsonSchema serverSchema = Schemas.CreateSchema(typeof(AppSettings));
-			JsonSchema globalSchema = Schemas.CreateSchema(typeof(GlobalConfig));
-			JsonSchema projectSchema = Schemas.CreateSchema(typeof(ProjectConfig));
-			JsonSchema streamSchema = Schemas.CreateSchema(typeof(StreamConfig));
+			XmlDocReader xmlDocReader = new XmlDocReader();
 
+			JsonSchema serverSchema = JsonSchema.FromType(typeof(ServerSettings), xmlDocReader);
+			JsonSchema globalSchema = JsonSchema.FromType(typeof(GlobalConfig), xmlDocReader);
+			JsonSchema projectSchema = JsonSchema.FromType(typeof(ProjectConfig), xmlDocReader);
+			JsonSchema streamSchema = JsonSchema.FromType(typeof(StreamConfig), xmlDocReader);
+
 			Dictionary<JsonSchemaType, PageInfo> typeToPageInfo = new Dictionary<JsonSchemaType, PageInfo>
 			{
-				[serverSchema.RootType] = new PageInfo("appsettings.json (Server)", "[Horde](../Home.md) > [Deployment](../Deployment.md) > Server (Server.md)", "Deployment/ServerSettings.md"),
+				[serverSchema.RootType] = new PageInfo("appsettings.json (Server)", "[Horde](../Home.md) > [Deployment](../Deployment.md) > Server (Server.md)", "Deployment/ServerSettings.md", AppSettingsIntro),
 				[globalSchema.RootType] = new PageInfo("Globals.json", "[Horde](../Home.md) > [Configuration](../Config.md)", "Config/Schema/Globals.md"),
 				[projectSchema.RootType] = new PageInfo("*.project.json", "[Horde](../Home.md) > [Configuration](../Config.md)", "Config/Schema/Projects.md"),
 				[streamSchema.RootType] = new PageInfo("*.stream.json", "[Horde](../Home.md) > [Configuration](../Config.md)", "Config/Schema/Streams.md"),
@@ -61,109 +60,20 @@
 				Assembly agentAssembly = Assembly.LoadFile(Agent.FullName);
 				Type agentSettingsType = agentAssembly.GetType("Horde.Agent.AgentSettings")!;
 
-				JsonSchema agentSchema = Schemas.CreateSchema(agentSettingsType);
-				await WriteDocAsync(agentSchema.RootType, "appsettings.json (Agent)", "Deployment/AgentSettings.md", "Horde (../Home.md) > Deployment (../Deployment.md) > Agent (Agent.md)", new Dictionary<string, string>(), logger);
+				JsonSchema agentSchema = JsonSchema.FromType(agentSettingsType, xmlDocReader);
+				await WriteDocAsync(agentSchema.RootType, "appsettings.json (Agent)", "Deployment/AgentSettings.md", "Horde (../Home.md) > Deployment (../Deployment.md) > Agent (Agent.md)", AppSettingsIntro, new Dictionary<string, string>(), logger);
 			}
 
 			Dictionary<string, string> typeNameToPageName = typeToPageInfo.ToDictionary(x => x.Key.Name!, x => x.Value.FileName, StringComparer.Ordinal);
 			foreach ((JsonSchemaType type, PageInfo pageInfo) in typeToPageInfo)
 			{
-				await WriteDocAsync(type, pageInfo.Title, pageInfo.FileName, pageInfo.LinkRail, typeNameToPageName, logger);
+				await WriteDocAsync(type, pageInfo.Title, pageInfo.FileName, pageInfo.LinkRail, pageInfo.Introduction, typeNameToPageName, logger);
 			}
 
 			await WriteAclDocAsync(Assembly.GetExecutingAssembly(), logger);
 			return 0;
 		}
 
-		async Task WriteAclDocAsync(Assembly assembly, ILogger logger)
-		{
-			FileReference file = FileReference.Combine(OutputDir, "Config/Schema/AclActions.md");
-			DirectoryReference.CreateDirectory(file.Directory);
-
-			using (FileStream stream = FileReference.Open(file, FileMode.Create, FileAccess.Write))
-			{
-				using (StreamWriter writer = new StreamWriter(stream))
-				{
-					await writer.WriteLineAsync($"[Horde](../Home.md) > [Configuration](../Config.md) > ACL Actions");
-					await writer.WriteLineAsync();
-					await writer.WriteLineAsync($"# ACL Actions");
-					await writer.WriteLineAsync();
-
-					Dictionary<string, List<PropertyInfo>> categoryToProperties = new Dictionary<string, List<PropertyInfo>>();
-
-					List<PropertyInfo> newProperties = new List<PropertyInfo>();
-					foreach (Type type in assembly.GetTypes())
-					{
-						if (type.IsClass)
-						{
-							foreach (PropertyInfo propertyInfo in type.GetProperties(BindingFlags.Static | BindingFlags.Public))
-							{
-								if (propertyInfo.PropertyType == typeof(AclAction))
-								{
-									newProperties.Add(propertyInfo);
-								}
-							}
-							if (newProperties.Count > 0)
-							{
-								string category = type.Namespace ?? "Default";
-
-								int topNamespaceIdx = category.LastIndexOf('.');
-								if (topNamespaceIdx != -1)
-								{
-									category = category.Substring(topNamespaceIdx + 1);
-								}
-
-								List<PropertyInfo>? existingProperties;
-								if (categoryToProperties.TryGetValue(category, out existingProperties))
-								{
-									existingProperties.AddRange(newProperties);
-									newProperties.Clear();
-								}
-								else
-								{
-									categoryToProperties.Add(category, newProperties);
-									newProperties = new List<PropertyInfo>();
-								}
-							}
-						}
-					}
-
-					XmlDocument documentation = new XmlDocument();
-
-					FileReference inputDocumentationFile = new FileReference(assembly.Location).ChangeExtension(".xml");
-					if (FileReference.Exists(inputDocumentationFile))
-					{
-						documentation.Load(inputDocumentationFile.FullName);
-					}
-
-					foreach ((string category, List<PropertyInfo> properties) in categoryToProperties.OrderBy(x => x.Key))
-					{
-						await writer.WriteLineAsync();
-						await writer.WriteLineAsync($"## {category}");
-						await writer.WriteLineAsync();
-						await writer.WriteLineAsync("| Name | Description |");
-						await writer.WriteLineAsync("| ---- | ----------- |");
-						foreach (PropertyInfo property in properties)
-						{
-							AclAction action = (AclAction)property.GetValue(null)!;
-							string description = String.Empty;
-
-							string selector = $"//member[@name='P:{property.DeclaringType!.FullName}.{property.Name}']/summary";
-							XmlNode? node = documentation.SelectSingleNode(selector);
-							if (node != null)
-							{
-								description = node.InnerText.Trim().Replace("\r\n", "\n", StringComparison.Ordinal);
-							}
-
-							await writer.WriteLineAsync($"| `{action.Name}` | {description} |");
-						}
-					}
-				}
-			}
-
-			logger.LogInformation("Written {File}", file);
-		}
-
 		class ObjectQueue
 		{
 			readonly Stack<JsonSchemaType> _stack = new Stack<JsonSchemaType>();
@@ -188,7 +98,7 @@
 			public bool TryPop([NotNullWhen(true)] out JsonSchemaType? obj) => _stack.TryPop(out obj);
 		}
 
-		async Task WriteDocAsync(JsonSchemaType rootType, string title, string fileName, string linkRail, Dictionary<string, string> typeNameToLink, ILogger logger)
+		async Task WriteDocAsync(JsonSchemaType rootType, string title, string fileName, string linkRail, string? introduction, Dictionary<string, string> typeNameToLink, ILogger logger)
 		{
 			FileReference file = FileReference.Combine(OutputDir, fileName);
 			DirectoryReference.CreateDirectory(file.Directory);
@@ -201,6 +111,12 @@
 					await writer.WriteLineAsync();
 					await writer.WriteLineAsync($"# {title}");
 
+					if (introduction != null)
+					{
+						await writer.WriteLineAsync();
+						await writer.WriteLineAsync(introduction);
+					}
+
 					HashSet<string> visitedTypeNames = new HashSet<string>(typeNameToLink.Keys, StringComparer.Ordinal);
 					visitedTypeNames.Remove(rootType.Name!);
 
@@ -217,24 +133,27 @@
 
 						await writer.WriteLineAsync();
 
-						if (schemaType.Description != null)
+						if (schemaType.Description != null && (introduction == null || schemaType != rootType))
 						{
-							string description = Regex.Replace(schemaType.Description, "\r?\n", Environment.NewLine);
+							string description = schemaType.Description.ReplaceLineEndings();
 							await writer.WriteLineAsync(description);
 							await writer.WriteLineAsync();
 						}
 
 						if (schemaType is JsonSchemaObject schemaObj)
 						{
-							await writer.WriteLineAsync("Name | Type | Description");
-							await writer.WriteLineAsync("---- | ---- | -----------");
-
-							foreach (JsonSchemaProperty property in schemaObj.Properties)
+							if (schemaObj.Properties.Count > 0)
 							{
-								string name = property.CamelCaseName;
-								string type = GetMarkdownType(property.Type, typeNameToLink);
-								string description = GetMarkdownDescription(property.Description);
-								await writer.WriteLineAsync($"`{name}` | {type} | {description}");
+								await writer.WriteLineAsync("Name | Type | Description");
+								await writer.WriteLineAsync("---- | ---- | -----------");
+							
+								foreach (JsonSchemaProperty property in schemaObj.Properties)
+								{
+									string name = property.CamelCaseName;
+									string type = GetMarkdownType(property.Type, typeNameToLink);
+									string description = GetMarkdownDescription(property.Description);
+									await writer.WriteLineAsync($"`{name}` | {type} | {description}");
+								}
 							}
 						}
 						else if (schemaType is JsonSchemaEnum schemaEnum)
@@ -395,5 +314,95 @@
 			anchor = Regex.Replace(anchor, @"[^a-z0-9]+", "-");
 			return "#" + anchor.Trim('-');
 		}
+
+
+		async Task WriteAclDocAsync(Assembly assembly, ILogger logger)
+		{
+			FileReference file = FileReference.Combine(OutputDir, "Config/Schema/AclActions.md");
+			DirectoryReference.CreateDirectory(file.Directory);
+
+			using (FileStream stream = FileReference.Open(file, FileMode.Create, FileAccess.Write))
+			{
+				using (StreamWriter writer = new StreamWriter(stream))
+				{
+					await writer.WriteLineAsync($"[Horde](../Home.md) > [Configuration](../Config.md) > ACL Actions");
+					await writer.WriteLineAsync();
+					await writer.WriteLineAsync($"# ACL Actions");
+					await writer.WriteLineAsync();
+
+					Dictionary<string, List<PropertyInfo>> categoryToProperties = new Dictionary<string, List<PropertyInfo>>();
+
+					List<PropertyInfo> newProperties = new List<PropertyInfo>();
+					foreach (Type type in assembly.GetTypes())
+					{
+						if (type.IsClass)
+						{
+							foreach (PropertyInfo propertyInfo in type.GetProperties(BindingFlags.Static | BindingFlags.Public))
+							{
+								if (propertyInfo.PropertyType == typeof(AclAction))
+								{
+									newProperties.Add(propertyInfo);
+								}
+							}
+							if (newProperties.Count > 0)
+							{
+								string category = type.Namespace ?? "Default";
+
+								int topNamespaceIdx = category.LastIndexOf('.');
+								if (topNamespaceIdx != -1)
+								{
+									category = category.Substring(topNamespaceIdx + 1);
+								}
+
+								List<PropertyInfo>? existingProperties;
+								if (categoryToProperties.TryGetValue(category, out existingProperties))
+								{
+									existingProperties.AddRange(newProperties);
+									newProperties.Clear();
+								}
+								else
+								{
+									categoryToProperties.Add(category, newProperties);
+									newProperties = new List<PropertyInfo>();
+								}
+							}
+						}
+					}
+
+					XmlDocument documentation = new XmlDocument();
+
+					FileReference inputDocumentationFile = new FileReference(assembly.Location).ChangeExtension(".xml");
+					if (FileReference.Exists(inputDocumentationFile))
+					{
+						documentation.Load(inputDocumentationFile.FullName);
+					}
+
+					foreach ((string category, List<PropertyInfo> properties) in categoryToProperties.OrderBy(x => x.Key))
+					{
+						await writer.WriteLineAsync();
+						await writer.WriteLineAsync($"## {category}");
+						await writer.WriteLineAsync();
+						await writer.WriteLineAsync("| Name | Description |");
+						await writer.WriteLineAsync("| ---- | ----------- |");
+						foreach (PropertyInfo property in properties)
+						{
+							AclAction action = (AclAction)property.GetValue(null)!;
+							string description = String.Empty;
+
+							string selector = $"//member[@name='P:{property.DeclaringType!.FullName}.{property.Name}']/summary";
+							XmlNode? node = documentation.SelectSingleNode(selector);
+							if (node != null)
+							{
+								description = node.InnerText.Trim().Replace("\r\n", "\n", StringComparison.Ordinal);
+							}
+
+							await writer.WriteLineAsync($"| `{action.Name}` | {description} |");
+						}
+					}
+				}
+			}
+
+			logger.LogInformation("Written {File}", file);
+		}
 	}
 }
diff --git a/Engine/Source/Programs/Horde/Horde.Server/Program.cs b/Engine/Source/Programs/Horde/Horde.Server/Program.cs
--- a/Engine/Source/Programs/Horde/Horde.Server/Program.cs
+++ b/Engine/Source/Programs/Horde/Horde.Server/Program.cs
@@ -42,7 +42,7 @@ namespace Horde.Server
 				{
 					theme = AnsiConsoleTheme.Code;
 				}
-				return sinkConfig.Console(outputTemplate: "[{Timestamp:HH:mm:ss} {Level:w3}] {Indent}{Message:l}{NewLine}{Exception}", theme: theme, restrictedToMinimumLevel: settings.ConsoleLogLevel);
+				return sinkConfig.Console(outputTemplate: "[{Timestamp:HH:mm:ss} {Level:w3}] {Indent}{Message:l}{NewLine}{Exception}", theme: theme, restrictedToMinimumLevel: LogEventLevel.Debug);
 			}
 		}
 
diff --git a/Engine/Source/Programs/Horde/Horde.Server/ServerSettings.cs b/Engine/Source/Programs/Horde/Horde.Server/ServerSettings.cs
--- a/Engine/Source/Programs/Horde/Horde.Server/ServerSettings.cs
+++ b/Engine/Source/Programs/Horde/Horde.Server/ServerSettings.cs
@@ -254,7 +254,7 @@ namespace Horde.Server
 	}
 	
 	/// <summary>
-	/// Feature flags to aid rollout of new features
+	/// Feature flags to aid rollout of new features.
 	///
 	/// Once a feature is running in its intended state and is stable, the flag should be removed.
 	/// A name and date of when the flag was created is noted next to it to help encourage this behavior.
@@ -394,7 +394,9 @@ namespace Horde.Server
 		[JsonIgnore]
 		public AclScopeName ScopeName => AclScopeName.Root;
 
-		/// <inheritdoc cref="RunMode" />
+		/// <summary>
+		/// Modes that the server should run in. Runmodes can be used in a multi-server deployment to limit the operations that a particular instance will try to perform.
+		/// </summary>
 		public RunMode[]? RunModes { get; set; } = null;
 
 		/// <summary>
@@ -402,11 +404,6 @@ namespace Horde.Server
 		/// </summary>
 		public string? DataDir { get; set; } = null;
 
-		/// <summary>
-		/// Output level for console
-		/// </summary>
-		public LogEventLevel ConsoleLogLevel { get; set; } = LogEventLevel.Debug;
-
 		/// <summary>
 		/// Main port for serving HTTP. Uses the default Kestrel port (5000) if not specified.
 		/// </summary>
@@ -821,7 +818,9 @@ namespace Horde.Server
 		/// </summary>
 		public bool OpenBrowser { get; set; } = false;
 
-		/// <inheritdoc cref="FeatureFlags" />
+		/// <summary>
+		/// Experimental features to enable on the server.
+		/// </summary>
 		public FeatureFlagSettings FeatureFlags { get; set; } = new ();
 
 		/// <summary>
diff --git a/Engine/Source/Programs/Horde/Horde.Server/Utilities/Schemas.cs b/Engine/Source/Programs/Horde/Horde.Server/Utilities/Schemas.cs
--- a/Engine/Source/Programs/Horde/Horde.Server/Utilities/Schemas.cs
+++ b/Engine/Source/Programs/Horde/Horde.Server/Utilities/Schemas.cs
@@ -2,8 +2,6 @@
 
 using System;
 using System.Collections.Concurrent;
-using System.Reflection;
-using System.Xml;
 using EpicGames.Core;
 
 namespace Horde.Server.Utilities
@@ -13,7 +11,7 @@ namespace Horde.Server.Utilities
 	/// </summary>
 	static class Schemas
 	{
-		static readonly ConcurrentDictionary<Assembly, XmlDocument?> s_cachedDocumentation = new ConcurrentDictionary<Assembly, XmlDocument?>();
+		static readonly XmlDocReader s_cachedReader = new XmlDocReader();
 		static readonly ConcurrentDictionary<Type, JsonSchema> s_cachedSchemas = new ConcurrentDictionary<Type, JsonSchema>();
 
 		/// <summary>
@@ -26,19 +24,7 @@ namespace Horde.Server.Utilities
 			JsonSchema? schema;
 			if (!s_cachedSchemas.TryGetValue(type, out schema))
 			{
-				XmlDocument? documentation;
-				if (!s_cachedDocumentation.TryGetValue(type.Assembly, out documentation))
-				{
-					FileReference inputDocumentationFile = new FileReference(type.Assembly.Location).ChangeExtension(".xml");
-					if (FileReference.Exists(inputDocumentationFile))
-					{
-						documentation = new XmlDocument();
-						documentation.Load(inputDocumentationFile.FullName);
-					}
-					s_cachedDocumentation.TryAdd(type.Assembly, documentation);
-				}
-
-				schema = JsonSchema.FromType(type, documentation);
+				schema = JsonSchema.FromType(type, s_cachedReader);
 				s_cachedSchemas.TryAdd(type, schema);
 			}
 			return schema;
diff --git a/Engine/Source/Programs/Shared/EpicGames.Core/JsonSchema.cs b/Engine/Source/Programs/Shared/EpicGames.Core/JsonSchema.cs
--- a/Engine/Source/Programs/Shared/EpicGames.Core/JsonSchema.cs
+++ b/Engine/Source/Programs/Shared/EpicGames.Core/JsonSchema.cs
@@ -1,6 +1,7 @@
 // Copyright Epic Games, Inc. All Rights Reserved.
 
 using System;
+using System.Collections.Concurrent;
 using System.Collections.Generic;
 using System.IO;
 using System.Reflection;
@@ -280,66 +281,27 @@ namespace EpicGames.Core
 		/// Constructs a Json schema from a type
 		/// </summary>
 		/// <param name="type">The type to construct from</param>
-		/// <param name="xmlDoc">C# Xml documentation file, to use for property descriptions</param>
+		/// <param name="xmlDocReader">Reads documentation for requested types</param>
 		/// <returns>New schema object</returns>
-		public static JsonSchema FromType(Type type, XmlDocument? xmlDoc)
+		public static JsonSchema FromType(Type type, XmlDocReader? xmlDocReader)
 		{
 			JsonSchemaAttribute? schemaAttribute = type.GetCustomAttribute<JsonSchemaAttribute>();
-			return new JsonSchema(schemaAttribute?.Id, CreateSchemaType(type, new Dictionary<Type, JsonSchemaType>(), xmlDoc));
+			return new JsonSchema(schemaAttribute?.Id, CreateSchemaType(type, new Dictionary<Type, JsonSchemaType>(), xmlDocReader));
 		}
 
 		/// <summary>
-		/// Gets a property description from Xml documentation file
+		/// Create a Json schema (or retrieve a cached schema)
 		/// </summary>
-		/// <param name="type"></param>
-		/// <param name="name"></param>
-		/// <param name="xmlDoc"></param>
-		/// <returns></returns>
-		static string? GetPropertyDescription(Type type, string name, XmlDocument? xmlDoc)
+		public static JsonSchema CreateSchema(Type type, XmlDocReader? xmlDocReader)
 		{
-			return GetSummaryFromXmlDoc(type, "P", name, xmlDoc);
-		}
-
-		/// <summary>
-		/// Gets a description from Xml documentation file
-		/// </summary>
-		/// <param name="type">Type to retrieve summary for</param>
-		/// <param name="qualifier">Type of element to retrieve</param>
-		/// <param name="member">Name of the member</param>
-		/// <param name="xmlDoc">XML documentation file to search</param>
-		/// <returns>Summary string, or null if it's not available</returns>
-		static string? GetSummaryFromXmlDoc(Type type, string qualifier, string? member, XmlDocument? xmlDoc)
-		{
-			if (xmlDoc == null)
-			{
-				return null;
-			}
-
-			string? fullName = type.FullName;
-			if (member != null)
-			{
-				fullName = $"{fullName}.{member}";
-			}
-
-			string selector = $"//member[@name='{qualifier}:{fullName}']/summary";
-
-			XmlNode? node = xmlDoc.SelectSingleNode(selector);
-			if (node == null)
-			{
-				return null;
-			}
-
-			return node.InnerText.Trim().Replace("\r\n", "\n", StringComparison.Ordinal);
+			JsonSchemaAttribute? schemaAttribute = type.GetCustomAttribute<JsonSchemaAttribute>();
+			return new JsonSchema(schemaAttribute?.Id, CreateSchemaType(type, new Dictionary<Type, JsonSchemaType>(), xmlDocReader));
 		}
 
 		/// <summary>
 		/// Constructs a schema type from the given type object
 		/// </summary>
-		/// <param name="type"></param>
-		/// <param name="typeCache"></param>
-		/// <param name="xmlDoc"></param>
-		/// <returns></returns>
-		static JsonSchemaType CreateSchemaType(Type type, Dictionary<Type, JsonSchemaType> typeCache, XmlDocument? xmlDoc)
+		static JsonSchemaType CreateSchemaType(Type type, Dictionary<Type, JsonSchemaType> typeCache, XmlDocReader? xmlDocReader)
 		{
 			switch (Type.GetTypeCode(type))
 			{
@@ -355,7 +317,7 @@ namespace EpicGames.Core
 				case TypeCode.UInt64:
 					if (type.IsEnum)
 					{
-						return CreateEnumSchemaType(type, xmlDoc);
+						return CreateEnumSchemaType(type, xmlDocReader);
 					}
 					else
 					{
@@ -377,7 +339,7 @@ namespace EpicGames.Core
 
 			if (type.IsGenericType && type.GetGenericTypeDefinition() == typeof(Nullable<>))
 			{
-				return CreateSchemaType(type.GetGenericArguments()[0], typeCache, xmlDoc);
+				return CreateSchemaType(type.GetGenericArguments()[0], typeCache, xmlDocReader);
 			}
 			if (type == typeof(DateTime) || type == typeof(DateTimeOffset))
 			{
@@ -400,12 +362,12 @@ namespace EpicGames.Core
 					Type[] arguments = interfaceType.GetGenericArguments();
 					if (interfaceType.GetGenericTypeDefinition() == typeof(IList<>))
 					{
-						return new JsonSchemaArray(CreateSchemaType(arguments[0], typeCache, xmlDoc));
+						return new JsonSchemaArray(CreateSchemaType(arguments[0], typeCache, xmlDocReader));
 					}
 					if (interfaceType.GetGenericTypeDefinition() == typeof(IDictionary<,>))
 					{
 						JsonSchemaObject obj = new JsonSchemaObject();
-						obj.AdditionalProperties = CreateSchemaType(arguments[1], typeCache, xmlDoc);
+						obj.AdditionalProperties = CreateSchemaType(arguments[1], typeCache, xmlDocReader);
 						return obj;
 					}
 				}
@@ -423,14 +385,14 @@ namespace EpicGames.Core
 				{
 					JsonSchemaOneOf obj = new JsonSchemaOneOf();
 					typeCache[type] = obj;
-					SetOneOfProperties(obj, type, knownTypes.Types, typeCache, xmlDoc);
+					SetOneOfProperties(obj, type, knownTypes.Types, typeCache, xmlDocReader);
 					return obj;
 				}
 				else
 				{
 					JsonSchemaObject obj = new JsonSchemaObject();
 					typeCache[type] = obj;
-					SetObjectProperties(obj, type, typeCache, xmlDoc);
+					SetObjectProperties(obj, type, typeCache, xmlDocReader);
 					return obj;
 				}
 			}
@@ -438,21 +400,21 @@ namespace EpicGames.Core
 			throw new Exception($"Unknown type for schema generation: {type}");
 		}
 
-		static JsonSchemaEnum CreateEnumSchemaType(Type type, XmlDocument? xmlDoc)
+		static JsonSchemaEnum CreateEnumSchemaType(Type type, XmlDocReader? xmlDocReader)
 		{
 			string[] names = Enum.GetNames(type);
 			string[] descriptions = new string[names.Length];
 
 			for (int idx = 0; idx < names.Length; idx++)
 			{
-				descriptions[idx] = GetSummaryFromXmlDoc(type, "F", names[idx], xmlDoc) ?? String.Empty;
+				descriptions[idx] = xmlDocReader?.GetDescription(type, names[idx]) ?? String.Empty;
 			}
 
-			string? enumDescription = GetSummaryFromXmlDoc(type, "T", null, xmlDoc);
+			string? enumDescription = xmlDocReader?.GetDescription(type);
 			return new JsonSchemaEnum(names, descriptions) { Name = type.Name, Description = enumDescription };
 		}
 
-		static void SetOneOfProperties(JsonSchemaOneOf obj, Type type, Type[] knownTypes, Dictionary<Type, JsonSchemaType> typeCache, XmlDocument? xmlDoc)
+		static void SetOneOfProperties(JsonSchemaOneOf obj, Type type, Type[] knownTypes, Dictionary<Type, JsonSchemaType> typeCache, XmlDocReader? xmlDocReader)
 		{
 			obj.Name = type.Name;
 
@@ -463,24 +425,24 @@ namespace EpicGames.Core
 				{
 					JsonSchemaObject knownObject = new JsonSchemaObject();
 					knownObject.Properties.Add(new JsonSchemaProperty("type", "Type discriminator", new JsonSchemaEnum(new[] { attribute.Name }, new[] { "Identifier for the derived type" })));
-					SetObjectProperties(knownObject, knownType, typeCache, xmlDoc);
+					SetObjectProperties(knownObject, knownType, typeCache, xmlDocReader);
 					obj.Types.Add(knownObject);
 				}
 			}
 		}
 
-		static void SetObjectProperties(JsonSchemaObject obj, Type type, Dictionary<Type, JsonSchemaType> typeCache, XmlDocument? xmlDoc)
+		static void SetObjectProperties(JsonSchemaObject obj, Type type, Dictionary<Type, JsonSchemaType> typeCache, XmlDocReader? xmlDocReader)
 		{
 			obj.Name = type.Name;
-			obj.Description = GetSummaryFromXmlDoc(type, "T", null, xmlDoc);
+			obj.Description = xmlDocReader?.GetDescription(type);
 
 			PropertyInfo[] properties = type.GetProperties(BindingFlags.Instance | BindingFlags.Public);
 			foreach (PropertyInfo property in properties)
 			{
 				if (property.GetCustomAttribute<JsonIgnoreAttribute>() == null)
 				{
-					string? description = GetPropertyDescription(type, property.Name, xmlDoc);
-					JsonSchemaType propertyType = CreateSchemaType(property.PropertyType, typeCache, xmlDoc);
+					string? description = xmlDocReader?.GetDescription(property);
+					JsonSchemaType propertyType = CreateSchemaType(property.PropertyType, typeCache, xmlDocReader);
 					obj.Properties.Add(new JsonSchemaProperty(property.Name, description, propertyType));
 				}
 			}
diff --git a/Engine/Source/Programs/Shared/EpicGames.Core/XmlDocReader.cs b/Engine/Source/Programs/Shared/EpicGames.Core/XmlDocReader.cs
new file mode 100644
--- /dev/null
+++ b/Engine/Source/Programs/Shared/EpicGames.Core/XmlDocReader.cs
@@ -0,0 +1,142 @@
+// Copyright Epic Games, Inc. All Rights Reserved.
+
+using System;
+using System.Collections.Concurrent;
+using System.Collections.Generic;
+using System.Diagnostics.CodeAnalysis;
+using System.Reflection;
+using System.Text.RegularExpressions;
+using System.Xml;
+
+namespace EpicGames.Core
+{
+	/// <summary>
+	/// Cache for XML documentation
+	/// </summary>
+	public class XmlDocReader
+	{
+		readonly ConcurrentDictionary<Assembly, XmlDocument?> _cachedDocumentation = new ConcurrentDictionary<Assembly, XmlDocument?>();
+
+		/// <summary>
+		/// Adds a parsed file as the documentation for a particular assembly
+		/// </summary>
+		public void Add(Assembly assembly, XmlDocument? xmlDoc)
+		{
+			_cachedDocumentation[assembly] = xmlDoc;
+		}
+
+		/// <summary>
+		/// Gets a property description from Xml documentation file
+		/// </summary>
+		public string? GetDescription(Type type)
+		{
+			return GetSummaryFromXmlDoc(type, "T", null);
+		}
+
+		/// <summary>
+		/// Gets a property description from Xml documentation file
+		/// </summary>
+		public string? GetDescription(PropertyInfo propertyInfo)
+		{
+			return GetSummaryFromXmlDoc(propertyInfo.DeclaringType!, "P", propertyInfo.Name);
+		}
+
+		/// <summary>
+		/// Gets a property description from Xml documentation file
+		/// </summary>
+		public string? GetDescription(Type type, string field)
+		{
+			return GetSummaryFromXmlDoc(type, "F", field);
+		}
+
+		/// <summary>
+		/// Gets a description from Xml documentation file
+		/// </summary>
+		/// <param name="type">Type to retrieve summary for</param>
+		/// <param name="qualifier">Type of element to retrieve</param>
+		/// <param name="member">Name of the member</param>
+		/// <returns>Summary string, or null if it's not available</returns>
+		string? GetSummaryFromXmlDoc(Type type, string qualifier, string? member)
+		{
+			XmlDocument? xmlDoc;
+			if (!TryReadXmlDoc(type.Assembly, out xmlDoc))
+			{
+				return null;
+			}
+
+			string? fullName = type.FullName;
+			if (member != null)
+			{
+				fullName = $"{fullName}.{member}";
+			}
+
+			XmlNode? node = xmlDoc.SelectSingleNode($"//member[@name='{qualifier}:{fullName}']");
+			if (node == null)
+			{
+				return null;
+			}
+
+			XmlNode? summaryNode = node.SelectSingleNode("summary");
+			if (summaryNode == null)
+			{
+				if (node.SelectSingleNode("inheritdoc") != null)
+				{
+					foreach (Type baseType in FindBaseTypes(type))
+					{
+						string? summary = GetSummaryFromXmlDoc(baseType, qualifier, member);
+						if (summary != null)
+						{
+							return summary;
+						}
+					}
+				}
+				return null;
+			}
+
+			string text = summaryNode.InnerText.Trim().Replace("\r\n", "\n", StringComparison.Ordinal);
+			text = Regex.Replace(text, @"[\t ]*\n[\t ]*", "\n");
+			text = Regex.Replace(text, @"(?<!\n)\n(?!\n)", " ");
+			text = Regex.Replace(text, @"\n+", "\n");
+			return text;
+		}
+
+		static IEnumerable<Type> FindBaseTypes(Type type)
+		{
+			if (type.BaseType != null)
+			{
+				yield return type.BaseType;
+			}
+
+			foreach (Type interfaceType in type.GetInterfaces())
+			{
+				yield return interfaceType;
+			}
+		}
+	
+		bool TryReadXmlDoc(Assembly assembly, [NotNullWhen(true)] out XmlDocument? xmlDoc)
+		{
+			XmlDocument? documentation;
+			if (!_cachedDocumentation.TryGetValue(assembly, out documentation))
+			{
+				FileReference inputDocumentationFile = new FileReference(assembly.Location).ChangeExtension(".xml");
+				if (FileReference.Exists(inputDocumentationFile))
+				{
+					documentation = new XmlDocument();
+					documentation.Load(inputDocumentationFile.FullName);
+				}
+				_cachedDocumentation.TryAdd(assembly, documentation);
+			}
+
+			if (documentation != null)
+			{
+				xmlDoc = documentation;
+				return true;
+			}
+			else
+			{
+				xmlDoc = null;
+				return false;
+			}
+		}
+	}
+}
