Horde: Fix paths in link rail for schema documentation.

diff --git a/Engine/Source/Programs/Horde/Horde.Server/Commands/Config/DocsCommand.cs b/Engine/Source/Programs/Horde/Horde.Server/Commands/Config/DocsCommand.cs
--- a/Engine/Source/Programs/Horde/Horde.Server/Commands/Config/DocsCommand.cs
+++ b/Engine/Source/Programs/Horde/Horde.Server/Commands/Config/DocsCommand.cs
@@ -45,10 +45,10 @@ namespace Horde.Server.Commands.Config
 
 			Dictionary<JsonSchemaType, PageInfo> typeToPageInfo = new Dictionary<JsonSchemaType, PageInfo>
 			{
-				[serverSchema.RootType] = new PageInfo("appsettings.json (Server)", "[Horde](../Home.md) > [Deployment](../Deployment.md) > Server (Server.md)", "Deployment/ServerSettings.md", AppSettingsIntro),
-				[globalSchema.RootType] = new PageInfo("Globals.json", "[Horde](../Home.md) > [Configuration](../Config.md)", "Config/Schema/Globals.md"),
-				[projectSchema.RootType] = new PageInfo("*.project.json", "[Horde](../Home.md) > [Configuration](../Config.md)", "Config/Schema/Projects.md"),
-				[streamSchema.RootType] = new PageInfo("*.stream.json", "[Horde](../Home.md) > [Configuration](../Config.md)", "Config/Schema/Streams.md"),
+				[serverSchema.RootType] = new PageInfo("appsettings.json (Server)", "[Horde](../Home.md) > [Deployment](../Deployment.md) > [Server](Server.md)", "Deployment/ServerSettings.md", AppSettingsIntro),
+				[globalSchema.RootType] = new PageInfo("Globals.json", "[Horde](../../Home.md) > [Configuration](../../Config.md)", "Config/Schema/Globals.md"),
+				[projectSchema.RootType] = new PageInfo("*.project.json", "[Horde](../../Home.md) > [Configuration](../../Config.md)", "Config/Schema/Projects.md"),
+				[streamSchema.RootType] = new PageInfo("*.stream.json", "[Horde](../../Home.md) > [Configuration](../../Config.md)", "Config/Schema/Streams.md"),
 			};
 
 			if (Agent == null)
@@ -61,7 +61,7 @@ namespace Horde.Server.Commands.Config
 				Type agentSettingsType = agentAssembly.GetType("Horde.Agent.AgentSettings")!;
 
 				JsonSchema agentSchema = JsonSchema.FromType(agentSettingsType, xmlDocReader);
-				await WriteDocAsync(agentSchema.RootType, "appsettings.json (Agent)", "Deployment/AgentSettings.md", "Horde (../Home.md) > Deployment (../Deployment.md) > Agent (Agent.md)", AppSettingsIntro, new Dictionary<string, string>(), logger);
+				await WriteDocAsync(agentSchema.RootType, "appsettings.json (Agent)", "Deployment/AgentSettings.md", "[Horde](../Home.md) > [Deployment](../Deployment.md) > [Agent](Agent.md)", AppSettingsIntro, new Dictionary<string, string>(), logger);
 			}
 
 			Dictionary<string, string> typeNameToPageName = typeToPageInfo.ToDictionary(x => x.Key.Name!, x => x.Value.FileName, StringComparer.Ordinal);
