Horde: Fix build error.

diff --git a/Engine/Source/Programs/Horde/Horde.Server/Program.cs b/Engine/Source/Programs/Horde/Horde.Server/Program.cs
--- a/Engine/Source/Programs/Horde/Horde.Server/Program.cs
+++ b/Engine/Source/Programs/Horde/Horde.Server/Program.cs
@@ -42,7 +42,7 @@ namespace Horde.Server
 				{
 					theme = AnsiConsoleTheme.Code;
 				}
-				return sinkConfig.Console(outputTemplate: "[{Timestamp:HH:mm:ss} {Level:w3}] {Indent}{Message:l}{NewLine}{Exception}", theme: theme, restrictedToMinimumLevel: LogEventLevel.Debug);
+				return sinkConfig.Console(outputTemplate: "[{Timestamp:HH:mm:ss} {Level:w3}] {Indent}{Message:l}{NewLine}{Exception}", theme: theme, restrictedToMinimumLevel: Serilog.Events.LogEventLevel.Debug);
 			}
 		}
 
