Horde: Fix UE_HORDE_LAST_SUCCESS_CL and UE_HORDE_LAST_WARNING_CL not being set correctly if DB documents have not yet been updated with correct values.

diff --git a/Engine/Source/Programs/Horde/Horde.Server/Jobs/JobRpcService.cs b/Engine/Source/Programs/Horde/Horde.Server/Jobs/JobRpcService.cs
--- a/Engine/Source/Programs/Horde/Horde.Server/Jobs/JobRpcService.cs
+++ b/Engine/Source/Programs/Horde/Horde.Server/Jobs/JobRpcService.cs
@@ -37,6 +37,7 @@ using Microsoft.Extensions.Logging;
 using Microsoft.Extensions.Options;
 using MongoDB.Bson;
 using MongoDB.Bson.Serialization;
+using Microsoft.CodeAnalysis.CSharp.Syntax;
 
 namespace Horde.Server.Jobs
 {
@@ -513,22 +514,55 @@ namespace Horde.Server.Jobs
 				{
 					response.EnvVars.Add("UE_HORDE_LAST_CL", lastStep.Change.ToString(CultureInfo.InvariantCulture));
 
+					int? lastSuccessChange = null;
 					if (lastStep.Outcome == JobStepOutcome.Success)
 					{
-						response.EnvVars.Add("UE_HORDE_LAST_SUCCESS_CL", lastStep.Change.ToString(CultureInfo.InvariantCulture));
+						lastSuccessChange = lastStep.Change;
 					}
 					else if (lastStep.LastSuccess != null)
 					{
-						response.EnvVars.Add("UE_HORDE_LAST_SUCCESS_CL", lastStep.LastSuccess.Value.ToString(CultureInfo.InvariantCulture));
+						lastSuccessChange = lastStep.LastSuccess.Value;
+					}
+					else
+					{
+						// Previous job hasn't finished yet; need to search for *current* last success step
+						IJobStepRef? lastSuccess = await _jobStepRefCollection.GetPrevStepForNodeAsync(job.StreamId, job.TemplateId, node.Name, job.Change, outcome: JobStepOutcome.Success);
+						if (lastSuccess != null)
+						{
+							lastSuccessChange = lastSuccess.Change;
+						}
 					}
 
+					int? lastWarningChange = null;
 					if (lastStep.Outcome == JobStepOutcome.Success || lastStep.Outcome == JobStepOutcome.Warnings)
 					{
-						response.EnvVars.Add("UE_HORDE_LAST_WARNING_CL", lastStep.Change.ToString(CultureInfo.InvariantCulture));
+						lastWarningChange = lastStep.Change;
 					}
 					else if (lastStep.LastWarning != null)
 					{
-						response.EnvVars.Add("UE_HORDE_LAST_WARNING_CL", lastStep.LastWarning.Value.ToString(CultureInfo.InvariantCulture));
+						lastWarningChange = lastStep.LastWarning.Value;
+					}
+					else
+					{
+						// Previous job hasn't finished yet; need to search for *current* last warning step
+						IJobStepRef? lastWarnings = await _jobStepRefCollection.GetPrevStepForNodeAsync(job.StreamId, job.TemplateId, node.Name, job.Change, outcome: JobStepOutcome.Warnings);
+						if (lastWarnings != null && (lastSuccessChange == null || lastWarnings.Change > lastSuccessChange.Value))
+						{
+							lastWarningChange = lastWarnings.Change;
+						}
+						else
+						{
+							lastWarningChange = lastSuccessChange;
+						}
+					}
+
+					if (lastSuccessChange != null)
+					{
+						response.EnvVars.Add("UE_HORDE_LAST_SUCCESS_CL", lastSuccessChange.Value.ToString(CultureInfo.InvariantCulture));
+					}
+					if (lastWarningChange != null)
+					{
+						response.EnvVars.Add("UE_HORDE_LAST_WARNING_CL", lastWarningChange.Value.ToString(CultureInfo.InvariantCulture));
 					}
 				}
 
