Horde: Prevent commits being cached that aren't submitted.

diff --git a/Engine/Source/Programs/Horde/Horde.Server/Perforce/PerforceServiceCache.cs b/Engine/Source/Programs/Horde/Horde.Server/Perforce/PerforceServiceCache.cs
--- a/Engine/Source/Programs/Horde/Horde.Server/Perforce/PerforceServiceCache.cs
+++ b/Engine/Source/Programs/Horde/Horde.Server/Perforce/PerforceServiceCache.cs
@@ -452,6 +452,12 @@ namespace Horde.Server.Perforce
 				{
 					foreach (DescribeRecord describeRecord in describeRecords)
 					{
+						if (describeRecord.Status != ChangeStatus.Submitted)
+						{
+							// This can happen because we received a P4 trigger notifying us of a form save. Need to filter out non-committed changes.
+							continue;
+						}
+
 						files.Clear();
 						foreach (DescribeFileRecord describeFile in describeRecord.Files)
 						{
