Horde: Create log objects for each compute lease.

diff --git a/Engine/Source/Programs/Horde/Horde.Server/Compute/ComputeController.cs b/Engine/Source/Programs/Horde/Horde.Server/Compute/ComputeController.cs
--- a/Engine/Source/Programs/Horde/Horde.Server/Compute/ComputeController.cs
+++ b/Engine/Source/Programs/Horde/Horde.Server/Compute/ComputeController.cs
@@ -3,6 +3,7 @@
 using System;
 using System.Collections.Generic;
 using System.Net;
+using System.Threading;
 using System.Threading.Tasks;
 using EpicGames.Core;
 using EpicGames.Horde.Compute;
@@ -88,11 +89,12 @@ namespace Horde.Server.Compute
 		/// </summary>
 		/// <param name="clusterId">Id of the compute cluster</param>
 		/// <param name="request">The request parameters</param>
+		/// <param name="cancellationToken">Cancellation token for the operation</param>
 		/// <returns></returns>
 		[HttpPost]
 		[Authorize]
 		[Route("/api/v2/compute/{clusterId}")]
-		public async Task<ActionResult<AssignComputeResponse>> AssignComputeResourceAsync(ClusterId clusterId, [FromBody] AssignComputeRequest request)
+		public async Task<ActionResult<AssignComputeResponse>> AssignComputeResourceAsync(ClusterId clusterId, [FromBody] AssignComputeRequest request, CancellationToken cancellationToken)
 		{
 			ComputeClusterConfig? clusterConfig;
 			if (!_globalConfig.Value.TryGetComputeCluster(clusterId, out clusterConfig))
@@ -108,7 +110,7 @@ namespace Horde.Server.Compute
 
 			Requirements requirements = request.Requirements ?? new Requirements();
 
-			ComputeResource? computeResource = await _computeService.TryAllocateResourceAsync(requirements, parentLeaseId);
+			ComputeResource? computeResource = await _computeService.TryAllocateResourceAsync(requirements, parentLeaseId, cancellationToken);
 			if (computeResource == null)
 			{
 				return StatusCode((int)HttpStatusCode.ServiceUnavailable);
diff --git a/Engine/Source/Programs/Horde/Horde.Server/Compute/ComputeService.cs b/Engine/Source/Programs/Horde/Horde.Server/Compute/ComputeService.cs
--- a/Engine/Source/Programs/Horde/Horde.Server/Compute/ComputeService.cs
+++ b/Engine/Source/Programs/Horde/Horde.Server/Compute/ComputeService.cs
@@ -5,6 +5,7 @@ using System.Collections.Generic;
 using System.Linq;
 using System.Net;
 using System.Security.Cryptography;
+using System.Threading;
 using System.Threading.Tasks;
 using EpicGames.Horde.Compute;
 using EpicGames.Horde.Compute.Clients;
@@ -13,6 +14,8 @@ using Google.Protobuf;
 using Google.Protobuf.WellKnownTypes;
 using Horde.Server.Agents;
 using Horde.Server.Agents.Leases;
+using Horde.Server.Jobs;
+using Horde.Server.Logs;
 using HordeCommon;
 using HordeCommon.Rpc.Tasks;
 
@@ -24,21 +27,23 @@ namespace Horde.Server.Compute
 	public class ComputeService
 	{
 		readonly IAgentCollection _agentCollection;
+		readonly ILogFileService _logService;
 		readonly AgentService _agentService;
 
 		/// <summary>
 		/// Constructor
 		/// </summary>
-		public ComputeService(IAgentCollection agentCollection, AgentService agentService)
+		public ComputeService(IAgentCollection agentCollection, ILogFileService logService, AgentService agentService)
 		{
 			_agentCollection = agentCollection;
+			_logService = logService;
 			_agentService = agentService;
 		}
 
 		/// <summary>
 		/// Allocates a compute resource
 		/// </summary>
-		public async Task<ComputeResource?> TryAllocateResourceAsync(Requirements requirements, LeaseId? parentLeaseId)
+		public async Task<ComputeResource?> TryAllocateResourceAsync(Requirements requirements, LeaseId? parentLeaseId, CancellationToken cancellationToken)
 		{
 			List<IAgent> agents = await _agentCollection.FindAsync();
 			foreach (IAgent agent in agents)
@@ -46,10 +51,13 @@ namespace Horde.Server.Compute
 				Dictionary<string, int> assignedResources = new Dictionary<string, int>();
 				if (agent.MeetsRequirements(requirements, assignedResources))
 				{
-					ComputeTask computeTask = CreateComputeTask(assignedResources);
+					LeaseId leaseId = LeaseId.GenerateNewId();
+					ILogFile? log = await _logService.CreateLogFileAsync(JobId.Empty, leaseId, agent.SessionId, LogType.Json, useNewStorageBackend: true, cancellationToken: cancellationToken);
+
+					ComputeTask computeTask = CreateComputeTask(assignedResources, log?.Id);
 
 					byte[] payload = Any.Pack(computeTask).ToByteArray();
-					AgentLease lease = new AgentLease(LeaseId.GenerateNewId(), parentLeaseId, "Compute task", null, null, null, LeaseState.Pending, assignedResources, requirements.Exclusive, payload);
+					AgentLease lease = new AgentLease(leaseId, parentLeaseId, "Compute task", null, null, log?.Id, LeaseState.Pending, assignedResources, requirements.Exclusive, payload);
 
 					ComputeResource? resource = TryAssign(agent, computeTask);
 					if (resource != null)
@@ -84,12 +92,13 @@ namespace Horde.Server.Compute
 			return new ComputeResource(ip, port, computeTask, agent.Properties);
 		}
 
-		static ComputeTask CreateComputeTask(Dictionary<string, int> assignedResources)
+		static ComputeTask CreateComputeTask(Dictionary<string, int> assignedResources, LogId? logId)
 		{
 			ComputeTask computeTask = new ComputeTask();
 			computeTask.Nonce = UnsafeByteOperations.UnsafeWrap(RandomNumberGenerator.GetBytes(ServerComputeClient.NonceLength));
 			computeTask.Key = UnsafeByteOperations.UnsafeWrap(AesTransport.CreateKey());
 			computeTask.Resources.Add(assignedResources);
+			computeTask.LogId = logId?.ToString();
 			return computeTask;
 		}
 	}
