Horde: Support for writing to a server log file from compute leases.

diff --git a/Engine/Source/Programs/Horde/Horde.Agent/Leases/Handlers/ComputeHandler.cs b/Engine/Source/Programs/Horde/Horde.Agent/Leases/Handlers/ComputeHandler.cs
--- a/Engine/Source/Programs/Horde/Horde.Agent/Leases/Handlers/ComputeHandler.cs
+++ b/Engine/Source/Programs/Horde/Horde.Agent/Leases/Handlers/ComputeHandler.cs
@@ -11,6 +11,7 @@ using EpicGames.Core;
 using EpicGames.Horde.Compute;
 using EpicGames.Horde.Compute.Transports;
 using Horde.Agent.Services;
+using Horde.Agent.Utility;
 using HordeCommon.Rpc.Tasks;
 using Microsoft.Extensions.Caching.Memory;
 using Microsoft.Extensions.Logging;
@@ -61,16 +62,18 @@ namespace Horde.Agent.Leases.Handlers
 
 		readonly ComputeListenerService _listenerService;
 		readonly IMemoryCache _memoryCache;
+		readonly IServerLoggerFactory _serverLoggerFactory;
 		readonly AgentSettings _settings;
 		readonly ILogger _logger;
 
 		/// <summary>
 		/// Constructor
 		/// </summary>
-		public ComputeHandler(ComputeListenerService listenerService, IMemoryCache memoryCache, IOptions<AgentSettings> settings, ILogger<ComputeHandler> logger)
+		public ComputeHandler(ComputeListenerService listenerService, IMemoryCache memoryCache, IServerLoggerFactory serverLoggerFactory, IOptions<AgentSettings> settings, ILogger<ComputeHandler> logger)
 		{
 			_listenerService = listenerService;
 			_memoryCache = memoryCache;
+			_serverLoggerFactory = serverLoggerFactory;
 			_settings = settings.Value;
 			_logger = logger;
 		}
@@ -78,7 +81,10 @@ namespace Horde.Agent.Leases.Handlers
 		/// <inheritdoc/>
 		public override async Task<LeaseResult> ExecuteAsync(ISession session, string leaseId, ComputeTask computeTask, CancellationToken cancellationToken)
 		{
-			_logger.LogInformation("Starting compute task (lease {LeaseId}). Waiting for connection with nonce {Nonce}...", leaseId, StringUtils.FormatHexString(computeTask.Nonce.Span));
+			await using IServerLogger? serverLogger = (computeTask.LogId != null)? _serverLoggerFactory.CreateLogger(session, computeTask.LogId, null, true) : null;
+			ILogger logger = serverLogger ?? _logger;
+
+			logger.LogInformation("Starting compute task (lease {LeaseId}). Waiting for connection with nonce {Nonce}...", leaseId, StringUtils.FormatHexString(computeTask.Nonce.Span));
 
 			TcpClient? tcpClient = null;
 			try
@@ -88,19 +94,19 @@ namespace Horde.Agent.Leases.Handlers
 				tcpClient = await _listenerService.WaitForClientAsync(new ByteString(computeTask.Nonce.Memory), TimeSpan.FromSeconds(TimeoutSeconds), cancellationToken);
 				if (tcpClient == null)
 				{
-					_logger.LogInformation("Timed out waiting for connection after {Time}s.", TimeoutSeconds); 
+					logger.LogInformation("Timed out waiting for connection after {Time}s.", TimeoutSeconds); 
 					return LeaseResult.Success;
 				}
 
-				_logger.LogInformation("Matched connection for {Nonce}", StringUtils.FormatHexString(computeTask.Nonce.Span));
+				logger.LogInformation("Matched connection for {Nonce}", StringUtils.FormatHexString(computeTask.Nonce.Span));
 
 				TcpTransportWithTimeout transport = new TcpTransportWithTimeout(tcpClient.Client);
 				using (CancellationTokenSource cts = CancellationTokenSource.CreateLinkedTokenSource(cancellationToken))
 				{
-					await using BackgroundTask timeoutTask = BackgroundTask.StartNew(ctx => TickTimeoutAsync(transport, cts, ctx));
+					await using BackgroundTask timeoutTask = BackgroundTask.StartNew(ctx => TickTimeoutAsync(transport, cts, logger, ctx));
 					try
 					{
-						await using (RemoteComputeSocket socket = new RemoteComputeSocket(transport, ComputeSocketEndpoint.Local, _logger))
+						await using (RemoteComputeSocket socket = new RemoteComputeSocket(transport, ComputeSocketEndpoint.Local, logger))
 						{
 							DirectoryReference sandboxDir = DirectoryReference.Combine(session.WorkingDir, "Sandbox", leaseId);
 							try
@@ -113,7 +119,7 @@ namespace Horde.Agent.Leases.Handlers
 								Dictionary<string, string?> newEnvVars = new Dictionary<string, string?>();
 								newEnvVars["UE_HORDE_SHARED_DIR"] = sharedDir.FullName;
 
-								AgentMessageHandler worker = new AgentMessageHandler(sandboxDir, _memoryCache, newEnvVars, false, _settings.WineExecutablePath, _logger);
+								AgentMessageHandler worker = new AgentMessageHandler(sandboxDir, _memoryCache, newEnvVars, false, _settings.WineExecutablePath, logger);
 								await worker.RunAsync(socket, cts.Token);
 								await socket.CloseAsync(cts.Token);
 								return LeaseResult.Success;
@@ -126,14 +132,14 @@ namespace Horde.Agent.Leases.Handlers
 					}
 					catch (OperationCanceledException ex) when (cts.IsCancellationRequested && transport.TimeSinceActivity > TimeSpan.FromMinutes(NoDataTimeoutMinutes))
 					{
-						_logger.LogError(ex, "Lease was terminated due to no data being received for {Time} minutes", NoDataTimeoutMinutes);
+						logger.LogError(ex, "Lease was terminated due to no data being received for {Time} minutes", NoDataTimeoutMinutes);
 						return LeaseResult.Failed;
 					}
 				}
 			}
 			catch (Exception ex)
 			{
-				_logger.LogError(ex, "Exception while executing compute task: {Message}", ex.Message);
+				logger.LogError(ex, "Exception while executing compute task: {Message}", ex.Message);
 				return LeaseResult.Failed;
 			}
 			finally
@@ -142,13 +148,13 @@ namespace Horde.Agent.Leases.Handlers
 			}
 		}
 
-		async Task TickTimeoutAsync(TcpTransportWithTimeout transport, CancellationTokenSource cts, CancellationToken cancellationToken)
+		static async Task TickTimeoutAsync(TcpTransportWithTimeout transport, CancellationTokenSource cts, ILogger logger, CancellationToken cancellationToken)
 		{
 			while(!cancellationToken.IsCancellationRequested)
 			{
 				if (transport.TimeSinceActivity > TimeSpan.FromMinutes(NoDataTimeoutMinutes))
 				{
-					_logger.LogWarning("Terminating compute task due to timeout (last tick at {Time})", DateTime.UtcNow - transport.TimeSinceActivity);
+					logger.LogWarning("Terminating compute task due to timeout (last tick at {Time})", DateTime.UtcNow - transport.TimeSinceActivity);
 					cts.Cancel();
 					break;
 				}
