Horde: Output the parent lease id to the log at the start of compute tasks.

diff --git a/Engine/Source/Programs/Horde/Horde.Agent/Leases/Handlers/ComputeHandler.cs b/Engine/Source/Programs/Horde/Horde.Agent/Leases/Handlers/ComputeHandler.cs
--- a/Engine/Source/Programs/Horde/Horde.Agent/Leases/Handlers/ComputeHandler.cs
+++ b/Engine/Source/Programs/Horde/Horde.Agent/Leases/Handlers/ComputeHandler.cs
@@ -84,6 +84,11 @@ namespace Horde.Agent.Leases.Handlers
 			await using IServerLogger? serverLogger = (computeTask.LogId != null)? _serverLoggerFactory.CreateLogger(session, computeTask.LogId, null, true) : null;
 			ILogger logger = serverLogger ?? _logger;
 
+			if (!String.IsNullOrEmpty(computeTask.ParentLeaseId))
+			{
+				logger.LogInformation("Parent lease: {LeaseId}", computeTask.ParentLeaseId);
+			}
+
 			logger.LogInformation("Starting compute task (lease {LeaseId}). Waiting for connection with nonce {Nonce}...", leaseId, StringUtils.FormatHexString(computeTask.Nonce.Span));
 
 			TcpClient? tcpClient = null;
diff --git a/Engine/Source/Programs/Horde/Horde.Server/Compute/ComputeService.cs b/Engine/Source/Programs/Horde/Horde.Server/Compute/ComputeService.cs
--- a/Engine/Source/Programs/Horde/Horde.Server/Compute/ComputeService.cs
+++ b/Engine/Source/Programs/Horde/Horde.Server/Compute/ComputeService.cs
@@ -54,7 +54,7 @@ namespace Horde.Server.Compute
 					LeaseId leaseId = LeaseId.GenerateNewId();
 					ILogFile? log = await _logService.CreateLogFileAsync(JobId.Empty, leaseId, agent.SessionId, LogType.Json, useNewStorageBackend: true, cancellationToken: cancellationToken);
 
-					ComputeTask computeTask = CreateComputeTask(assignedResources, log?.Id);
+					ComputeTask computeTask = CreateComputeTask(assignedResources, log?.Id, parentLeaseId);
 
 					byte[] payload = Any.Pack(computeTask).ToByteArray();
 					AgentLease lease = new AgentLease(leaseId, parentLeaseId, "Compute task", null, null, log?.Id, LeaseState.Pending, assignedResources, requirements.Exclusive, payload);
@@ -92,13 +92,14 @@ namespace Horde.Server.Compute
 			return new ComputeResource(ip, port, computeTask, agent.Properties);
 		}
 
-		static ComputeTask CreateComputeTask(Dictionary<string, int> assignedResources, LogId? logId)
+		static ComputeTask CreateComputeTask(Dictionary<string, int> assignedResources, LogId? logId, LeaseId? parentLeaseId)
 		{
 			ComputeTask computeTask = new ComputeTask();
 			computeTask.Nonce = UnsafeByteOperations.UnsafeWrap(RandomNumberGenerator.GetBytes(ServerComputeClient.NonceLength));
 			computeTask.Key = UnsafeByteOperations.UnsafeWrap(AesTransport.CreateKey());
 			computeTask.Resources.Add(assignedResources);
 			computeTask.LogId = logId?.ToString();
+			computeTask.ParentLeaseId = parentLeaseId?.ToString();
 			return computeTask;
 		}
 	}
