Horde: Return the lease id of the assigned resource to the caller with compute leases.

diff --git a/Engine/Source/Programs/Horde/Horde.Server/Compute/ComputeController.cs b/Engine/Source/Programs/Horde/Horde.Server/Compute/ComputeController.cs
--- a/Engine/Source/Programs/Horde/Horde.Server/Compute/ComputeController.cs
+++ b/Engine/Source/Programs/Horde/Horde.Server/Compute/ComputeController.cs
@@ -53,6 +53,11 @@ namespace Horde.Server.Compute
 		/// </summary>
 		public string Key { get; set; } = String.Empty;
 
+		/// <summary>
+		/// Identifier for the new lease on the remote machine
+		/// </summary>
+		public LeaseId LeaseId { get; set; }
+
 		/// <summary>
 		/// Resources assigned to this machine
 		/// </summary>
@@ -121,6 +126,7 @@ namespace Horde.Server.Compute
 			response.Port = computeResource.Port;
 			response.Nonce = StringUtils.FormatHexString(computeResource.Task.Nonce.Span);
 			response.Key = StringUtils.FormatHexString(computeResource.Task.Key.Span);
+			response.LeaseId = computeResource.LeaseId;
 			response.Properties = computeResource.Properties;
 
 			foreach (KeyValuePair<string, int> pair in computeResource.Task.Resources)
diff --git a/Engine/Source/Programs/Horde/Horde.Server/Compute/ComputeService.cs b/Engine/Source/Programs/Horde/Horde.Server/Compute/ComputeService.cs
--- a/Engine/Source/Programs/Horde/Horde.Server/Compute/ComputeService.cs
+++ b/Engine/Source/Programs/Horde/Horde.Server/Compute/ComputeService.cs
@@ -61,7 +61,7 @@ namespace Horde.Server.Compute
 					byte[] payload = Any.Pack(computeTask).ToByteArray();
 					AgentLease lease = new AgentLease(leaseId, parentLeaseId, "Compute task", null, null, log?.Id, LeaseState.Pending, assignedResources, requirements.Exclusive, payload);
 
-					ComputeResource? resource = TryAssign(agent, computeTask);
+					ComputeResource? resource = TryAssign(agent, computeTask, leaseId);
 					if (resource != null)
 					{
 						IAgent? newAgent = await _agentCollection.TryAddLeaseAsync(agent, lease);
@@ -77,7 +77,7 @@ namespace Horde.Server.Compute
 			return null;
 		}
 
-		static ComputeResource? TryAssign(IAgent agent, ComputeTask computeTask)
+		static ComputeResource? TryAssign(IAgent agent, ComputeTask computeTask, LeaseId leaseId)
 		{
 			string? ipStr = agent.GetPropertyValues("ComputeIp").FirstOrDefault();
 			if (ipStr == null || !IPAddress.TryParse(ipStr, out IPAddress? ip))
@@ -91,7 +91,7 @@ namespace Horde.Server.Compute
 				return null;
 			}
 
-			return new ComputeResource(ip, port, computeTask, agent.Properties);
+			return new ComputeResource(ip, port, computeTask, agent.Properties, leaseId);
 		}
 
 		static ComputeTask CreateComputeTask(Dictionary<string, int> assignedResources, LogId? logId, LeaseId? parentLeaseId)
