Horde: Missing file.

diff --git a/Engine/Source/Programs/Horde/Horde.Server/Compute/ComputeTaskSource.cs b/Engine/Source/Programs/Horde/Horde.Server/Compute/ComputeTaskSource.cs
--- a/Engine/Source/Programs/Horde/Horde.Server/Compute/ComputeTaskSource.cs
+++ b/Engine/Source/Programs/Horde/Horde.Server/Compute/ComputeTaskSource.cs
@@ -8,6 +8,7 @@ using System.Security.Cryptography;
 using System.Threading;
 using System.Threading.Tasks;
 using EpicGames.Core;
+using EpicGames.Horde.Api;
 using EpicGames.Horde.Compute;
 using EpicGames.Horde.Compute.Clients;
 using EpicGames.Horde.Compute.Transports;
@@ -49,15 +50,21 @@ namespace Horde.Server.Compute
 		/// </summary>
 		public IReadOnlyList<string> Properties { get; }
 
+		/// <summary>
+		/// Lease id on the remote machine
+		/// </summary>
+		public LeaseId LeaseId { get; }
+
 		/// <summary>
 		/// Constructor
 		/// </summary>
-		public ComputeResource(IPAddress ip, int port, ComputeTask task, IReadOnlyList<string> properties)
+		public ComputeResource(IPAddress ip, int port, ComputeTask task, IReadOnlyList<string> properties, LeaseId leaseId)
 		{
 			Ip = ip;
 			Port = port;
 			Task = task;
 			Properties = properties;
+			LeaseId = leaseId;
 		}
 	}
 
