Horde: Support ephemeral agents

When an agent is registered as ephemeral, it will not persist any records on the server after it disconnects. This is ideal for short-lived agents that are auto-scaled, joining and leaving as workload fluctuates.

diff --git a/Engine/Source/Programs/Horde/Horde.Agent/AgentSettings.cs b/Engine/Source/Programs/Horde/Horde.Agent/AgentSettings.cs
--- a/Engine/Source/Programs/Horde/Horde.Agent/AgentSettings.cs
+++ b/Engine/Source/Programs/Horde/Horde.Agent/AgentSettings.cs
@@ -190,6 +190,14 @@ namespace Horde.Agent
 		/// </summary>
 		public string? Name { get; set; }
 
+		/// <summary>
+		/// Whether agent should register as being ephemeral.
+		/// Doing so will not persist any long-lived data on the server and
+		/// once disconnected it's assumed to have been deleted permanently.
+		/// Ideal for short-lived agents, such as spot instances on AWS EC2.
+		/// </summary>
+		public bool Ephemeral { get; set; } = false;
+
 		/// <summary>
 		/// The executor to use for jobs. Defaults to the Perforce executor.
 		/// </summary>
diff --git a/Engine/Source/Programs/Horde/Horde.Agent/Services/SessionFactory.cs b/Engine/Source/Programs/Horde/Horde.Agent/Services/SessionFactory.cs
--- a/Engine/Source/Programs/Horde/Horde.Agent/Services/SessionFactory.cs
+++ b/Engine/Source/Programs/Horde/Horde.Agent/Services/SessionFactory.cs
@@ -270,6 +270,7 @@ namespace Horde.Agent.Services
 
 					CreateAgentRequest createAgentRequest = new CreateAgentRequest();
 					createAgentRequest.Name = currentSettings.GetAgentName();
+					createAgentRequest.Ephemeral = currentSettings.Ephemeral;
 
 					CreateAgentResponse createAgentResponse = await rpcClient.CreateAgentAsync(createAgentRequest, null, null, cancellationToken);
 					registration = new AgentRegistration(grpcService.ServerProfile.Url, createAgentResponse.Id, createAgentResponse.Token);
diff --git a/Engine/Source/Programs/Horde/Horde.Server.Tests/AgentServiceTest.cs b/Engine/Source/Programs/Horde/Horde.Server.Tests/AgentServiceTest.cs
--- a/Engine/Source/Programs/Horde/Horde.Server.Tests/AgentServiceTest.cs
+++ b/Engine/Source/Programs/Horde/Horde.Server.Tests/AgentServiceTest.cs
@@ -8,6 +8,7 @@ using System.Threading;
 using System.Threading.Tasks;
 using EpicGames.Core;
 using Horde.Server.Agents;
+using Horde.Server.Agents.Pools;
 using Horde.Server.Auditing;
 using Horde.Server.Jobs;
 using Horde.Server.Server;
@@ -172,6 +173,34 @@ public class AgentServiceTest : TestSetup
 		double? rate2 = await AgentService.GetRateAsync(agent2.Id);
 		Assert.AreEqual(300, rate2!.Value, 0.1);
 	}
+	
+	[TestMethod]
+	public async Task EphemeralTest()
+	{
+		IAgent agent1 = await CreateAgentAsync(new PoolId("pool1"), ephemeral: true);
+		IAgent agent2 = await CreateAgentAsync(new PoolId("pool1"), ephemeral: false);
+		Assert.IsTrue(agent1.Ephemeral);
+		Assert.IsFalse(agent2.Ephemeral);
+		
+		agent1 = (await AgentService.GetAgentAsync(agent1.Id))!;
+		agent2 = (await AgentService.GetAgentAsync(agent2.Id))!;
+		Assert.IsTrue(agent1.Ephemeral);
+		Assert.IsFalse(agent2.Ephemeral);
+		Assert.AreEqual(AgentStatus.Ok, agent1.Status);
+		Assert.AreEqual(AgentStatus.Ok, agent2.Status);
+
+		// Let background task run for purging outdated sessions
+		await AgentService.StartAsync(CancellationToken.None);
+		await Clock.AdvanceAsync(TimeSpan.FromHours(1));
+		
+		// Ephemeral agent is deleted once its session is terminated
+		Assert.IsNull((await AgentService.GetAgentAsync(agent1.Id)));
+		
+		// Normal agent is stopped but kept around
+		agent2 = (await AgentService.GetAgentAsync(agent2.Id))!;
+		Assert.AreEqual(AgentStatus.Stopped, agent2.Status);
+		Assert.IsFalse(agent2.Deleted);
+	}
 
 	private static ClaimsPrincipal GetUser(IAgent agent)
 	{
diff --git a/Engine/Source/Programs/Horde/Horde.Server.Tests/TestSetup.cs b/Engine/Source/Programs/Horde/Horde.Server.Tests/TestSetup.cs
--- a/Engine/Source/Programs/Horde/Horde.Server.Tests/TestSetup.cs
+++ b/Engine/Source/Programs/Horde/Horde.Server.Tests/TestSetup.cs
@@ -368,10 +368,11 @@ namespace Horde.Server.Tests
 		/// <param name="adjustClockBy">Time span to temporarily skew the clock when creating the agent</param>
 		/// <param name="awsInstanceId">AWS instance ID for the agent (will be set in properties)</param>
 		/// <param name="lease">A lease to assign the agent</param>
+		/// <param name="ephemeral">Whether the agent is ephemeral</param>
 		/// <returns>A new agent</returns>
 		public async Task<IAgent> CreateAgentAsync(
 			PoolId poolId, bool enabled = true, bool requestShutdown = false, List<string>? properties = null,
-			TimeSpan? adjustClockBy = null, string? awsInstanceId = null, AgentLease? lease = null)
+			TimeSpan? adjustClockBy = null, string? awsInstanceId = null, AgentLease? lease = null, bool ephemeral = false)
 		{
 			DateTime now = Clock.UtcNow;
 			if (adjustClockBy != null)
@@ -386,7 +387,7 @@ namespace Horde.Server.Tests
 				tempProps.Add(KnownPropertyNames.AwsInstanceId + "=" + awsInstanceId);
 			}
 			
-			IAgent agent = await AgentService.CreateAgentAsync("TestAgent" + s_agentIdCounter++, enabled, new List<PoolId> { poolId });
+			IAgent agent = await AgentService.CreateAgentAsync("TestAgent" + s_agentIdCounter++, enabled, new List<PoolId> { poolId }, ephemeral);
 			agent = await AgentService.CreateSessionAsync(agent, AgentStatus.Ok, tempProps, resources, null);
 			if (requestShutdown)
 			{
diff --git a/Engine/Source/Programs/Horde/Horde.Server/Agents/AgentCollection.cs b/Engine/Source/Programs/Horde/Horde.Server/Agents/AgentCollection.cs
--- a/Engine/Source/Programs/Horde/Horde.Server/Agents/AgentCollection.cs
+++ b/Engine/Source/Programs/Horde/Horde.Server/Agents/AgentCollection.cs
@@ -126,11 +126,12 @@ namespace Horde.Server.Agents
 			{
 			}
 
-			public AgentDocument(AgentId id, bool enabled, List<PoolId> pools)
+			public AgentDocument(AgentId id, bool enabled, List<PoolId> pools, bool ephemeral)
 			{
 				Id = id;
 				Enabled = enabled;
 				Pools = pools;
+				Ephemeral = ephemeral;
 			}
 		}
 
@@ -156,9 +157,9 @@ namespace Horde.Server.Agents
 		}
 
 		/// <inheritdoc/>
-		public async Task<IAgent> AddAsync(AgentId id, bool enabled, List<PoolId>? pools)
+		public async Task<IAgent> AddAsync(AgentId id, bool enabled, List<PoolId>? pools, bool ephemeral)
 		{
-			AgentDocument agent = new AgentDocument(id, enabled, pools ?? new List<PoolId>());
+			AgentDocument agent = new AgentDocument(id, enabled, pools ?? new List<PoolId>(), ephemeral);
 			await _agents.InsertOneAsync(agent);
 			return agent;
 		}
diff --git a/Engine/Source/Programs/Horde/Horde.Server/Agents/AgentService.cs b/Engine/Source/Programs/Horde/Horde.Server/Agents/AgentService.cs
--- a/Engine/Source/Programs/Horde/Horde.Server/Agents/AgentService.cs
+++ b/Engine/Source/Programs/Horde/Horde.Server/Agents/AgentService.cs
@@ -189,10 +189,11 @@ namespace Horde.Server.Agents
 		/// <param name="name">Name of the agent</param>
 		/// <param name="enabled">Whether the agent is currently enabled</param>
 		/// <param name="pools">Pools for this agent</param>
+		/// <param name="ephemeral">Whether the agent is ephemeral or not</param>
 		/// <returns>Unique id for the agent</returns>
-		public Task<IAgent> CreateAgentAsync(string name, bool enabled, List<PoolId>? pools)
+		public Task<IAgent> CreateAgentAsync(string name, bool enabled, List<PoolId>? pools, bool ephemeral = false)
 		{
-			return Agents.AddAsync(new AgentId(name), enabled, pools);
+			return Agents.AddAsync(new AgentId(name), enabled, pools, ephemeral);
 		}
 
 		/// <summary>
@@ -236,10 +237,22 @@ namespace Horde.Server.Agents
 		/// Marks the agent as deleted
 		/// </summary>
 		/// <param name="agent">The agent to delete</param>
+		/// <param name="forceDelete">Whether to fully delete the agent as opposed to just marking it as deleted</param>
 		/// <returns>Async task</returns>
-		public async Task DeleteAgentAsync(IAgent? agent)
+		public async Task DeleteAgentAsync(IAgent? agent, bool forceDelete = false)
 		{
-			while (agent != null && !agent.Deleted)
+			if (agent == null)
+			{
+				return;
+			}
+
+			if (forceDelete)
+			{
+				await Agents.ForceDeleteAsync(agent.Id);
+				return;
+			}
+			
+			while (agent is { Deleted: false })
 			{
 				IAgent? newAgent = await Agents.TryDeleteAsync(agent);
 				if(newAgent != null)
@@ -1031,7 +1044,13 @@ namespace Horde.Server.Agents
 				{
 					stoppingToken.ThrowIfCancellationRequested();
 					_logger.LogDebug("Terminating session {SessionId} for agent {Agent}", expiredAgent.SessionId, expiredAgent.Id);
-					await TryTerminateSessionAsync(expiredAgent);
+					IAgent? terminatedAgent = await TryTerminateSessionAsync(expiredAgent);
+
+					if (terminatedAgent is { Status: AgentStatus.Stopped, Ephemeral: true })
+					{
+						_logger.LogDebug("Deleting ephemeral agent {Agent}", terminatedAgent.Id);
+						await DeleteAgentAsync(terminatedAgent, true);
+					}
 				}
 
 				// Try again if we didn't fetch everything
diff --git a/Engine/Source/Programs/Horde/Horde.Server/Agents/IAgent.cs b/Engine/Source/Programs/Horde/Horde.Server/Agents/IAgent.cs
--- a/Engine/Source/Programs/Horde/Horde.Server/Agents/IAgent.cs
+++ b/Engine/Source/Programs/Horde/Horde.Server/Agents/IAgent.cs
@@ -482,6 +482,11 @@ namespace Horde.Server.Agents
 		/// Whether the agent is enabled
 		/// </summary>
 		public bool Enabled { get; }
+		
+		/// <summary>
+		/// Whether the agent is ephemeral
+		/// </summary>
+		public bool Ephemeral { get; }
 
 		/// <summary>
 		/// Whether the agent should be included on the dashboard. This is set to true for ephemeral agents once they are no longer online, or agents that are explicitly deleted.
diff --git a/Engine/Source/Programs/Horde/Horde.Server/Agents/IAgentCollection.cs b/Engine/Source/Programs/Horde/Horde.Server/Agents/IAgentCollection.cs
--- a/Engine/Source/Programs/Horde/Horde.Server/Agents/IAgentCollection.cs
+++ b/Engine/Source/Programs/Horde/Horde.Server/Agents/IAgentCollection.cs
@@ -23,7 +23,8 @@ namespace Horde.Server.Agents
 		/// <param name="id">Id for the new agent</param>
 		/// <param name="enabled">Whether the agent is enabled or not</param>
 		/// <param name="pools">Pools for the agent</param>
-		Task<IAgent> AddAsync(AgentId id, bool enabled, List<PoolId>? pools = null);
+		/// <param name="ephemeral">Whether the agent is ephemeral or not</param>
+		Task<IAgent> AddAsync(AgentId id, bool enabled, List<PoolId>? pools = null, bool ephemeral = false);
 
 		/// <summary>
 		/// Deletes an agent
diff --git a/Engine/Source/Programs/Horde/Horde.Server/Server/RpcService.cs b/Engine/Source/Programs/Horde/Horde.Server/Server/RpcService.cs
--- a/Engine/Source/Programs/Horde/Horde.Server/Server/RpcService.cs
+++ b/Engine/Source/Programs/Horde/Horde.Server/Server/RpcService.cs
@@ -267,7 +267,7 @@ namespace Horde.Server.Server
 			IAgent? agent = await _agentService.GetAgentAsync(new AgentId(request.Name));
 			if (agent == null)
 			{
-				agent = await _agentService.CreateAgentAsync(request.Name, _globalConfig.Value.ServerSettings.EnableNewAgentsByDefault, null);
+				agent = await _agentService.CreateAgentAsync(request.Name, _globalConfig.Value.ServerSettings.EnableNewAgentsByDefault, null, request.Ephemeral);
 			}
 
 			List<AclClaimConfig> claims = new List<AclClaimConfig>();
