Horde: Always allow admins to view lease logs.

diff --git a/Engine/Source/Programs/Horde/Horde.Server/Logs/LogsController.cs b/Engine/Source/Programs/Horde/Horde.Server/Logs/LogsController.cs
--- a/Engine/Source/Programs/Horde/Horde.Server/Logs/LogsController.cs
+++ b/Engine/Source/Programs/Horde/Horde.Server/Logs/LogsController.cs
@@ -370,6 +370,10 @@ namespace Horde.Server.Logs
 		async Task<bool> AuthorizeAsync(ILogFile logFile, AclAction action, ClaimsPrincipal user)
 		{
 			GlobalConfig globalConfig = _globalConfig.Value;
+			if (user.HasAdminClaim())
+			{
+				return true;
+			}
 			if (logFile.LeaseId != null && user.HasLeaseClaim(logFile.LeaseId.Value))
 			{
 				return true;
