Horde: Improvements to artifact zip streaming
#skipci
#rnx

diff --git a/Engine/Source/Programs/Horde/HordeDashboard/src/backend/Backend.ts b/Engine/Source/Programs/Horde/HordeDashboard/src/backend/Backend.ts
--- a/Engine/Source/Programs/Horde/HordeDashboard/src/backend/Backend.ts
+++ b/Engine/Source/Programs/Horde/HordeDashboard/src/backend/Backend.ts
@@ -584,42 +584,21 @@ export class Backend {
         });
     }
 
-    downloadArtifactV2(artifactId: string, path: string, filename: string): Promise<boolean> {
-
-        return new Promise<any>((resolve, reject) => {
-            this.backend.get(`/api/v2/artifacts/${artifactId}/file`, { params: { path: path }, responseBlob: true }).then(response => {
-                const url = window.URL.createObjectURL(response.data);
-                const link = document.createElement('a');
-                link.href = url;
-                link.setAttribute('download', filename);
-                link.click();
-                setTimeout(() => {
-                    window.URL.revokeObjectURL(url);
-                }, 100);
-                resolve(true);
-            }).catch(reason => {
-                reject(reason);
-            });
-        });
+    downloadArtifactV2(artifactId: string, path: string) {
+        const url = `/api/v2/artifacts/${artifactId}/file?path=${encodeURIComponent(path)}`;
+        window.location.assign(url);
     }
 
-    downloadArtifactZipV2(artifactId: string, request: CreateZipRequest, filename: string): Promise<boolean> {
+    downloadArtifactZipV2(artifactId: string, request: CreateZipRequest) {
+        
+        const filter = request.filter.map(f => `filter=${encodeURIComponent(f)}`).join("&")
 
-        return new Promise<any>((resolve, reject) => {
-            this.backend.post(`/api/v2/artifacts/${artifactId}/zip`, request, { responseBlob: true }).then(response => {
-                const url = window.URL.createObjectURL(response.data);
-                const link = document.createElement('a');
-                link.href = url;
-                link.setAttribute('download', filename);
-                link.click();
-                setTimeout(() => {
-                    window.URL.revokeObjectURL(url);
-                }, 100);
-                resolve(true);
-            }).catch(reason => {
-                reject(reason);
-            });
-        });
+        let url = `/api/v2/artifacts/${artifactId}/zip`;
+        if (filter.length) {
+            url += "?" + filter;
+        }
+
+        window.location.assign(url);
 
     }
 
diff --git a/Engine/Source/Programs/Horde/HordeDashboard/src/components/artifacts/ArtifactsModal.tsx b/Engine/Source/Programs/Horde/HordeDashboard/src/components/artifacts/ArtifactsModal.tsx
--- a/Engine/Source/Programs/Horde/HordeDashboard/src/components/artifacts/ArtifactsModal.tsx
+++ b/Engine/Source/Programs/Horde/HordeDashboard/src/components/artifacts/ArtifactsModal.tsx
@@ -474,13 +474,12 @@ const DownloadButton: React.FC<{ handler: ArtifactsHandler }> = observer(({ hand
             const item = selection[0] as BrowserItem;
             if (item.type === BrowserType.File) {
 
-               try {
-                  setDownloading(true);
-                  await backend.downloadArtifactV2(handler.artifact.id, (handler.path ? handler.path + "/" : "") + item.text, item.text);
+               try {                  
+                  backend.downloadArtifactV2(handler.artifact.id, (handler.path ? handler.path + "/" : "") + item.text);
                } catch (err) {
                   console.error(err);
                } finally {
-                  setDownloading(false);
+                  
                }
 
                return;
@@ -511,16 +510,13 @@ const DownloadButton: React.FC<{ handler: ArtifactsHandler }> = observer(({ hand
          if (context === "step-trace") {
             contextName = "trace";
          }
-
-         const filename = "horde-" + contextName + "-artifacts-" + handler.jobId + '-' + handler.stepId + ".zip";
-
-         try {
-            setDownloading(true);
-            await backend.downloadArtifactZipV2(handler.artifact.id, { filter: filters }, filename);
+         
+         try {            
+            backend.downloadArtifactZipV2(handler.artifact.id, { filter: filters });
          } catch (err) {
             console.error(err);
          } finally {
-            setDownloading(false);
+            
          }
 
       }}>{buttonText}</PrimaryButton>
