Horde: Adding lease and sesson ids to logfile message
#rnx

diff --git a/Engine/Source/Programs/Horde/Horde.Server/Logs/LogFileMessages.cs b/Engine/Source/Programs/Horde/Horde.Server/Logs/LogFileMessages.cs
--- a/Engine/Source/Programs/Horde/Horde.Server/Logs/LogFileMessages.cs
+++ b/Engine/Source/Programs/Horde/Horde.Server/Logs/LogFileMessages.cs
@@ -60,6 +60,16 @@ namespace Horde.Server.Logs
 		/// </summary>
         public string JobId { get; set; }
 
+		/// <summary>
+		/// The lease allowed to write to this log
+		/// </summary>
+		public string? LeaseId { get; }
+
+		/// <summary>
+		/// The session allowed to write to this log
+		/// </summary>
+		public string? SessionId { get; }
+
 		/// <summary>
 		/// Type of events stored in this log
 		/// </summary>
@@ -79,6 +89,8 @@ namespace Horde.Server.Logs
 		{
             Id = logFile.Id.ToString();
             JobId = logFile.JobId.ToString();
+			LeaseId = logFile.LeaseId.ToString();
+			SessionId = logFile.SessionId.ToString();
 			Type = logFile.Type;
             LineCount = metadata.MaxLineIndex;
 		}
