Horde: Add redirect for lease id to log data
#rnx

diff --git a/Engine/Source/Programs/Horde/Horde.Server/Agents/Leases/LeasesController.cs b/Engine/Source/Programs/Horde/Horde.Server/Agents/Leases/LeasesController.cs
--- a/Engine/Source/Programs/Horde/Horde.Server/Agents/Leases/LeasesController.cs
+++ b/Engine/Source/Programs/Horde/Horde.Server/Agents/Leases/LeasesController.cs
@@ -141,6 +141,34 @@ namespace Horde.Server.Agents.Leases
 			return new GetAgentLeaseResponse(lease, details, agentRate);
 		}
 
+		/// <summary>
+		/// Get lease log, redirecting from lease id to log id
+		/// </summary>
+		/// <param name="leaseId">Unique id of the particular lease</param>
+		/// <returns>Redirect to log endpoint</returns>
+		[HttpGet]
+		[Route("/api/v1/leases/{leaseId}/log")]
+		public async Task<ActionResult> GetLeaseLogAsync(LeaseId leaseId)
+		{
+			if (!_globalConfig.Value.Authorize(LeaseAclAction.ViewLeases, User))
+			{
+				return Forbid(LeaseAclAction.ViewLeases);
+			}
+
+			ILease? lease = await _agentService.GetLeaseAsync(leaseId);
+			if (lease == null)
+			{
+				return NotFound(leaseId);
+			}
+
+			if (lease.LogId == null)
+			{
+				return NotFound("null lease.LogId");
+			}
+
+			return new RedirectResult($"/api/v1/logs/{lease.LogId}");
+		}
+
 		/// <summary>
 		/// Update a particular lease
 		/// </summary>
