Horde: Add clickable agent id and lease id in log renders
#skipci
#rnx

diff --git a/Engine/Source/Programs/Horde/HordeDashboard/src/backend/Api.ts b/Engine/Source/Programs/Horde/HordeDashboard/src/backend/Api.ts
--- a/Engine/Source/Programs/Horde/HordeDashboard/src/backend/Api.ts
+++ b/Engine/Source/Programs/Horde/HordeDashboard/src/backend/Api.ts
@@ -1364,6 +1364,12 @@ export type GetLogFileResponse = {
 	/** Unique id of the job for this log file */
 	jobId: string;
 
+	/** Unique id of the lease for this log file */
+	leaseId?: string;
+
+	/** Unique id of the session for this log file */
+	sessionId?: string;
+
 	/** Type of events stored in this log */
 	type: LogType;
 
diff --git a/Engine/Source/Programs/Horde/HordeDashboard/src/backend/Backend.ts b/Engine/Source/Programs/Horde/HordeDashboard/src/backend/Backend.ts
--- a/Engine/Source/Programs/Horde/HordeDashboard/src/backend/Backend.ts
+++ b/Engine/Source/Programs/Horde/HordeDashboard/src/backend/Backend.ts
@@ -1,7 +1,7 @@
 // Copyright Epic Games, Inc. All Rights Reserved.
 
 import templateCache from '../backend/TemplateCache';
-import { AgentData, AgentQuery, ArtifactData, AuditLogEntry, AuditLogQuery, BatchUpdatePoolRequest, ChangeSummaryData, CreateBisectTaskRequest, CreateBisectTaskResponse, CreateDeviceRequest, CreateDeviceResponse, CreateExternalIssueRequest, CreateExternalIssueResponse, CreateJobRequest, CreateJobResponse, CreateNoticeRequest, CreatePoolRequest, CreateSoftwareResponse, CreateSubscriptionRequest, CreateSubscriptionResponse, CreateZipRequest, DashboardPreference, DevicePoolTelemetryQuery, DeviceTelemetryQuery, EventData, FindArtifactsResponse, FindIssueResponse, FindJobTimingsResponse, GetAgentSoftwareChannelResponse, GetArtifactDirectoryResponse, GetArtifactZipRequest, GetBisectTaskResponse, GetDashboardConfigResponse, GetDevicePlatformResponse, GetDevicePoolResponse, GetDevicePoolTelemetryResponse, GetDeviceReservationResponse, GetDeviceResponse, GetDeviceTelemetryResponse, GetExternalIssueProjectResponse, GetExternalIssueResponse, GetGraphResponse, GetIssueStreamResponse, GetJobsTabResponse, GetJobStepRefResponse, GetJobStepTraceResponse, GetJobTimingResponse, GetLogEventResponse, GetNoticeResponse, GetNotificationResponse, GetPerforceServerStatusResponse, GetPoolResponse, GetServerInfoResponse, GetServerSettingsResponse, GetSoftwareResponse, GetSubscriptionResponse, GetTestDataDetailsResponse, GetTestDataRefResponse, GetTestMetaResponse, GetTestResponse, GetTestsRequest, GetTestStreamResponse, GetToolSummaryResponse, GetUserResponse, GetUtilizationTelemetryResponse, GlobalConfig, IssueData, IssueQuery, IssueQueryV2, JobData, JobQuery, JobsTabColumnType, JobStepOutcome, JobStreamQuery, JobTimingsQuery, LeaseData, LogData, LogLineData, PoolData, ProjectData, ScheduleData, ScheduleQuery, SearchLogFileResponse, ServerUpdateResponse, SessionData, StreamData, TabType, TestData, UpdateAgentRequest, UpdateBisectTaskRequest, UpdateDeviceRequest, UpdateGlobalConfigRequest, UpdateIssueRequest, UpdateJobRequest, UpdateLeaseRequest, UpdateNoticeRequest, UpdateNotificationsRequest, UpdatePoolRequest, UpdateServerSettingsRequest, UpdateStepRequest, UpdateStepResponse, UpdateTemplateRefRequest, UpdateUserRequest, UsersQuery } from './Api';
+import { AgentData, AgentQuery, ArtifactData, AuditLogEntry, AuditLogQuery, BatchUpdatePoolRequest, ChangeSummaryData, CreateBisectTaskRequest, CreateBisectTaskResponse, CreateDeviceRequest, CreateDeviceResponse, CreateExternalIssueRequest, CreateExternalIssueResponse, CreateJobRequest, CreateJobResponse, CreateNoticeRequest, CreatePoolRequest, CreateSoftwareResponse, CreateSubscriptionRequest, CreateSubscriptionResponse, CreateZipRequest, DashboardPreference, DevicePoolTelemetryQuery, DeviceTelemetryQuery, EventData, FindArtifactsResponse, FindIssueResponse, FindJobTimingsResponse, GetAgentSoftwareChannelResponse, GetArtifactDirectoryResponse, GetArtifactZipRequest, GetBisectTaskResponse, GetDashboardConfigResponse, GetDevicePlatformResponse, GetDevicePoolResponse, GetDevicePoolTelemetryResponse, GetDeviceReservationResponse, GetDeviceResponse, GetDeviceTelemetryResponse, GetExternalIssueProjectResponse, GetExternalIssueResponse, GetGraphResponse, GetIssueStreamResponse, GetJobsTabResponse, GetJobStepRefResponse, GetJobStepTraceResponse, GetJobTimingResponse, GetLogEventResponse, GetLogFileResponse, GetNoticeResponse, GetNotificationResponse, GetPerforceServerStatusResponse, GetPoolResponse, GetServerInfoResponse, GetServerSettingsResponse, GetSoftwareResponse, GetSubscriptionResponse, GetTestDataDetailsResponse, GetTestDataRefResponse, GetTestMetaResponse, GetTestResponse, GetTestsRequest, GetTestStreamResponse, GetToolSummaryResponse, GetUserResponse, GetUtilizationTelemetryResponse, GlobalConfig, IssueData, IssueQuery, IssueQueryV2, JobData, JobQuery, JobsTabColumnType, JobStepOutcome, JobStreamQuery, JobTimingsQuery, LeaseData, LogData, LogLineData, PoolData, ProjectData, ScheduleData, ScheduleQuery, SearchLogFileResponse, ServerUpdateResponse, SessionData, StreamData, TabType, TestData, UpdateAgentRequest, UpdateBisectTaskRequest, UpdateDeviceRequest, UpdateGlobalConfigRequest, UpdateIssueRequest, UpdateJobRequest, UpdateLeaseRequest, UpdateNoticeRequest, UpdateNotificationsRequest, UpdatePoolRequest, UpdateServerSettingsRequest, UpdateStepRequest, UpdateStepResponse, UpdateTemplateRefRequest, UpdateUserRequest, UsersQuery } from './Api';
 import dashboard from './Dashboard';
 import { ChallengeStatus, Fetch } from './Fetch';
 import graphCache, { GraphQuery } from './GraphCache';
@@ -141,6 +141,14 @@ export class Backend {
         });
     }
 
+    async getLeaseLog(leaseId: string): Promise<GetLogFileResponse> {
+        return new Promise<GetLogFileResponse>((resolve, reject) => {
+            this.backend.get(`/api/v1/leases/${leaseId}/log`).then((response) => {
+                resolve(response.data as GetLogFileResponse);
+            }).catch(reason => { reject(reason); });
+        });
+    }
+
     getLease(leaseId: string, agentId?: string): Promise<LeaseData> {
 
         if (agentId) {
@@ -590,7 +598,7 @@ export class Backend {
     }
 
     downloadArtifactZipV2(artifactId: string, request: CreateZipRequest) {
-        
+
         const filter = request.filter.map(f => `filter=${encodeURIComponent(f)}`).join("&")
 
         let url = `/api/v2/artifacts/${artifactId}/zip`;
diff --git a/Engine/Source/Programs/Horde/HordeDashboard/src/components/AutomationReport.tsx b/Engine/Source/Programs/Horde/HordeDashboard/src/components/AutomationReport.tsx
--- a/Engine/Source/Programs/Horde/HordeDashboard/src/components/AutomationReport.tsx
+++ b/Engine/Source/Programs/Horde/HordeDashboard/src/components/AutomationReport.tsx
@@ -3,7 +3,7 @@ import { action, makeObservable, observable } from "mobx";
 import { observer } from "mobx-react-lite";
 import moment from "moment";
 import { useEffect, useState } from "react";
-import { Link } from 'react-router-dom';
+import { Link, useNavigate } from 'react-router-dom';
 import backend from "../backend";
 import { EventSeverity, GetLogEventResponse, GetTestDataDetailsResponse, GetTestDataRefResponse, GetTestResponse } from "../backend/Api";
 import { TestDataHandler } from "../backend/AutomationTestData";
@@ -301,6 +301,8 @@ class TestReportGenerator {
 
 const ErrorPane: React.FC<{ failure: TestFailureEvent }> = ({ failure }) => {
 
+   const navigate = useNavigate();
+
    const events = failure.events;
    if (!events) {
       return null;
@@ -318,7 +320,7 @@ const ErrorPane: React.FC<{ failure: TestFailureEvent }> = ({ failure }) => {
 
       const url = `/log/${failure.logId}?lineindex=${event.lineIndex}`;
 
-      const lines = event.lines.filter(line => line.message?.trim().length).map(line => <Stack styles={{ root: { paddingLeft: 8, paddingRight: 8, lineBreak: "anywhere", whiteSpace: "pre-wrap", lineHeight: 18, fontSize: 10, fontFamily: "Horde Cousine Regular, monospace, monospace" } }}> <Link className="log-link" target="_blank" to={url}>{renderLine(line, undefined, {})}</Link></Stack>);
+      const lines = event.lines.filter(line => line.message?.trim().length).map(line => <Stack styles={{ root: { paddingLeft: 8, paddingRight: 8, lineBreak: "anywhere", whiteSpace: "pre-wrap", lineHeight: 18, fontSize: 10, fontFamily: "Horde Cousine Regular, monospace, monospace" } }}> <Link className="log-link" target="_blank" to={url}>{renderLine(navigate, line, undefined, {})}</Link></Stack>);
 
       return (<Stack className={styles.itemCell} styles={{ root: { padding: 8, marginRight: 8 } }}><Stack className={event.severity === EventSeverity.Warning ? styles.gutterWarning : styles.gutter} styles={{ root: { padding: 0, margin: 0 } }}>
          <Stack styles={{ root: { paddingLeft: 14 } }}>
diff --git a/Engine/Source/Programs/Horde/HordeDashboard/src/components/ErrorPane.tsx b/Engine/Source/Programs/Horde/HordeDashboard/src/components/ErrorPane.tsx
--- a/Engine/Source/Programs/Horde/HordeDashboard/src/components/ErrorPane.tsx
+++ b/Engine/Source/Programs/Horde/HordeDashboard/src/components/ErrorPane.tsx
@@ -4,7 +4,7 @@ import { observer } from 'mobx-react-lite';
 import { IColumn, List, Stack, Text } from '@fluentui/react';
 import { getFocusStyle, getTheme, mergeStyleSets } from '@fluentui/react/lib/Styling';
 import React from 'react';
-import { Link } from 'react-router-dom';
+import { Link, useNavigate } from 'react-router-dom';
 import { EventData, EventSeverity, GetLabelResponse } from '../backend/Api';
 import { JobDetails } from '../backend/JobDetails';
 import { JobEventHandler } from '../backend/JobEventHandler';
@@ -56,6 +56,8 @@ const styles = mergeStyleSets({
 
 export const ErrorPane: React.FC<{ jobDetails: JobDetails; stepId: string; showErrors: boolean; count?: number }> = ({ jobDetails, stepId, showErrors, count }) => {
 
+   const navigate = useNavigate();
+
    if (!stepId) {
       return (<div />);
    }
@@ -99,7 +101,7 @@ export const ErrorPane: React.FC<{ jobDetails: JobDetails; stepId: string; showE
 
       const url = `/log/${item.logId}?lineindex=${item.lineIndex}`;
 
-      const lines = item.lines.filter(line => line.message?.trim().length).map(line => <Stack styles={{ root: { paddingLeft: 8, paddingRight: 8, lineBreak: "anywhere", whiteSpace: "pre-wrap", lineHeight: 18, fontSize: 10, fontFamily: "Horde Cousine Regular, monospace, monospace" } }}> <Link className="log-link" to={url}>{renderLine(line, undefined, {})}</Link></Stack>);
+      const lines = item.lines.filter(line => line.message?.trim().length).map(line => <Stack styles={{ root: { paddingLeft: 8, paddingRight: 8, lineBreak: "anywhere", whiteSpace: "pre-wrap", lineHeight: 18, fontSize: 10, fontFamily: "Horde Cousine Regular, monospace, monospace" } }}> <Link className="log-link" to={url}>{renderLine(navigate, line, undefined, {})}</Link></Stack>);
 
       return ( <Stack className={styles.itemCell} styles={{ root: { padding: 8, marginRight: 8 } }}><Stack className={item.severity === EventSeverity.Warning ? styles.gutterWarning : styles.gutter} styles={{ root: { padding: 0, margin: 0 } }}>
                <Stack styles={{ root: { paddingLeft: 14 } }}>
@@ -137,6 +139,8 @@ export const ErrorPane: React.FC<{ jobDetails: JobDetails; stepId: string; showE
 
 export const JobEventListPanel: React.FC<{ jobDetails: JobDetails, stepIds: string[], eventHandler: JobEventHandler }> = ({ jobDetails, stepIds, eventHandler }) => {
 
+   const navigate = useNavigate();
+
    if (!stepIds.length) {
       return null;
    }
@@ -190,7 +194,7 @@ export const JobEventListPanel: React.FC<{ jobDetails: JobDetails, stepIds: stri
 
       const url = `/log/${logId}?lineindex=${event.lineIndex}`;
 
-      const lines = event.lines.map(line => <Stack styles={{ root: { paddingLeft: 8, paddingRight: 8, lineBreak: "anywhere", whiteSpace: "pre-wrap", lineHeight: 18, fontSize: 10, fontFamily: "Horde Cousine Regular, monospace, monospace" } }}> {renderLine(line, undefined, {})}</Stack>);
+      const lines = event.lines.map(line => <Stack styles={{ root: { paddingLeft: 8, paddingRight: 8, lineBreak: "anywhere", whiteSpace: "pre-wrap", lineHeight: 18, fontSize: 10, fontFamily: "Horde Cousine Regular, monospace, monospace" } }}> {renderLine(navigate, line, undefined, {})}</Stack>);
       return (
          <Link className="log-link" to={url}>
             <Stack className={styles.itemCell} styles={{ root: { padding: 8, paddingLeft: 24, marginRight: 8 } }}><Stack className={event.severity === EventSeverity.Warning ? styles.gutterWarning : styles.gutter} styles={{ root: { padding: 0, margin: 0 } }}>
diff --git a/Engine/Source/Programs/Horde/HordeDashboard/src/components/IssueViewV2.tsx b/Engine/Source/Programs/Horde/HordeDashboard/src/components/IssueViewV2.tsx
--- a/Engine/Source/Programs/Horde/HordeDashboard/src/components/IssueViewV2.tsx
+++ b/Engine/Source/Programs/Horde/HordeDashboard/src/components/IssueViewV2.tsx
@@ -1333,6 +1333,7 @@ let lineKey = 0;
 
 export const ErrorPane: React.FC<{ events?: GetLogEventResponse[]; onClose?: () => void }> = ({ events }) => {
 
+   const navigate = useNavigate();
    useWindowSize();
 
    const listEvents = events?.sort((a, b) => {
@@ -1354,7 +1355,7 @@ export const ErrorPane: React.FC<{ events?: GetLogEventResponse[]; onClose?: ()
 
       const url = `/log/${item.logId}?lineindex=${item.lineIndex}`;
 
-      const lines = item.lines.filter(line => line.message?.trim().length).map(line => <Stack key={`errorpane_line_${item.lineIndex}_${lineKey++}`} styles={{ root: { paddingLeft: 8, paddingRight: 8, lineBreak: "normal", whiteSpace: "pre-wrap", lineHeight: 18, fontSize: 10, fontFamily: "Horde Cousine Regular, monospace, monospace" } }}> <Link className="log-link" to={url}>{renderLine(line, undefined, {})}</Link></Stack>);
+      const lines = item.lines.filter(line => line.message?.trim().length).map(line => <Stack key={`errorpane_line_${item.lineIndex}_${lineKey++}`} styles={{ root: { paddingLeft: 8, paddingRight: 8, lineBreak: "normal", whiteSpace: "pre-wrap", lineHeight: 18, fontSize: 10, fontFamily: "Horde Cousine Regular, monospace, monospace" } }}> <Link className="log-link" to={url}>{renderLine(navigate, line, undefined, {})}</Link></Stack>);
 
       return (<Stack className={errorStyles.itemCell} style={{ padding: 8 }}><Stack className={item.severity === EventSeverity.Warning ? errorStyles.gutterWarning : errorStyles.gutter} styles={{ root: { padding: 0, margin: 0 } }}>
          <Stack styles={{ root: { paddingLeft: 14 } }}>
diff --git a/Engine/Source/Programs/Horde/HordeDashboard/src/components/LogRender.tsx b/Engine/Source/Programs/Horde/HordeDashboard/src/components/LogRender.tsx
--- a/Engine/Source/Programs/Horde/HordeDashboard/src/components/LogRender.tsx
+++ b/Engine/Source/Programs/Horde/HordeDashboard/src/components/LogRender.tsx
@@ -3,11 +3,15 @@ import { Stack, Text } from '@fluentui/react';
 import React from 'react';
 import Highlight from 'react-highlighter';
 import { IssueData, LogLine } from '../backend/Api';
+import backend from '../backend';
+import { NavigateFunction } from 'react-router-dom';
 
 enum TagType {
    None,
    SourceFile,
-   MSDNCode
+   MSDNCode,
+   AgentId,
+   LeaseId
 }
 
 export type LogItem = {
@@ -50,7 +54,7 @@ const renderMessage = (line: LogLine, lineNumber: number | undefined, logStyle:
 };
 
 
-const renderTags = (line: LogLine, lineNumber: number | undefined, logStyle: any, tags: string[], search?: string) => {
+const renderTags = (navigate: NavigateFunction, line: LogLine, lineNumber: number | undefined, logStyle: any, tags: string[], search?: string) => {
 
    if (!line || !line.format || !line.properties) {
       return <Stack styles={{ root: { color: "#000000", paddingLeft: 8, whiteSpace: "pre", tabSize: "3" } }}>Internal log line format error</Stack>;
@@ -136,16 +140,59 @@ const renderTags = (line: LogLine, lineNumber: number | undefined, logStyle: any
             tagType = TagType.MSDNCode;
          }
 
+         if (type === "LeaseId") {
+            tagType = TagType.LeaseId;
+         }
+
+         if (type === "AgentId") {
+            tagType = TagType.AgentId;
+         }
+
       }
 
       if (tagType === TagType.None || !record) {
 
          return <Highlight key={key} search={search ? search : ""} className={logStyle.logLine}>{text}</Highlight>;
 
+      } else if (tagType === TagType.LeaseId) {
+
+         const navigateToLeaseLog = async (toplevel: boolean) => {
+
+            const logData = await backend.getLease(text);
+            const url = `/log/${logData?.logId}?leaseId=${text}`;
+            if (!toplevel) {
+               navigate(url)
+            } else {
+               window.open(url, "_blank");
+            }
+         }
+
+         return <a key={key} href="/"
+            
+            onAuxClick={(ev) =>
+            {
+               ev.preventDefault();
+               ev.stopPropagation()
+               navigateToLeaseLog(true)
+            }}
+         
+            onClick={(ev) => {
+               ev.stopPropagation();
+               ev.preventDefault();
+               navigateToLeaseLog(!!ev?.ctrlKey || !!ev.metaKey)
+            }}>
+            <Highlight search={search ? search : ""} className={logStyle.logLine}>{text}</Highlight>
+         </a>;
+
       } else if (tagType === TagType.MSDNCode) {
 
          return <a key={key} target="_blank" rel="noopener noreferrer" href={`https://msdn.microsoft.com/query/dev16.query?appId=Dev16IDEF1&l=EN-US&k=k(${text.toLowerCase()})&rd=true`} onClick={(ev) => ev.stopPropagation()}><Highlight search={search ? search : ""} className={logStyle.logLine}>{text}</Highlight></a>;
 
+      } else if (tagType === TagType.AgentId) {
+         const search = window.location.search;
+         const url = `${window.location.pathname}` + (search ? search + `&agentId=${encodeURIComponent(text)}` : `?agentId=${encodeURIComponent(text)}`);
+         return <a key={key} href="/" onClick={async (ev) => { ev.stopPropagation(); ev.preventDefault(); navigate(url) }}><Highlight search={search ? search : ""} className={logStyle.logLine}>{text}</Highlight></a>;
+
       } else if (tagType === TagType.SourceFile) {
 
          let depotPath = record.depotPath;
@@ -214,7 +261,7 @@ const renderTags = (line: LogLine, lineNumber: number | undefined, logStyle: any
 
 };
 
-export const renderLine = (line: LogLine | undefined, lineNumber: number | undefined, logStyle: any, search?: string) => {
+export const renderLine = (navigate: NavigateFunction, line: LogLine | undefined, lineNumber: number | undefined, logStyle: any, search?: string) => {
 
    if (!line) {
       return null;
@@ -245,7 +292,7 @@ export const renderLine = (line: LogLine | undefined, lineNumber: number | undef
 
 
    if (tags.length && format && line.properties) {
-      return renderTags(line, lineNumber, logStyle, tags, search);
+      return renderTags(navigate, line, lineNumber, logStyle, tags, search);
    }
 
    return renderMessage(line, lineNumber, logStyle, search);
diff --git a/Engine/Source/Programs/Horde/HordeDashboard/src/components/LogView.tsx b/Engine/Source/Programs/Horde/HordeDashboard/src/components/LogView.tsx
--- a/Engine/Source/Programs/Horde/HordeDashboard/src/components/LogView.tsx
+++ b/Engine/Source/Programs/Horde/HordeDashboard/src/components/LogView.tsx
@@ -697,7 +697,7 @@ export const LogList: React.FC<{ logId: string }> = observer(({ logId }) => {
                         <Stack className={gutterStyle}></Stack>
                         {(!item.issueId || !ev) && <Stack styles={{ root: { color: "#8a8a8a", width: tsWidth, whiteSpace: "nowrap", fontSize: handler.fontSize, userSelect: "none" } }}> {timestamp}</Stack>}
                         {!!item.issueId && !!ev && <IssueButton item={item} event={ev!} />}
-                        <div className={styles.logLineOuter}> <Stack styles={{ root: { paddingLeft: 8, paddingRight: 8 } }}> {renderLine(item.line, item.lineNumber, handler.lineRenderStyle, searchState.search)}</Stack></div>
+                        <div className={styles.logLineOuter}> <Stack styles={{ root: { paddingLeft: 8, paddingRight: 8 } }}> {renderLine(navigate, item.line, item.lineNumber, handler.lineRenderStyle, searchState.search)}</Stack></div>
                      </Stack>
                   </Stack>
                </div>
diff --git a/Engine/Source/Programs/Horde/HordeDashboard/src/components/jobDetailsV2/StepErrorPanel.tsx b/Engine/Source/Programs/Horde/HordeDashboard/src/components/jobDetailsV2/StepErrorPanel.tsx
--- a/Engine/Source/Programs/Horde/HordeDashboard/src/components/jobDetailsV2/StepErrorPanel.tsx
+++ b/Engine/Source/Programs/Horde/HordeDashboard/src/components/jobDetailsV2/StepErrorPanel.tsx
@@ -4,7 +4,7 @@ import { List, Stack, Text } from "@fluentui/react";
 import { getFocusStyle, getTheme, mergeStyleSets } from '@fluentui/react/lib/Styling';
 import { observer } from "mobx-react-lite";
 import React, { useEffect } from 'react';
-import { Link } from 'react-router-dom';
+import { Link, useNavigate } from 'react-router-dom';
 import backend from "../../backend";
 import { EventData, EventSeverity } from '../../backend/Api';
 import { ISideRailLink } from "../../base/components/SideRail";
@@ -150,6 +150,8 @@ const styles = mergeStyleSets({
 
 const ErrorPane: React.FC<{ jobDetails: JobDetailsV2; view: StepSummaryErrorsView, stepId: string; showErrors: boolean; count?: number }> = ({ jobDetails, view, stepId, showErrors, count }) => {
 
+   const navigate = useNavigate();
+
    if (!stepId) {
       return (<div />);
    }
@@ -178,7 +180,7 @@ const ErrorPane: React.FC<{ jobDetails: JobDetailsV2; view: StepSummaryErrorsVie
 
       const url = `/log/${item.logId}?lineindex=${item.lineIndex}`;
 
-      const lines = item.lines.filter(line => line.message?.trim().length).map(line => <Stack key={`steperrorpane_line_${item.lineIndex}`} styles={{ root: { paddingLeft: 8, paddingRight: 8, lineBreak: "anywhere", whiteSpace: "pre-wrap", lineHeight: 18, fontSize: 10, fontFamily: "Horde Cousine Regular, monospace, monospace" } }}> <Link className="log-link" to={url}>{renderLine(line, undefined, {})}</Link></Stack>);
+      const lines = item.lines.filter(line => line.message?.trim().length).map(line => <Stack key={`steperrorpane_line_${item.lineIndex}`} styles={{ root: { paddingLeft: 8, paddingRight: 8, lineBreak: "anywhere", whiteSpace: "pre-wrap", lineHeight: 18, fontSize: 10, fontFamily: "Horde Cousine Regular, monospace, monospace" } }}> <Link className="log-link" to={url}>{renderLine(navigate, line, undefined, {})}</Link></Stack>);
 
       return (<Stack className={styles.itemCell} styles={{ root: { padding: 8, marginRight: 8 } }}><Stack className={item.severity === EventSeverity.Warning ? styles.gutterWarning : styles.gutter} styles={{ root: { padding: 0, margin: 0 } }}>
          <Stack styles={{ root: { paddingLeft: 14 } }}>
