Horde: Optimize JobQueueStrategy by assuming queue size is zero during maintenance windows

diff --git a/Engine/Source/Programs/Horde/Horde.Server.Tests/Fleet/JobQueueStrategyTest.cs b/Engine/Source/Programs/Horde/Horde.Server.Tests/Fleet/JobQueueStrategyTest.cs
--- a/Engine/Source/Programs/Horde/Horde.Server.Tests/Fleet/JobQueueStrategyTest.cs
+++ b/Engine/Source/Programs/Horde/Horde.Server.Tests/Fleet/JobQueueStrategyTest.cs
@@ -31,6 +31,16 @@ namespace Horde.Server.Tests.Fleet
 			Assert.AreEqual(5, poolQueueSizes[pool.Id]);
 		}
 		
+		[TestMethod]
+		public async Task DowntimeActive()
+		{
+			(JobQueueStrategy strategy, PoolSizeResult poolSizeResult, IPool pool, List<IAgent> agents) = await SetUpJobsAsync(1, 5, isDowntimeActive: true);
+			await Clock.AdvanceAsync(TimeSpan.FromSeconds(strategy.Settings.ReadyTimeThresholdSec) + TimeSpan.FromSeconds(5));
+			Dictionary<PoolId, int> poolQueueSizes = await strategy.GetPoolQueueSizesAsync(Clock.UtcNow - TimeSpan.FromHours(2));
+			Assert.AreEqual(1, poolQueueSizes.Count);
+			Assert.AreEqual(0, poolQueueSizes[pool.Id]);
+		}
+		
 		[TestMethod]
 		public async Task EmptyJobQueue()
 		{
@@ -79,9 +89,9 @@ namespace Horde.Server.Tests.Fleet
 			await AssertAgentCount(25, 7);
 		}
 
-		public async Task AssertAgentCount(int numBatchesReady, int expectedAgentDelta, bool waitedBeyondThreshold = true, int numAgents = 8)
+		public async Task AssertAgentCount(int numBatchesReady, int expectedAgentDelta, bool waitedBeyondThreshold = true, int numAgents = 8, bool isDowntimeActive = false)
 		{
-			(JobQueueStrategy strategy, PoolSizeResult poolSizeResult, IPool pool, List<IAgent> agents) = await SetUpJobsAsync(1, numBatchesReady, numAgents);
+			(JobQueueStrategy strategy, PoolSizeResult poolSizeResult, IPool pool, List<IAgent> agents) = await SetUpJobsAsync(1, numBatchesReady, numAgents, isDowntimeActive);
 			TimeSpan timeToWait = waitedBeyondThreshold
 				? TimeSpan.FromSeconds(strategy.Settings.ReadyTimeThresholdSec) + TimeSpan.FromSeconds(5)
 				: TimeSpan.FromSeconds(15);
@@ -97,7 +107,7 @@ namespace Horde.Server.Tests.Fleet
 		/// </summary>
 		/// <param name="numBatchesRunning">Num of job batches that should be in state running</param>
 		/// <param name="numBatchesReady">Num of job batches that should be in state waiting</param>
-		private async Task<(JobQueueStrategy, PoolSizeResult, IPool, List<IAgent> agents)> SetUpJobsAsync(int numBatchesRunning, int numBatchesReady, int numAgents = 8)
+		private async Task<(JobQueueStrategy, PoolSizeResult, IPool, List<IAgent> agents)> SetUpJobsAsync(int numBatchesRunning, int numBatchesReady, int numAgents = 8, bool isDowntimeActive = false)
 		{
 			IPool pool = await PoolService.CreatePoolAsync("bogusPool1", new AddPoolOptions { EnableAutoscaling = true, MinAgents = 0, NumReserveAgents = 0 });
 			List<IAgent> agents = new();
@@ -143,7 +153,7 @@ namespace Horde.Server.Tests.Fleet
 				await JobCollection.TryUpdateBatchAsync(job, graph, job.Batches[0].Id, null, JobStepBatchState.Ready, null);
 			}
 
-			return (new (JobCollection, GraphCollection, StreamCollection, Clock, Cache, GlobalConfig), poolSize, pool, agents);
+			return (new (JobCollection, GraphCollection, StreamCollection, Clock, Cache, isDowntimeActive, GlobalConfig), poolSize, pool, agents);
 		}
 		
 		private async Task<IJob> AddPlaceholderJob(IGraph graph, StreamId streamId, string nodeNameToExecute)
diff --git a/Engine/Source/Programs/Horde/Horde.Server/Agents/Fleet/FleetService.cs b/Engine/Source/Programs/Horde/Horde.Server/Agents/Fleet/FleetService.cs
--- a/Engine/Source/Programs/Horde/Horde.Server/Agents/Fleet/FleetService.cs
+++ b/Engine/Source/Programs/Horde/Horde.Server/Agents/Fleet/FleetService.cs
@@ -381,7 +381,7 @@ namespace Horde.Server.Agents.Fleet
 						{
 							case PoolSizeStrategy.JobQueue:
 								JobQueueSettings jqSettings = DeserializeConfig<JobQueueSettings>(info.Config);
-								JobQueueStrategy jqStrategy = new (_jobCollection, _graphCollection, _streamCollection, _clock, _cache, _globalConfig, jqSettings);
+								JobQueueStrategy jqStrategy = new (_jobCollection, _graphCollection, _streamCollection, _clock, _cache, _downtimeService.IsDowntimeActive, _globalConfig, jqSettings);
 								return info.ExtraAgentCount != 0 ? new ExtraAgentCountStrategy(jqStrategy, info.ExtraAgentCount) : jqStrategy;
 							
 							case PoolSizeStrategy.LeaseUtilization:
@@ -409,7 +409,7 @@ namespace Horde.Server.Agents.Fleet
 			switch (pool.SizeStrategy ?? _settings.Value.DefaultAgentPoolSizeStrategy)
 			{
 				case PoolSizeStrategy.JobQueue:
-					return new JobQueueStrategy(_jobCollection, _graphCollection, _streamCollection, _clock, _cache, _globalConfig, pool.JobQueueSettings);
+					return new JobQueueStrategy(_jobCollection, _graphCollection, _streamCollection, _clock, _cache, _downtimeService.IsDowntimeActive, _globalConfig, pool.JobQueueSettings);
 				case PoolSizeStrategy.LeaseUtilization:
 					LeaseUtilizationSettings luSettings = new();
 					if (pool.MinAgents != null)
diff --git a/Engine/Source/Programs/Horde/Horde.Server/Agents/Fleet/JobQueueStrategy.cs b/Engine/Source/Programs/Horde/Horde.Server/Agents/Fleet/JobQueueStrategy.cs
--- a/Engine/Source/Programs/Horde/Horde.Server/Agents/Fleet/JobQueueStrategy.cs
+++ b/Engine/Source/Programs/Horde/Horde.Server/Agents/Fleet/JobQueueStrategy.cs
@@ -98,6 +98,7 @@ namespace Horde.Server.Agents.Fleet
 		private readonly IStreamCollection _streamCollection;
 		private readonly IClock _clock;
 		private readonly IMemoryCache _cache;
+		private readonly bool _isDowntimeActive;
 		private readonly IOptionsMonitor<GlobalConfig> _globalConfig;
 
 		/// <summary>
@@ -108,15 +109,17 @@ namespace Horde.Server.Agents.Fleet
 		/// <param name="streamCollection"></param>
 		/// <param name="clock"></param>
 		/// <param name="cache"></param>
+		/// <param name="isDowntimeActive"></param>
 		/// <param name="globalConfig"></param>
 		/// <param name="settings"></param>
-		public JobQueueStrategy(IJobCollection jobs, IGraphCollection graphs, IStreamCollection streamCollection, IClock clock, IMemoryCache cache, IOptionsMonitor<GlobalConfig> globalConfig, JobQueueSettings? settings = null)
+		public JobQueueStrategy(IJobCollection jobs, IGraphCollection graphs, IStreamCollection streamCollection, IClock clock, IMemoryCache cache, bool isDowntimeActive, IOptionsMonitor<GlobalConfig> globalConfig, JobQueueSettings? settings = null)
 		{
 			_jobs = jobs;
 			_graphs = graphs;
 			_streamCollection = streamCollection;
 			_clock = clock;
 			_cache = cache;
+			_isDowntimeActive = isDowntimeActive;
 			_globalConfig = globalConfig;
 			Settings = settings ?? new JobQueueSettings();
 		}
@@ -188,6 +191,13 @@ namespace Horde.Server.Agents.Fleet
 			List<(PoolId PoolId, int QueueSize)> poolsWithQueueSize = jobBatches.GroupBy(t => t.PoolId).Select(t => (t.Key, t.Count())).ToList();
 
 			span.SetAttribute("numPools", poolsWithQueueSize.Count);
+
+			if (_isDowntimeActive)
+			{
+				// As an optimization, assume queue size is zero during maintenance windows.
+				return poolsWithQueueSize.ToDictionary(x => x.PoolId, x => 0);
+			}
+			
 			return poolsWithQueueSize.ToDictionary(x => x.PoolId, x => x.QueueSize);
 		}
 
