Horde: Mark ephemeral and expired agents as deleted

But delay actual deletion from database for 7 days for easier debugging.

diff --git a/Engine/Source/Programs/Horde/Horde.Server.Tests/AgentServiceTest.cs b/Engine/Source/Programs/Horde/Horde.Server.Tests/AgentServiceTest.cs
--- a/Engine/Source/Programs/Horde/Horde.Server.Tests/AgentServiceTest.cs
+++ b/Engine/Source/Programs/Horde/Horde.Server.Tests/AgentServiceTest.cs
@@ -178,29 +178,25 @@ public class AgentServiceTest : TestSetup
 	[TestMethod]
 	public async Task EphemeralTest()
 	{
-		IAgent agent1 = await CreateAgentAsync(new PoolId("pool1"), ephemeral: true);
-		IAgent agent2 = await CreateAgentAsync(new PoolId("pool1"), ephemeral: false);
-		Assert.IsTrue(agent1.Ephemeral);
-		Assert.IsFalse(agent2.Ephemeral);
+		IAgent agent = await CreateAgentAsync(new PoolId("pool1"), ephemeral: true);
+		Assert.IsTrue(agent.Ephemeral);
 		
-		agent1 = (await AgentService.GetAgentAsync(agent1.Id))!;
-		agent2 = (await AgentService.GetAgentAsync(agent2.Id))!;
-		Assert.IsTrue(agent1.Ephemeral);
-		Assert.IsFalse(agent2.Ephemeral);
-		Assert.AreEqual(AgentStatus.Ok, agent1.Status);
-		Assert.AreEqual(AgentStatus.Ok, agent2.Status);
-
-		// Let background task run for purging outdated sessions
-		await AgentService.StartAsync(CancellationToken.None);
+		agent = (await AgentService.GetAgentAsync(agent.Id))!;
+		Assert.IsTrue(agent.Ephemeral);
+		Assert.AreEqual(AgentStatus.Ok, agent.Status);
+
+		// Let background task run for purging outdated sessions, which will terminate session for our agent
 		await Clock.AdvanceAsync(TimeSpan.FromHours(1));
+		await AgentService.TickAsync(CancellationToken.None);
+		
+		// Ephemeral agent is marked as deleted once its session is terminated
+		Assert.IsTrue((await AgentService.GetAgentAsync(agent.Id))!.Deleted);
 		
-		// Ephemeral agent is deleted once its session is terminated
-		Assert.IsNull((await AgentService.GetAgentAsync(agent1.Id)));
+		await Clock.AdvanceAsync(TimeSpan.FromDays(8));
+		await AgentService.TickAsync(CancellationToken.None);
 		
-		// Normal agent is stopped but kept around
-		agent2 = (await AgentService.GetAgentAsync(agent2.Id))!;
-		Assert.AreEqual(AgentStatus.Stopped, agent2.Status);
-		Assert.IsFalse(agent2.Deleted);
+		// Once more time has passed, the ephemeral agent marked as deleted is removed from database
+		Assert.IsNull(await AgentService.GetAgentAsync(agent.Id));
 	}
 
 	private static ClaimsPrincipal GetUser(IAgent agent)
diff --git a/Engine/Source/Programs/Horde/Horde.Server/Agents/AgentCollection.cs b/Engine/Source/Programs/Horde/Horde.Server/Agents/AgentCollection.cs
--- a/Engine/Source/Programs/Horde/Horde.Server/Agents/AgentCollection.cs
+++ b/Engine/Source/Programs/Horde/Horde.Server/Agents/AgentCollection.cs
@@ -245,6 +245,13 @@ namespace Horde.Server.Agents
 			List<AgentDocument> results = await _agents.Find(x => x.SessionId.HasValue && !(x.SessionExpiresAt > utcNow)).Limit(maxAgents).ToListAsync();
 			return results.ConvertAll<IAgent>(x => x);
 		}
+		
+		/// <inheritdoc/>
+		public async Task<List<IAgent>> FindDeletedAsync()
+		{
+			List<AgentDocument> results = await _agents.Find(x => x.Deleted).ToListAsync();
+			return results.ConvertAll<IAgent>(x => x);
+		}
 
 		/// <summary>
 		/// Update a single document
@@ -490,10 +497,9 @@ namespace Horde.Server.Agents
 			update = update.Set(x => x.Status, AgentStatus.Stopped);
 			update = update.Set(x => x.LastStatusChange, _clock.UtcNow);
 
-			bool deleted = agent.Deleted || agent.Ephemeral;
-			if (deleted != agent.Deleted)
+			if (agent.Ephemeral)
 			{
-				update = update.Set(x => x.Deleted, agent.Deleted);
+				update = update.Set(x => x.Deleted, true);
 			}
 
 			return await TryUpdateAsync(agent, update);
diff --git a/Engine/Source/Programs/Horde/Horde.Server/Agents/AgentService.cs b/Engine/Source/Programs/Horde/Horde.Server/Agents/AgentService.cs
--- a/Engine/Source/Programs/Horde/Horde.Server/Agents/AgentService.cs
+++ b/Engine/Source/Programs/Horde/Horde.Server/Agents/AgentService.cs
@@ -1030,9 +1030,16 @@ namespace Horde.Server.Agents
 		/// </summary>
 		/// <param name="stoppingToken">Token indicating the service is shutting down</param>
 		/// <returns>Async task</returns>
-		async ValueTask TickAsync(CancellationToken stoppingToken)
+		internal async ValueTask TickAsync(CancellationToken stoppingToken)
 		{
-			while (!stoppingToken.IsCancellationRequested)
+			await TerminateExpiredSessionsAsync(stoppingToken);
+			await DeleteExpiredEphemeralAgentsAsync(stoppingToken);
+			await CollectMetricsAsync();
+		}
+
+		private async Task TerminateExpiredSessionsAsync(CancellationToken cancellationToken)
+		{
+			while (!cancellationToken.IsCancellationRequested)
 			{
 				// Find all the agents which are ready to be expired
 				const int MaxAgents = 100;
@@ -1042,25 +1049,31 @@ namespace Horde.Server.Agents
 				// Transition each agent to being offline
 				foreach (IAgent expiredAgent in expiredAgents)
 				{
-					stoppingToken.ThrowIfCancellationRequested();
+					cancellationToken.ThrowIfCancellationRequested();
 					_logger.LogDebug("Terminating session {SessionId} for agent {Agent}", expiredAgent.SessionId, expiredAgent.Id);
-					IAgent? terminatedAgent = await TryTerminateSessionAsync(expiredAgent);
-
-					if (terminatedAgent is { Status: AgentStatus.Stopped, Ephemeral: true })
-					{
-						_logger.LogDebug("Deleting ephemeral agent {Agent}", terminatedAgent.Id);
-						await DeleteAgentAsync(terminatedAgent, true);
-					}
+					await TryTerminateSessionAsync(expiredAgent);
 				}
 
 				// Try again if we didn't fetch everything
-				if(expiredAgents.Count < MaxAgents)
+				if (expiredAgents.Count < MaxAgents)
 				{
 					break;
 				}
 			}
-
-			await CollectMetricsAsync();
+		}
+		
+		private async Task DeleteExpiredEphemeralAgentsAsync(CancellationToken cancellationToken)
+		{
+			foreach (IAgent agent in await Agents.FindDeletedAsync())
+			{
+				cancellationToken.ThrowIfCancellationRequested();
+				bool noStatusChangeDuringPeriod = _clock.UtcNow > agent.LastStatusChange + TimeSpan.FromDays(7);
+				if (agent.Status == AgentStatus.Stopped && noStatusChangeDuringPeriod)
+				{
+					_logger.LogDebug("Deleting ephemeral agent {Agent}", agent.Id);
+					await DeleteAgentAsync(agent, true);
+				}
+			}
 		}
 
 		private async Task CollectMetricsAsync()
diff --git a/Engine/Source/Programs/Horde/Horde.Server/Agents/IAgentCollection.cs b/Engine/Source/Programs/Horde/Horde.Server/Agents/IAgentCollection.cs
--- a/Engine/Source/Programs/Horde/Horde.Server/Agents/IAgentCollection.cs
+++ b/Engine/Source/Programs/Horde/Horde.Server/Agents/IAgentCollection.cs
@@ -69,12 +69,18 @@ namespace Horde.Server.Agents
 		Task<List<IAgent>> FindAsync(PoolId? poolId = null, DateTime? modifiedAfter = null, string? property = null, AgentStatus? status = null, bool? enabled = null, int? index = null, int? count = null);
 
 		/// <summary>
-		/// Finds all the expired agents
+		/// Finds all agents with an expired session
 		/// </summary>
 		/// <param name="utcNow">The current time</param>
 		/// <param name="maxAgents">Maximum number of agents to return</param>
 		/// <returns>List of agents</returns>
 		Task<List<IAgent>> FindExpiredAsync(DateTime utcNow, int maxAgents);
+		
+		/// <summary>
+		/// Finds all agents marked as deleted
+		/// </summary>
+		/// <returns>List of agents</returns>
+		Task<List<IAgent>> FindDeletedAsync();
 
 		/// <summary>
 		/// Update an agent's settings
