Horde: Missing files.

diff --git a/Engine/Source/Programs/Horde/Horde.Agent.Tests/JsonRpcLoggerTests.cs b/Engine/Source/Programs/Horde/Horde.Agent.Tests/JsonRpcLoggerTests.cs
--- a/Engine/Source/Programs/Horde/Horde.Agent.Tests/JsonRpcLoggerTests.cs
+++ b/Engine/Source/Programs/Horde/Horde.Agent.Tests/JsonRpcLoggerTests.cs
@@ -154,7 +154,7 @@ namespace Horde.Agent.Tests
 			LogNode file;
 			await using (FakeJsonRpcLoggerBackend sink = new FakeJsonRpcLoggerBackend(null!, "foo", innerSink, store, NullLogger.Instance))
 			{
-				await using (JsonRpcLogger logger = new JsonRpcLogger(sink, "foo", null, NullLogger.Instance))
+				await using (JsonRpcLogger logger = new JsonRpcLogger(sink, "foo", null, LogLevel.Information, NullLogger.Instance))
 				{
 					for (int idx = 0; idx < Count; idx++)
 					{
diff --git a/Engine/Source/Programs/Horde/Horde.Agent.Tests/WorkerServiceTest.cs b/Engine/Source/Programs/Horde/Horde.Agent.Tests/WorkerServiceTest.cs
--- a/Engine/Source/Programs/Horde/Horde.Agent.Tests/WorkerServiceTest.cs
+++ b/Engine/Source/Programs/Horde/Horde.Agent.Tests/WorkerServiceTest.cs
@@ -57,7 +57,7 @@ namespace Horde.Agent.Tests
 
 		class FakeServerLoggerFactory : IServerLoggerFactory
 		{
-			public IServerLogger CreateLogger(ISession session, string logId, string? jobId, string? batchId, string? stepId, bool? warnings = null, bool? useNewLogger = null) => new FakeServerLogger();
+			public IServerLogger CreateLogger(ISession session, string logId, string? jobId, string? batchId, string? stepId, bool? warnings = null, bool? useNewLogger = null, LogLevel outputLevel = LogLevel.Information) => new FakeServerLogger();
 		}
 
 		class FakeSessionStorageFactory : IServerStorageFactory
diff --git a/Engine/Source/Programs/Horde/Horde.Agent/Commands/Utilities/BuildGraphCommand.cs b/Engine/Source/Programs/Horde/Horde.Agent/Commands/Utilities/BuildGraphCommand.cs
--- a/Engine/Source/Programs/Horde/Horde.Agent/Commands/Utilities/BuildGraphCommand.cs
+++ b/Engine/Source/Programs/Horde/Horde.Agent/Commands/Utilities/BuildGraphCommand.cs
@@ -130,7 +130,7 @@ namespace Horde.Agent.Commands.Utilities
 
 			using (LogSink sink = new LogSink(runUatBat.Directory, logger))
 			{
-				await using (JsonRpcLogger jsonLogger = new JsonRpcLogger(sink, "log", null, logger))
+				await using (JsonRpcLogger jsonLogger = new JsonRpcLogger(sink, "log", null, LogLevel.Information, logger))
 				{
 					using (LogParser filter = new LogParser(jsonLogger, new List<string>()))
 					{
