Horde: Fix null value assigned to parent lease id in compute tasks.

diff --git a/Engine/Source/Programs/Horde/Horde.Server/Compute/ComputeService.cs b/Engine/Source/Programs/Horde/Horde.Server/Compute/ComputeService.cs
--- a/Engine/Source/Programs/Horde/Horde.Server/Compute/ComputeService.cs
+++ b/Engine/Source/Programs/Horde/Horde.Server/Compute/ComputeService.cs
@@ -100,7 +100,7 @@ namespace Horde.Server.Compute
 			computeTask.Key = UnsafeByteOperations.UnsafeWrap(AesTransport.CreateKey());
 			computeTask.Resources.Add(assignedResources);
 			computeTask.LogId = logId?.ToString();
-			computeTask.ParentLeaseId = parentLeaseId?.ToString();
+			computeTask.ParentLeaseId = parentLeaseId?.ToString() ?? String.Empty;
 			return computeTask;
 		}
 	}
