Horde: Allow specifying an explicit pool to use for compute jobs, rather than writing a condition. This is easier to specify on the client side and easier to optimize on the server side.

diff --git a/Engine/Source/Programs/Horde/Horde.Server/Agents/IAgent.cs b/Engine/Source/Programs/Horde/Horde.Server/Agents/IAgent.cs
--- a/Engine/Source/Programs/Horde/Horde.Server/Agents/IAgent.cs
+++ b/Engine/Source/Programs/Horde/Horde.Server/Agents/IAgent.cs
@@ -725,19 +725,25 @@ namespace Horde.Server.Agents
 		/// <returns>True if the new lease can be granted</returns>
 		public static bool MeetsRequirements(this IAgent agent, Requirements requirements, Dictionary<string, int> assignedResources)
 		{
-			return MeetsRequirements(agent, requirements.Condition, requirements.Resources, requirements.Exclusive, assignedResources);
+			PoolId? poolId = null;
+			if (!String.IsNullOrEmpty(requirements.Pool))
+			{
+				poolId = new PoolId(requirements.Pool);
+			}
+			return MeetsRequirements(agent, poolId, requirements.Condition, requirements.Resources, requirements.Exclusive, assignedResources);
 		}
 
 		/// <summary>
 		/// Determine whether it's possible to add a lease for the given resources
 		/// </summary>
 		/// <param name="agent">The agent to create a lease for</param>
-		/// <param name="exclusive">Whether t</param>
+		/// <param name="poolId">Pool to take the machine from</param>
 		/// <param name="condition">Condition to satisfy</param>
 		/// <param name="resources">Resources required to execute</param>
+		/// <param name="exclusive">Whether the lease needs to be executed exclusively on the machine</param>
 		/// <param name="assignedResources">Resources allocated to the task</param>
 		/// <returns>True if the new lease can be granted</returns>
-		public static bool MeetsRequirements(this IAgent agent, Condition? condition, Dictionary<string, ResourceRequirements>? resources, bool exclusive, Dictionary<string, int> assignedResources)
+		public static bool MeetsRequirements(this IAgent agent, PoolId? poolId, Condition? condition, Dictionary<string, ResourceRequirements>? resources, bool exclusive, Dictionary<string, int> assignedResources)
 		{
 			if (!agent.Enabled || agent.Status != AgentStatus.Ok)
 			{
@@ -751,6 +757,10 @@ namespace Horde.Server.Agents
 			{
 				return false;
 			}
+			if (poolId.HasValue && !agent.IsInPool(poolId.Value))
+			{
+				return false;
+			}
 			if (condition != null && !agent.SatisfiesCondition(condition))
 			{
 				return false;
diff --git a/Engine/Source/Programs/Shared/EpicGames.Horde/Compute/Requirements.cs b/Engine/Source/Programs/Shared/EpicGames.Horde/Compute/Requirements.cs
--- a/Engine/Source/Programs/Shared/EpicGames.Horde/Compute/Requirements.cs
+++ b/Engine/Source/Programs/Shared/EpicGames.Horde/Compute/Requirements.cs
@@ -12,6 +12,12 @@ namespace EpicGames.Horde.Compute
 	/// </summary>
 	public class Requirements
 	{
+		/// <summary>
+		/// Pool of machines to draw from
+		/// </summary>
+		[CbField("p")]
+		public string? Pool { get; set; }
+
 		/// <summary>
 		/// Condition string to be evaluated against the machine spec, eg. cpu-cores >= 10 &amp;&amp; ram.mb >= 200 &amp;&amp; pool == 'worker'
 		/// </summary>
@@ -60,6 +66,10 @@ namespace EpicGames.Horde.Compute
 		public override string ToString()
 		{
 			List<string> list = new List<string>();
+			if (Pool != null)
+			{
+				list.Add($"Pool:{Pool}");
+			}
 			if (Condition != null)
 			{
 				list.Add($"\"{Condition}\"");
