Horde: Fix issue with agent update spam and history modal
#rnx
#skipci

diff --git a/Engine/Source/Programs/Horde/HordeDashboard/src/backend/AgentStore.ts b/Engine/Source/Programs/Horde/HordeDashboard/src/backend/AgentStore.ts
--- a/Engine/Source/Programs/Horde/HordeDashboard/src/backend/AgentStore.ts
+++ b/Engine/Source/Programs/Horde/HordeDashboard/src/backend/AgentStore.ts
@@ -1,7 +1,7 @@
 // Copyright Epic Games, Inc. All Rights Reserved.
 
 import { action, makeObservable, observable } from 'mobx';
-import { AgentData, PoolData } from './Api';
+import { AgentData, GetAgentResponse, PoolData } from './Api';
 import backend from '.';
 
 export class AgentStore {
@@ -18,9 +18,9 @@ export class AgentStore {
         if (this.agentsUpdated) { }
         return this._agents;
     }
-            
-	private _agents: AgentData[] = [];
-	
+
+    private _agents: AgentData[] = [];
+
     @observable
     poolsUpdated: number = 0;
 
@@ -29,13 +29,13 @@ export class AgentStore {
         if (this.poolsUpdated) { }
         return this._pools;
     }
-    
+
     private _pools: PoolData[] = [];
 
-	// don't need this for now
+    // don't need this for now
     modifiedAfterDate: Date = new Date(0);
 
-    @action 
+    @action
     private _setPools(data: PoolData[]) {
         this._pools = data;
         this.poolsUpdated++;
@@ -43,11 +43,11 @@ export class AgentStore {
 
     @action
     private _setAgents(data: AgentData[]) {
-        if (data.length !== 0) {            
+        if (data.length !== 0) {
             const toSet: AgentData[] = [...this._agents];
             data.forEach(updatedAgent => {
                 const newUpdateIdx = toSet.findIndex(existingAgent => existingAgent.id === updatedAgent.id);
-                if(newUpdateIdx !== -1) {
+                if (newUpdateIdx !== -1) {
                     toSet[newUpdateIdx] = updatedAgent;
                 }
                 else {
@@ -57,19 +57,34 @@ export class AgentStore {
             this._agents = toSet;
             this.agentsUpdated++;
         }
-	}
+    }
+
+    async updateAgent(agentId: string): Promise<GetAgentResponse> {
+        return new Promise<GetAgentResponse>(async (resolve, reject) => {
+            try {
+                const agent = await backend.getAgent(agentId);
+                this._agents = this._agents.filter(a => a.id !== agentId);
+                this._agents.push(agent);
+
+                resolve(agent);
+            } catch (error) {
+                reject(error);
+            }
+        })
+    }
+
 
     async update(slim = false): Promise<void> {
         return new Promise<void>((resolve, reject) => {
             const promises: any[] = [];
-            promises.push(backend.getAgents({modifiedAfter: this.modifiedAfterDate?.toISOString()}));
-            if(!slim) {
+            promises.push(backend.getAgents({ modifiedAfter: this.modifiedAfterDate?.toISOString() }));
+            if (!slim) {
                 promises.push(backend.getPools());
             }
-			Promise.all(promises).then(values => {
+            Promise.all(promises).then(values => {
                 this.modifiedAfterDate = new Date();
                 this._setAgents(values[0]);
-                if(!slim){
+                if (!slim) {
                     this._setPools(values[1]);
                 }
                 resolve();
diff --git a/Engine/Source/Programs/Horde/HordeDashboard/src/components/HistoryModal.tsx b/Engine/Source/Programs/Horde/HordeDashboard/src/components/HistoryModal.tsx
--- a/Engine/Source/Programs/Horde/HordeDashboard/src/components/HistoryModal.tsx
+++ b/Engine/Source/Programs/Horde/HordeDashboard/src/components/HistoryModal.tsx
@@ -56,8 +56,8 @@ class HistoryModalState {
       // subscribe in any observers
       if (this.selectedAgentUpdated) { }
       return this._selectedAgent;
-   } 
-   
+   }
+
    private _selectedAgent: AgentData | undefined = undefined;
    @observable.shallow currentData: any = [];
    @observable.shallow infoItems: InfoPanelItem[] = [];
@@ -294,9 +294,13 @@ export const HistoryModal: React.FC<{ agentId: string | undefined, onDismiss: (.
 
    if (selectedAgent !== agentId) {
 
-      agentStore.update(agentStore.pools?.length ? true : false).then(() => {
-         state.setSelectedAgent(agentStore.agents.find(agent => agent.id === agentId));
-      });
+      agentStore.updateAgent(agentId).then(agent => {
+
+         state.setSelectedAgent(agent);
+
+      }).catch(error => {
+         console.error(`Unable to find agent id: ${agentId} ${error}`);
+      })
 
       setSelectedAgent(agentId);
    }
@@ -334,7 +338,7 @@ export const HistoryModal: React.FC<{ agentId: string | undefined, onDismiss: (.
          name: 'Disable',
          confirmText: "Are you sure you would like to disable this agent?",
          textInput: true,
-         update: (request, comment) => { request.enabled = false;  request.comment = comment }
+         update: (request, comment) => { request.enabled = false; request.comment = comment }
       },
       {
          name: 'Cancel Leases',
@@ -402,11 +406,11 @@ export const HistoryModal: React.FC<{ agentId: string | undefined, onDismiss: (.
       if (currentAction.update) {
          const request: UpdateAgentRequest = {};
          currentAction.update(request, actionState.comment);
-         backend.updateAgent(agentId, request).then(() => {            
+         backend.updateAgent(agentId, request).then(() => {
             agentStore.update(agentStore.pools?.length ? true : false).then(() => {
                state.setSelectedAgent(agentStore.agents.find(agent => agent.id === agentId));
                setActionState({});
-            });                  
+            });
          }).catch((reason) => {
             console.error(reason);
          });
