Horde: Add agent message enum and type definitions to C++ sample.

diff --git a/Engine/Source/Programs/Horde/Samples/RemoteWorkerCpp/AgentMessage.cpp b/Engine/Source/Programs/Horde/Samples/RemoteWorkerCpp/AgentMessage.cpp
new file mode 100644
--- /dev/null
+++ b/Engine/Source/Programs/Horde/Samples/RemoteWorkerCpp/AgentMessage.cpp
@@ -0,0 +1,327 @@
+// Copyright Epic Games, Inc. All Rights Reserved.
+
+#include "AgentMessage.h"
+#include "ComputePlatform.h"
+
+//// FAgentMessageChannel ////
+
+FAgentMessageChannel::FAgentMessageChannel(std::shared_ptr<FComputeBufferReader> InRecvBufferReader, std::shared_ptr<FComputeBufferWriter> InSendBufferWriter)
+	: RecvBufferReader(std::move(InRecvBufferReader))
+	, SendBufferWriter(std::move(InSendBufferWriter))
+	, RequestData(0)
+	, RequestSize(0)
+	, MaxRequestSize(0)
+	, ResponseType(EAgentMessageType::None)
+	, ResponseData(nullptr)
+	, ResponseLength(0)
+{
+}
+
+FAgentMessageChannel::~FAgentMessageChannel()
+{
+}
+
+void FAgentMessageChannel::Close()
+{
+	CreateMessage(EAgentMessageType::None, 0);
+	FlushMessage();
+}
+
+void FAgentMessageChannel::Ping()
+{
+	CreateMessage(EAgentMessageType::Ping, 0);
+	FlushMessage();
+}
+
+void FAgentMessageChannel::Exception(const char* Description, const char* Trace)
+{
+	CreateMessage(EAgentMessageType::Exception, strlen(Description) + strlen(Trace) + 2);
+	WriteString(Description);
+	WriteString(Trace);
+	FlushMessage();
+}
+
+void FAgentMessageChannel::Fork(int ChannelId, int BufferSize)
+{
+	CreateMessage(EAgentMessageType::Fork, sizeof(int) + sizeof(int));
+	WriteInt32(ChannelId);
+	WriteInt32(BufferSize);
+	FlushMessage();
+}
+
+void FAgentMessageChannel::Attach()
+{
+	CreateMessage(EAgentMessageType::Attach, 0);
+	FlushMessage();
+}
+
+void FAgentMessageChannel::UploadFiles(const char* Path, const char* Locator)
+{
+	CreateMessage(EAgentMessageType::WriteFiles, strlen(Path) + strlen(Locator) + 20);
+	WriteString(Path);
+	WriteString(Locator);
+	FlushMessage();
+}
+
+void FAgentMessageChannel::DeleteFiles(const char** Paths, size_t Count)
+{
+	size_t RequiredSize = 10;
+	for (size_t Idx = 0; Idx < Count; Idx++)
+	{
+		RequiredSize += strlen(Paths[Idx]) + 10;
+	}
+
+	CreateMessage(EAgentMessageType::DeleteFiles, RequiredSize);
+	WriteUnsignedVarInt(Count);
+	for (size_t Idx = 0; Idx < Count; Idx++)
+	{
+		WriteString(Paths[Idx]);
+	}
+	FlushMessage();
+}
+
+void FAgentMessageChannel::Execute(const char* Exe, const char** Args, size_t NumArgs, const char* WorkingDir, const char** EnvVars, size_t NumEnvVars, EExecuteProcessFlags Flags)
+{
+	size_t RequiredSize = 50 + strlen(Exe);
+	for (size_t Idx = 0; Idx < NumArgs; Idx++)
+	{
+		RequiredSize += strlen(Args[Idx]) + 10;
+	}
+	if (WorkingDir != nullptr)
+	{
+		RequiredSize += strlen(WorkingDir) + 10;
+	}
+	for (size_t Idx = 0; Idx < NumEnvVars; Idx++)
+	{
+		RequiredSize += strlen(EnvVars[Idx]) + 20;
+	}
+
+	CreateMessage(EAgentMessageType::ExecuteV2, RequiredSize);
+	WriteString(Exe);
+
+	WriteUnsignedVarInt(NumArgs);
+	for (size_t Idx = 0; Idx < NumArgs; Idx++)
+	{
+		WriteString(Args[Idx]);
+	}
+
+	WriteOptionalString(WorkingDir);
+
+	WriteUnsignedVarInt(NumEnvVars);
+	for (size_t Idx = 0; Idx < NumEnvVars; Idx++)
+	{
+		const char* EqualsPtr = strchr(EnvVars[Idx], '=');
+		WriteString(std::string_view(EnvVars[Idx], EqualsPtr));
+		if (*(EqualsPtr + 1) == 0)
+		{
+			WriteOptionalString(nullptr);
+		}
+		else
+		{
+			WriteOptionalString(EqualsPtr + 1);
+		}
+	}
+
+	WriteInt32((int)Flags);
+	FlushMessage();
+}
+
+void FAgentMessageChannel::Blob(const unsigned char* Data, size_t Length)
+{
+	const size_t MaxChunkSize = 512 * 1024;
+	for (size_t ChunkOffset = 0; ChunkOffset < Length;)
+	{
+		size_t ChunkLength = std::min(Length - ChunkOffset, MaxChunkSize);
+		
+		CreateMessage(EAgentMessageType::ReadBlobResponse, ChunkLength + 128);
+		WriteInt32((int)ChunkOffset);
+		WriteInt32((int)Length);
+		WriteFixedLengthBytes(Data + ChunkOffset, ChunkLength);
+		FlushMessage();
+
+		ChunkOffset += ChunkLength;
+	}
+}
+
+void FAgentMessageChannel::Xor(const unsigned char* Data, size_t Length, unsigned char Value)
+{
+	CreateMessage(EAgentMessageType::XorRequest, Length + 1);
+	WriteFixedLengthBytes(Data, Length);
+	WriteFixedLengthBytes(&Value, 1);
+	FlushMessage();
+}
+
+EAgentMessageType FAgentMessageChannel::ReadResponse()
+{
+	if (ResponseLength > 0)
+	{
+		RecvBufferReader->AdvanceReadPosition(ResponseLength);
+		ResponseData = nullptr;
+		ResponseLength = 0;
+	}
+
+	const unsigned char* Header = RecvBufferReader->WaitToRead(MessageHeaderLength);
+	unsigned int Length = *((unsigned int*)(Header + 1));
+
+	Header = RecvBufferReader->WaitToRead(MessageHeaderLength + Length);
+
+	ResponseType = (EAgentMessageType)Header[0];
+	ResponseData = Header + MessageHeaderLength;
+	ResponseLength = Length;
+
+	return ResponseType;
+}
+
+void FAgentMessageChannel::ReadException(AgentMessage::FException& Ex)
+{
+	UE_COMPUTE_ASSERT(ResponseType == EAgentMessageType::Exception);
+
+	const unsigned char* Pos = ResponseData;
+	Ex.Message = ReadString(&Pos);
+	Ex.Description = ReadString(&Pos);
+	UE_COMPUTE_ASSERT(Pos == ResponseData + ResponseLength);
+}
+
+int FAgentMessageChannel::ReadExecuteResult()
+{
+	UE_COMPUTE_ASSERT(ResponseType == EAgentMessageType::ExecuteResult);
+
+	const unsigned char* Pos = ResponseData;
+	int Result = ReadInt32(&Pos);
+	UE_COMPUTE_ASSERT(Pos == ResponseData + ResponseLength);
+
+	return Result;
+}
+
+void FAgentMessageChannel::ReadBlobRequest(AgentMessage::FBlobRequest& Ex)
+{
+	UE_COMPUTE_ASSERT(ResponseType == EAgentMessageType::ReadBlob);
+
+	const unsigned char* Pos = ResponseData;
+	Ex.Locator = ReadString(&Pos);
+	Ex.Offset = (int)ReadUnsignedVarInt(&Pos);
+	Ex.Length = (int)ReadUnsignedVarInt(&Pos);
+	UE_COMPUTE_ASSERT(Pos == ResponseData + ResponseLength);
+}
+
+void FAgentMessageChannel::CreateMessage(EAgentMessageType Type, size_t MaxLength)
+{
+	RequestData = SendBufferWriter->WaitToWrite(MessageHeaderLength + MaxLength);
+	RequestData[0] = (unsigned char)Type;
+}
+
+void FAgentMessageChannel::FlushMessage()
+{
+	memcpy(&RequestData[1], &RequestSize, sizeof(int));
+	SendBufferWriter->AdvanceWritePosition(MessageHeaderLength + RequestSize);
+	RequestSize = 0;
+	MaxRequestSize = 0;
+	RequestData = nullptr;
+}
+
+void FAgentMessageChannel::WriteInt32(int Value)
+{
+	WriteFixedLengthBytes((const unsigned char*)&Value, sizeof(int));
+}
+
+int FAgentMessageChannel::ReadInt32(const unsigned char** Pos)
+{
+	int Value;
+	memcpy(&Value, *Pos, sizeof(int));
+	*Pos += sizeof(int);
+	return Value;
+}
+
+void FAgentMessageChannel::WriteFixedLengthBytes(const unsigned char* Data, size_t Length)
+{
+	UE_COMPUTE_ASSERT(RequestSize + Length <= MaxRequestSize);
+	memcpy(&RequestData[MessageHeaderLength + RequestSize], Data, Length);
+	RequestSize += Length;
+}
+
+const unsigned char* FAgentMessageChannel::ReadFixedLengthBytes(const unsigned char** Pos, size_t Length)
+{
+	const unsigned char* Data = *Pos;
+	*Pos += Length;
+	return Data;
+}
+
+size_t FAgentMessageChannel::MeasureUnsignedVarInt(size_t Value)
+{
+	UE_COMPUTE_ASSERT(Value == (unsigned int)Value);
+	return (FComputePlatform::FloorLog2((unsigned int)Value) / 7) + 1;
+}
+
+void FAgentMessageChannel::WriteUnsignedVarInt(size_t Value)
+{
+	size_t ByteCount = MeasureUnsignedVarInt(Value);
+	UE_COMPUTE_ASSERT(RequestSize + ByteCount <= MaxRequestSize);
+
+	unsigned char* Output = RequestData + MessageHeaderLength + RequestSize;
+	for (size_t Idx = 1; Idx < ByteCount; Idx++)
+	{
+		Output[ByteCount - Idx] = (unsigned)Value;
+		Value >>= 8;
+	}
+	Output[0] = (unsigned char)((0xff << (9 - (int)ByteCount)) | (unsigned char)Value);
+
+	RequestSize += ByteCount;
+}
+
+size_t FAgentMessageChannel::ReadUnsignedVarInt(const unsigned char** Pos)
+{
+	const unsigned char* Data = *Pos;
+
+	unsigned char FirstByte = Data[0];
+	size_t NumBytes = FComputePlatform::CountLeadingZeros(~(unsigned int)FirstByte) + 1 - 24; // Note byte -> int conversion here, hence ignoring subtracting 24 bits
+
+	size_t value = (size_t)(FirstByte & (0xff >> NumBytes));
+	for (size_t Idx = 1; Idx < NumBytes; Idx++)
+	{
+		value <<= 8;
+		value |= Data[Idx];
+	}
+
+	*Pos += NumBytes;
+	return value;
+}
+
+size_t FAgentMessageChannel::MeasureString(const char* Text) const
+{
+	size_t Length = strlen(Text);
+	return MeasureUnsignedVarInt(Length) + Length;
+}
+
+void FAgentMessageChannel::WriteString(const char* Text)
+{
+	size_t Length = strlen(Text);
+	WriteUnsignedVarInt(Length);
+	WriteFixedLengthBytes((const unsigned char*)Text, Length);
+}
+
+void FAgentMessageChannel::WriteString(const std::string_view& Text)
+{
+	WriteUnsignedVarInt(Text.size());
+	WriteFixedLengthBytes((const unsigned char*)Text.data(), Text.size());
+}
+
+std::string_view FAgentMessageChannel::ReadString(const unsigned char** Pos)
+{
+	size_t Length = ReadUnsignedVarInt(Pos);
+	const char* Start = (const char*)ReadFixedLengthBytes(Pos, Length);
+	return std::string_view(Start, Length);
+}
+
+void FAgentMessageChannel::WriteOptionalString(const char* Text)
+{
+	if (Text == nullptr)
+	{
+		WriteUnsignedVarInt(0);
+	}
+	else
+	{
+		size_t Length = strlen(Text);
+		WriteUnsignedVarInt(Length + 1);
+		WriteFixedLengthBytes((const unsigned char*)Text, Length);
+	}
+}
diff --git a/Engine/Source/Programs/Horde/Samples/RemoteWorkerCpp/AgentMessage.h b/Engine/Source/Programs/Horde/Samples/RemoteWorkerCpp/AgentMessage.h
new file mode 100644
--- /dev/null
+++ b/Engine/Source/Programs/Horde/Samples/RemoteWorkerCpp/AgentMessage.h
@@ -0,0 +1,183 @@
+// Copyright Epic Games, Inc. All Rights Reserved.
+
+#pragma once
+
+#include <memory>
+#include <vector>
+#include <map>
+#include <string_view>
+#include "ComputeBuffer.h"
+
+// Type of a compute message
+enum class EAgentMessageType : unsigned char
+{
+	// No message was received (end of stream)
+	None = 0x00,
+
+	// No-op message sent to keep the connection alive. Remote should reply with the same message.
+	Ping = 0x01,
+
+	// Sent in place of a regular response if an error occurs on the remote
+	Exception = 0x02,
+
+	// Fork the message loop into a new channel
+	Fork = 0x03,
+
+	// Sent as the first message on a channel to notify the remote that the remote end is attached
+	Attach = 0x04,
+
+	// Extract files on the remote machine (Initiator -> Remote)
+	WriteFiles = 0x10,
+
+	// Notification that files have been extracted (Remote -> Initiator)
+	WriteFilesResponse = 0x11,
+
+	// Deletes files on the remote machine (Initiator -> Remote)
+	DeleteFiles = 0x12,
+
+	// Execute a process in a sandbox (Initiator -> Remote)
+	ExecuteV1 = 0x16,
+
+	// Execute a process in a sandbox (Initiator -> Remote)
+	ExecuteV2 = 0x22,
+
+	// Returns output from the child process to the caller (Remote -> Initiator)
+	ExecuteOutput = 0x17,
+
+	// Returns the process exit code (Remote -> Initiator)
+	ExecuteResult = 0x18,
+
+	// Reads a blob from storage
+	ReadBlob = 0x20,
+
+	// Response to a ReadBlob request.
+	ReadBlobResponse = 0x21,
+
+	// Xor a block of data with a value
+	XorRequest = 0xf0,
+
+	// Result from an XorRequest request.
+	XorResponse = 0xf1,
+};
+
+// Flags describing how to execute a compute task process on the agent
+enum class EExecuteProcessFlags : unsigned char
+{
+	// No execute flags set
+	None = 0,
+
+	// Request execution to be wrapped under Wine when running on Linux. Agent still reserves the right to refuse it (e.g no Wine executable configured, mismatching OS etc)
+	UseWine = 1,
+};
+
+// Definitions for various message types
+namespace AgentMessage
+{
+	struct FException
+	{
+		std::string_view Message;
+		std::string_view Description;
+	};
+
+	struct FBlobRequest
+	{
+		std::string_view Locator;
+		int Offset;
+		int Length;
+	};
+}
+
+// Channel for receiving agent messages
+class FAgentMessageChannel
+{
+public:
+	FAgentMessageChannel(std::shared_ptr<FComputeBufferReader> InRecvBufferReader, std::shared_ptr<FComputeBufferWriter> InSendBufferWriter);
+	~FAgentMessageChannel();
+
+	//// Requests ////
+
+	// Closes the remote message loop
+	void Close();
+
+	// Sends a ping message to the remote
+	void Ping();
+
+	// Sends an exception response to the remote
+	void Exception(const char* Description, const char* Trace);
+
+	// Requests that the remote message loop be forked
+	void Fork(int ChannelId, int BufferSize);
+
+	// Notifies the remote that a buffer has been attached
+	void Attach();
+
+	// Extracts a bundle containing files to a particular path
+	void UploadFiles(const char* Path, const char* Locator);
+
+	// Deletes files matching a set of wildcards
+	void DeleteFiles(const char** Paths, size_t Count);
+
+	// Executes a process on the remote machine
+	void Execute(const char* Exe, const char** Args, size_t NumArgs, const char* WorkingDir, const char** EnvVars, size_t NumEnvVars, EExecuteProcessFlags Flags);
+
+	// Writes a blob requested by the remote
+	void Blob(const unsigned char* Data, size_t Length);
+
+	// Send a message to request that a byte string be xor'ed with a particular value
+	void Xor(const unsigned char* Data, size_t Length, unsigned char Value);
+
+	//// Responses ////
+
+	// Reads a response from the remote. Other Read methods can be used to access response data.
+	EAgentMessageType ReadResponse();
+
+	// Gets the raw data from a response
+	const void* GetResponseData() const { return ResponseData + MessageHeaderLength; }
+
+	// Gets the size of the response
+	size_t GetResponseSize() const { return ResponseLength - MessageHeaderLength; }
+
+	// Reads an exception response
+	void ReadException(AgentMessage::FException& Ex);
+
+	// Reads the result from executing a process
+	int ReadExecuteResult();
+
+	// Reads a blob request message
+	void ReadBlobRequest(AgentMessage::FBlobRequest& Ex);
+
+private:
+	const size_t MessageHeaderLength = 5;
+
+	std::shared_ptr<FComputeBufferReader> RecvBufferReader;
+	std::shared_ptr<FComputeBufferWriter> SendBufferWriter;
+
+	unsigned char* RequestData;
+	size_t RequestSize;
+	size_t MaxRequestSize;
+
+	EAgentMessageType ResponseType;
+	const unsigned char* ResponseData;
+	size_t ResponseLength;
+
+	void CreateMessage(EAgentMessageType Type, size_t MaxLength);
+	void FlushMessage();
+
+	void WriteInt32(int Value);
+	static int ReadInt32(const unsigned char** Pos);
+
+	void WriteFixedLengthBytes(const unsigned char* Data, size_t Length);
+	static const unsigned char* ReadFixedLengthBytes(const unsigned char** Pos, size_t Length);
+
+	static size_t MeasureUnsignedVarInt(size_t Value);
+	void WriteUnsignedVarInt(size_t Value);
+	static size_t ReadUnsignedVarInt(const unsigned char** Pos);
+
+	size_t MeasureString(const char* Text) const;
+	void WriteString(const char* Text);
+	void WriteString(const std::string_view& Text);
+	static std::string_view ReadString(const unsigned char** Pos);
+
+	void WriteOptionalString(const char* Text);
+};
+
