Horde: Rename "workspace" commands to "perforce" commands.

diff --git a/Engine/Source/Programs/Horde/Horde/Commands/Perforce/PerforceBase.cs b/Engine/Source/Programs/Horde/Horde/Commands/Perforce/PerforceBase.cs
new file mode 100644
--- /dev/null
+++ b/Engine/Source/Programs/Horde/Horde/Commands/Perforce/PerforceBase.cs
@@ -0,0 +1,148 @@
+// Copyright Epic Games, Inc. All Rights Reserved.
+
+using System.ComponentModel;
+using EpicGames.Core;
+using EpicGames.Perforce;
+using EpicGames.Perforce.Managed;
+using Microsoft.Extensions.Logging;
+
+namespace Horde.Commands.Perforce
+{
+	abstract class PerforceBase : Command
+	{
+		[CommandLine("-ServerAndPort")]
+		[Description("Specifies the Perforce server and port")]
+		protected string ServerAndPort { get; set; } = PerforceSettings.Default.ServerAndPort;
+
+		[CommandLine("-User")]
+		[Description("Specifies the Perforce username")]
+		protected string UserName { get; set; } = PerforceSettings.Default.UserName;
+
+		[CommandLine("-BaseDir", Required = true)]
+		[Description("Base directory to use for syncing workspaces")]
+		protected DirectoryReference BaseDir { get; set; } = null!;
+
+		[CommandLine("-Overwrite")]
+		[Description("")]
+		protected bool Overwrite { get; set; } = false;
+		
+		[CommandLine("-PreferNativeClient")]
+		[Description("Prefer to use native Perforce client (instead of launching separate p4 process)")]
+		protected bool PreferNativeClient { get; set; } = false;
+		
+		[CommandLine("-UseHaveTable")]
+		[Description("Use have-table for syncing")]
+		string UseHaveTable { get; set; } = "true";
+
+		public override void Configure(CommandLineArguments arguments, ILogger logger)
+		{
+			base.Configure(arguments, logger);
+
+			if(BaseDir == null)
+			{
+				for (DirectoryReference? parentDir = DirectoryReference.GetCurrentDirectory(); parentDir != null; parentDir = parentDir.ParentDirectory)
+				{
+					if (ManagedWorkspace.Exists(parentDir))
+					{
+						BaseDir = parentDir;
+						break;
+					}
+				}
+
+				if (BaseDir == null)
+				{
+					throw new FatalErrorException("Unable to find existing repository in current working tree. Specify -BaseDir=... explicitly.");
+				}
+			}
+		}
+
+		public override async Task<int> ExecuteAsync(ILogger logger)
+		{
+			PerforceSettings settings = new(ServerAndPort, UserName) { PreferNativeClient = PreferNativeClient };
+			using IPerforceConnection perforce = await PerforceConnection.CreateAsync(settings, logger);
+			InfoRecord info = await perforce.GetInfoAsync(InfoOptions.ShortOutput, CancellationToken.None);
+			bool useHaveTable = UseHaveTable.Equals("true", StringComparison.Ordinal);
+
+			ManagedWorkspace repo = await ManagedWorkspace.LoadOrCreateAsync(info.ClientHost!, BaseDir, Overwrite, useHaveTable, logger, CancellationToken.None);
+			await ExecuteAsync(perforce, repo, logger);
+			return 0;
+		}
+
+		protected abstract Task ExecuteAsync(IPerforceConnection perforce, ManagedWorkspace repo, ILogger logger);
+
+		protected static long ParseSize(string size)
+		{
+			long value;
+			if (size.EndsWith("gb", StringComparison.OrdinalIgnoreCase))
+			{
+				string sizeValue = size.Substring(0, size.Length - 2).TrimEnd();
+				if (Int64.TryParse(sizeValue, out value))
+				{
+					return value * (1024 * 1024 * 1024);
+				}
+			}
+			else if (size.EndsWith("mb", StringComparison.OrdinalIgnoreCase))
+			{
+				string sizeValue = size.Substring(0, size.Length - 2).TrimEnd();
+				if (Int64.TryParse(sizeValue, out value))
+				{
+					return value * (1024 * 1024);
+				}
+			}
+			else if (size.EndsWith("kb", StringComparison.OrdinalIgnoreCase))
+			{
+				string sizeValue = size.Substring(0, size.Length - 2).TrimEnd();
+				if (Int64.TryParse(sizeValue, out value))
+				{
+					return value * 1024;
+				}
+			}
+			else
+			{
+				if (Int64.TryParse(size, out value))
+				{
+					return value;
+				}
+			}
+			throw new FatalErrorException("Invalid size '{0}'", size);
+		}
+
+		protected static int ParseChangeNumber(string change)
+		{
+			int changeNumber;
+			if (Int32.TryParse(change, out changeNumber) && changeNumber > 0)
+			{
+				return changeNumber;
+			}
+			throw new FatalErrorException("Unable to parse change number from '{0}'", change);
+		}
+
+		protected static int ParseChangeNumberOrLatest(string change)
+		{
+			if (change.Equals("Latest", StringComparison.OrdinalIgnoreCase))
+			{
+				return -1;
+			}
+			else
+			{
+				return ParseChangeNumber(change);
+			}
+		}
+
+		protected static List<string> ExpandFilters(List<string> arguments)
+		{
+			List<string> filters = new List<string>();
+			foreach (string argument in arguments)
+			{
+				foreach (string singleArgument in argument.Split(';').Select(x => x.Trim()))
+				{
+					if (singleArgument.Length > 0)
+					{
+						filters.Add(singleArgument);
+					}
+				}
+			}
+			return filters;
+		}
+	}
+}
diff --git a/Engine/Source/Programs/Horde/Horde/Commands/Perforce/PerforceClean.cs b/Engine/Source/Programs/Horde/Horde/Commands/Perforce/PerforceClean.cs
new file mode 100644
--- /dev/null
+++ b/Engine/Source/Programs/Horde/Horde/Commands/Perforce/PerforceClean.cs
@@ -0,0 +1,23 @@
+// Copyright Epic Games, Inc. All Rights Reserved.
+
+using System.ComponentModel;
+using EpicGames.Core;
+using EpicGames.Perforce;
+using EpicGames.Perforce.Managed;
+using Microsoft.Extensions.Logging;
+
+namespace Horde.Commands.Perforce
+{
+	[Command("perforce", "clean", "Cleans all modified files from the workspace.")]
+	class PerforceClean : PerforceBase
+	{
+		[CommandLine("-Incremental")]
+		[Description("Performs an incremental sync, without removing intermediates")]
+		public bool Incremental { get; set; } = false;
+
+		protected override Task ExecuteAsync(IPerforceConnection perforce, ManagedWorkspace repo, ILogger logger)
+		{
+			return repo.CleanAsync(!Incremental, CancellationToken.None);
+		}
+	}
+}
diff --git a/Engine/Source/Programs/Horde/Horde/Commands/Perforce/PerforceClear.cs b/Engine/Source/Programs/Horde/Horde/Commands/Perforce/PerforceClear.cs
new file mode 100644
--- /dev/null
+++ b/Engine/Source/Programs/Horde/Horde/Commands/Perforce/PerforceClear.cs
@@ -0,0 +1,18 @@
+// Copyright Epic Games, Inc. All Rights Reserved.
+
+using EpicGames.Core;
+using EpicGames.Perforce;
+using EpicGames.Perforce.Managed;
+using Microsoft.Extensions.Logging;
+
+namespace Horde.Commands.Perforce
+{
+	[Command("perforce", "clear", "Empties the staging directory of any files, returning them to the cache")]
+	class PerforceClear : PerforceBase
+	{
+		protected override Task ExecuteAsync(IPerforceConnection perforce, ManagedWorkspace repo, ILogger logger)
+		{
+			return repo.ClearAsync(CancellationToken.None);
+		}
+	}
+}
diff --git a/Engine/Source/Programs/Horde/Horde/Commands/Perforce/PerforceDump.cs b/Engine/Source/Programs/Horde/Horde/Commands/Perforce/PerforceDump.cs
new file mode 100644
--- /dev/null
+++ b/Engine/Source/Programs/Horde/Horde/Commands/Perforce/PerforceDump.cs
@@ -0,0 +1,19 @@
+// Copyright Epic Games, Inc. All Rights Reserved.
+
+using EpicGames.Core;
+using EpicGames.Perforce;
+using EpicGames.Perforce.Managed;
+using Microsoft.Extensions.Logging;
+
+namespace Horde.Commands.Perforce
+{
+	[Command("perforce", "dump", "Dumps the contents of the repository to the log for analysis")]
+	class PerforceDump : PerforceBase
+	{
+		protected override Task ExecuteAsync(IPerforceConnection perforce, ManagedWorkspace repo, ILogger logger)
+		{
+			repo.Dump();
+			return Task.CompletedTask;
+		}
+	}
+}
diff --git a/Engine/Source/Programs/Horde/Horde/Commands/Perforce/PerforcePopulate.cs b/Engine/Source/Programs/Horde/Horde/Commands/Perforce/PerforcePopulate.cs
new file mode 100644
--- /dev/null
+++ b/Engine/Source/Programs/Horde/Horde/Commands/Perforce/PerforcePopulate.cs
@@ -0,0 +1,46 @@
+// Copyright Epic Games, Inc. All Rights Reserved.
+
+using System.ComponentModel;
+using EpicGames.Core;
+using EpicGames.Perforce;
+using EpicGames.Perforce.Managed;
+using Microsoft.Extensions.Logging;
+
+namespace Horde.Commands.Perforce
+{
+	[Command("perforce", "populatecache", "Populates the cache with the head revision of the given streams")]
+	class PerforcePopulate : PerforceBase
+	{
+		[CommandLine("-ClientAndStream=")]
+		[Description("Specifies client and stream pairs, in the format Client:Stream")]
+		List<string> ClientAndStreamParams { get; set; } = new List<string>();
+
+		[CommandLine("-Filter=")]
+		[Description("Filters for the files to sync, in P4 syntax (eg. /Engine/...)")]
+		List<string> Filters { get; set; } = new List<string>();
+
+		[CommandLine("-FakeSync")]
+		[Description("Simulates the sync without actually fetching any files")]
+		bool FakeSync { get; set; } = false;
+
+		protected override async Task ExecuteAsync(IPerforceConnection perforce, ManagedWorkspace repo, ILogger logger)
+		{
+			List<string> expandedFilters = ExpandFilters(Filters);
+
+			List<PopulateRequest> populateRequests = new List<PopulateRequest>();
+			foreach (string clientAndStreamParam in ClientAndStreamParams)
+			{
+				int idx = clientAndStreamParam.IndexOf(':', StringComparison.Ordinal);
+				if (idx == -1)
+				{
+					throw new FatalErrorException("Expected -ClientAndStream=<ClientName>:<StreamName>");
+				}
+
+				using IPerforceConnection perforceClient = await perforce.WithClientAsync(clientAndStreamParam.Substring(0, idx));
+				populateRequests.Add(new PopulateRequest(perforceClient, clientAndStreamParam.Substring(idx + 1), expandedFilters));
+			}
+
+			await repo.PopulateAsync(populateRequests, FakeSync, CancellationToken.None);
+		}
+	}
+}
diff --git a/Engine/Source/Programs/Horde/Horde/Commands/Perforce/PerforcePurge.cs b/Engine/Source/Programs/Horde/Horde/Commands/Perforce/PerforcePurge.cs
new file mode 100644
--- /dev/null
+++ b/Engine/Source/Programs/Horde/Horde/Commands/Perforce/PerforcePurge.cs
@@ -0,0 +1,27 @@
+// Copyright Epic Games, Inc. All Rights Reserved.
+
+using EpicGames.Core;
+using EpicGames.Perforce;
+using EpicGames.Perforce.Managed;
+using Microsoft.Extensions.Logging;
+
+namespace Horde.Commands.Perforce
+{
+	[Command("perforce", "purgecache", "Shrink the size of the cache to the given size")]
+	class PerforcePurge : PerforceBase
+	{
+		[CommandLine("-Size=")]
+		string? SizeParam { get; set; } = null;
+
+		protected override Task ExecuteAsync(IPerforceConnection perforce, ManagedWorkspace repo, ILogger logger)
+		{
+			long size = 0;
+			if (SizeParam != null)
+			{
+				size = ParseSize(SizeParam);
+			}
+
+			return repo.PurgeAsync(size, CancellationToken.None);
+		}
+	}
+}
diff --git a/Engine/Source/Programs/Horde/Horde/Commands/Perforce/PerforceRepair.cs b/Engine/Source/Programs/Horde/Horde/Commands/Perforce/PerforceRepair.cs
new file mode 100644
--- /dev/null
+++ b/Engine/Source/Programs/Horde/Horde/Commands/Perforce/PerforceRepair.cs
@@ -0,0 +1,18 @@
+// Copyright Epic Games, Inc. All Rights Reserved.
+
+using EpicGames.Core;
+using EpicGames.Perforce;
+using EpicGames.Perforce.Managed;
+using Microsoft.Extensions.Logging;
+
+namespace Horde.Commands.Perforce
+{
+	[Command("perforce", "repaircache", "Checks the integrity of the cache, and removes any invalid files")]
+	class PerforceRepair : PerforceBase
+	{
+		protected override Task ExecuteAsync(IPerforceConnection perforce, ManagedWorkspace repo, ILogger logger)
+		{
+			return repo.RepairAsync(CancellationToken.None);
+		}
+	}
+}
diff --git a/Engine/Source/Programs/Horde/Horde/Commands/Perforce/PerforceRevert.cs b/Engine/Source/Programs/Horde/Horde/Commands/Perforce/PerforceRevert.cs
new file mode 100644
--- /dev/null
+++ b/Engine/Source/Programs/Horde/Horde/Commands/Perforce/PerforceRevert.cs
@@ -0,0 +1,24 @@
+// Copyright Epic Games, Inc. All Rights Reserved.
+
+using System.ComponentModel;
+using EpicGames.Core;
+using EpicGames.Perforce;
+using EpicGames.Perforce.Managed;
+using Microsoft.Extensions.Logging;
+
+namespace Horde.Commands.Perforce
+{
+	[Command("perforce", "revert", "Revert all files that are open in the current workspace. Does not replace them with valid revisions.")]
+	class PerforceRevert : PerforceBase
+	{
+		[CommandLine("-Client=", Required = true)]
+		[Description("Client to revert all files for")]
+		string ClientName { get; set; } = null!;
+
+		protected override async Task ExecuteAsync(IPerforceConnection perforce, ManagedWorkspace repo, ILogger logger)
+		{
+			using IPerforceConnection perforceClient = await perforce.WithClientAsync(ClientName);
+			await repo.RevertAsync(perforceClient, CancellationToken.None);
+		}
+	}
+}
diff --git a/Engine/Source/Programs/Horde/Horde/Commands/Perforce/PerforceSetup.cs b/Engine/Source/Programs/Horde/Horde/Commands/Perforce/PerforceSetup.cs
new file mode 100644
--- /dev/null
+++ b/Engine/Source/Programs/Horde/Horde/Commands/Perforce/PerforceSetup.cs
@@ -0,0 +1,28 @@
+// Copyright Epic Games, Inc. All Rights Reserved.
+
+using System.ComponentModel;
+using EpicGames.Core;
+using EpicGames.Perforce;
+using EpicGames.Perforce.Managed;
+using Microsoft.Extensions.Logging;
+
+namespace Horde.Commands.Perforce
+{
+	[Command("perforce", "setup", "Creates or updates a client to use a given stream")]
+	class PerforceSetup : PerforceBase
+	{
+		[CommandLine("-Client=", Required = true)]
+		[Description("Name of the client to create")]
+		string ClientName { get; set; } = null!;
+
+		[CommandLine("-Stream=", Required = true)]
+		[Description("Name of the stream to configure")]
+		string StreamName { get; set; } = null!;
+
+		protected override async Task ExecuteAsync(IPerforceConnection perforce, ManagedWorkspace repo, ILogger logger)
+		{
+			using IPerforceConnection perforceClient = await perforce.WithClientAsync(ClientName);
+			await repo.SetupAsync(perforceClient, StreamName, CancellationToken.None);
+		}
+	}
+}
diff --git a/Engine/Source/Programs/Horde/Horde/Commands/Perforce/PerforceSizes.cs b/Engine/Source/Programs/Horde/Horde/Commands/Perforce/PerforceSizes.cs
new file mode 100644
--- /dev/null
+++ b/Engine/Source/Programs/Horde/Horde/Commands/Perforce/PerforceSizes.cs
@@ -0,0 +1,32 @@
+// Copyright Epic Games, Inc. All Rights Reserved.
+
+using System.ComponentModel;
+using EpicGames.Core;
+using EpicGames.Perforce;
+using EpicGames.Perforce.Managed;
+using Microsoft.Extensions.Logging;
+
+namespace Horde.Commands.Perforce
+{
+	[Command("perforce", "sizes", "Gathers stats for which streams take the most amount of space in the cache for the given configuration")]
+	class PerforceSizes : PerforceBase
+	{
+		[CommandLine("-TempClient=", Required = true)]
+		[Description("Name of a temporary client to switch between streams gathering metadata. Will be created if it does not exist.")]
+		string TempClientName { get; set; } = null!;
+
+		[CommandLine("-Stream=")]
+		[Description("Streams that should be included in the output")]
+		List<string> StreamNames { get; set; } = new List<string>();
+
+		[CommandLine("-Filter=")]
+		[Description("Filters for the files to sync, in P4 syntax (eg. /Engine/...)")]
+		List<string> Filters { get; set; } = new List<string>();
+
+		protected override async Task ExecuteAsync(IPerforceConnection perforce, ManagedWorkspace repo, ILogger logger)
+		{
+			using IPerforceConnection perforceClient = await perforce.WithClientAsync(TempClientName);
+			await repo.StatsAsync(perforceClient, StreamNames, Filters, CancellationToken.None);
+		}
+	}
+}
diff --git a/Engine/Source/Programs/Horde/Horde/Commands/Perforce/PerforceStatus.cs b/Engine/Source/Programs/Horde/Horde/Commands/Perforce/PerforceStatus.cs
new file mode 100644
--- /dev/null
+++ b/Engine/Source/Programs/Horde/Horde/Commands/Perforce/PerforceStatus.cs
@@ -0,0 +1,18 @@
+// Copyright Epic Games, Inc. All Rights Reserved.
+
+using EpicGames.Core;
+using EpicGames.Perforce;
+using EpicGames.Perforce.Managed;
+using Microsoft.Extensions.Logging;
+
+namespace Horde.Commands.Perforce
+{
+	[Command("perforce", "status", "Prints information about the state of the cache and workspace")]
+	class PerforceStatus : PerforceBase
+	{
+		protected override async Task ExecuteAsync(IPerforceConnection perforce, ManagedWorkspace repo, ILogger logger)
+		{
+			await repo.Status();
+		}
+	}
+}
diff --git a/Engine/Source/Programs/Horde/Horde/Commands/Perforce/PerforceSync.cs b/Engine/Source/Programs/Horde/Horde/Commands/Perforce/PerforceSync.cs
new file mode 100644
--- /dev/null
+++ b/Engine/Source/Programs/Horde/Horde/Commands/Perforce/PerforceSync.cs
@@ -0,0 +1,60 @@
+// Copyright Epic Games, Inc. All Rights Reserved.
+
+using System.ComponentModel;
+using EpicGames.Core;
+using EpicGames.Perforce;
+using EpicGames.Perforce.Managed;
+using Microsoft.Extensions.Logging;
+
+namespace Horde.Commands.Perforce
+{
+	[Command("perforce", "sync", "Syncs the files for a particular stream and changelist")]
+	class PerforceSync : PerforceBase
+	{
+		[CommandLine("-Client=", Required = true)]
+		[Description("Name of the client to sync. Will be created if it does not exist.")]
+		public string ClientName { get; set; } = null!;
+
+		[CommandLine("-Stream=", Required = true)]
+		[Description("The stream to sync")]
+		public string StreamName { get; set; } = null!;
+
+		[CommandLine("-Change=")]
+		[Description("The change to sync. May be a changelist number, or 'Latest'")]
+		public string Change { get; set; } = "Latest";
+
+		[CommandLine("-Preflight=")]
+		[Description("The change to unshelve into the workspace")]
+		public int PreflightChange { get; set; } = -1;
+
+		[CommandLine]
+		[Description("Optional path to a cache file used to store workspace metadata. Using a location on a network share allows multiple machines syncing the same CL to only query Perforce state once.")]
+		FileReference? CacheFile { get; set; } = null;
+
+		[CommandLine("-Filter=")]
+		[Description("Filters for the files to sync, in P4 syntax (eg. /Engine/...)")]
+		public List<string> Filters { get; set; } = new List<string>();
+
+		[CommandLine("-Incremental")]
+		[Description("Performs an incremental sync, without removing intermediates")]
+		public bool IncrementalSync { get; set; } = false;
+
+		[CommandLine("-FakeSync")]
+		[Description("Simulates the sync without actually fetching any files")]
+		public bool FakeSync { get; set; } = false;
+
+		protected override async Task ExecuteAsync(IPerforceConnection perforce, ManagedWorkspace repo, ILogger logger)
+		{
+			int changeNumber = ParseChangeNumberOrLatest(Change);
+			List<string> expandedFilters = ExpandFilters(Filters);
+
+			using IPerforceConnection perforceClient = await perforce.WithClientAsync(ClientName);
+			await repo.SyncAsync(perforceClient, StreamName, changeNumber, expandedFilters, !IncrementalSync, FakeSync, CacheFile, CancellationToken.None);
+
+			if (PreflightChange != -1)
+			{
+				await repo.UnshelveAsync(perforceClient, PreflightChange, CancellationToken.None);
+			}
+		}
+	}
+}
diff --git a/Engine/Source/Programs/Horde/Horde/Commands/Workspace/WorkspaceBase.cs b/Engine/Source/Programs/Horde/Horde/Commands/Workspace/WorkspaceBase.cs
deleted file mode 100644
diff --git a/Engine/Source/Programs/Horde/Horde/Commands/Workspace/WorkspaceClean.cs b/Engine/Source/Programs/Horde/Horde/Commands/Workspace/WorkspaceClean.cs
deleted file mode 100644
diff --git a/Engine/Source/Programs/Horde/Horde/Commands/Workspace/WorkspaceClear.cs b/Engine/Source/Programs/Horde/Horde/Commands/Workspace/WorkspaceClear.cs
deleted file mode 100644
diff --git a/Engine/Source/Programs/Horde/Horde/Commands/Workspace/WorkspaceDump.cs b/Engine/Source/Programs/Horde/Horde/Commands/Workspace/WorkspaceDump.cs
deleted file mode 100644
diff --git a/Engine/Source/Programs/Horde/Horde/Commands/Workspace/WorkspacePopulateCache.cs b/Engine/Source/Programs/Horde/Horde/Commands/Workspace/WorkspacePopulateCache.cs
deleted file mode 100644
diff --git a/Engine/Source/Programs/Horde/Horde/Commands/Workspace/WorkspacePurgeCache.cs b/Engine/Source/Programs/Horde/Horde/Commands/Workspace/WorkspacePurgeCache.cs
deleted file mode 100644
diff --git a/Engine/Source/Programs/Horde/Horde/Commands/Workspace/WorkspaceRepair.cs b/Engine/Source/Programs/Horde/Horde/Commands/Workspace/WorkspaceRepair.cs
deleted file mode 100644
diff --git a/Engine/Source/Programs/Horde/Horde/Commands/Workspace/WorkspaceRevert.cs b/Engine/Source/Programs/Horde/Horde/Commands/Workspace/WorkspaceRevert.cs
deleted file mode 100644
diff --git a/Engine/Source/Programs/Horde/Horde/Commands/Workspace/WorkspaceSetup.cs b/Engine/Source/Programs/Horde/Horde/Commands/Workspace/WorkspaceSetup.cs
deleted file mode 100644
diff --git a/Engine/Source/Programs/Horde/Horde/Commands/Workspace/WorkspaceSizes.cs b/Engine/Source/Programs/Horde/Horde/Commands/Workspace/WorkspaceSizes.cs
deleted file mode 100644
diff --git a/Engine/Source/Programs/Horde/Horde/Commands/Workspace/WorkspaceStatus.cs b/Engine/Source/Programs/Horde/Horde/Commands/Workspace/WorkspaceStatus.cs
deleted file mode 100644
diff --git a/Engine/Source/Programs/Horde/Horde/Commands/Workspace/WorkspaceSync.cs b/Engine/Source/Programs/Horde/Horde/Commands/Workspace/WorkspaceSync.cs
deleted file mode 100644
