PR #10571: Fix StatusService AddAccessRule exception for Users (due to computer localization)

diff --git a/Engine/Source/Programs/Horde/Horde.Agent/Services/StatusService.cs b/Engine/Source/Programs/Horde/Horde.Agent/Services/StatusService.cs
--- a/Engine/Source/Programs/Horde/Horde.Agent/Services/StatusService.cs
+++ b/Engine/Source/Programs/Horde/Horde.Agent/Services/StatusService.cs
@@ -148,7 +148,11 @@ namespace Horde.Agent.Services
 
 			PipeSecurity pipeSecurity = new PipeSecurity();
 			pipeSecurity.AddAccessRule(new PipeAccessRule(WindowsIdentity.GetCurrent().Name, PipeAccessRights.FullControl, AccessControlType.Allow));
-			pipeSecurity.AddAccessRule(new PipeAccessRule("Users", PipeAccessRights.ReadWrite, AccessControlType.Allow));
+			
+			IdentityReference usersReference = new SecurityIdentifier(WellKnownSidType.BuiltinUsersSid, null).Translate(typeof(NTAccount));
+			string users = usersReference.ToString().Replace(@"builtin\", "", StringComparison.InvariantCultureIgnoreCase);
+			
+			pipeSecurity.AddAccessRule(new PipeAccessRule(users, PipeAccessRights.ReadWrite, AccessControlType.Allow));
 
 			using (NamedPipeServerStream pipeServer = NamedPipeServerStreamAcl.Create(AgentMessagePipe.PipeName, PipeDirection.InOut, NumPipes, PipeTransmissionMode.Byte, PipeOptions.Asynchronous, 0, 0, pipeSecurity))
 			{
