Horde: Add a -SavePreprocessed job option, which preprocesses the build script and saves it as an artifact without using it in downstream steps.

diff --git a/Engine/Source/Programs/Horde/Horde.Agent/Execution/JobExecutor.cs b/Engine/Source/Programs/Horde/Horde.Agent/Execution/JobExecutor.cs
--- a/Engine/Source/Programs/Horde/Horde.Agent/Execution/JobExecutor.cs
+++ b/Engine/Source/Programs/Horde/Horde.Agent/Execution/JobExecutor.cs
@@ -184,6 +184,7 @@ namespace Horde.Agent.Execution
 		protected List<string> _targets = new List<string>();
 		protected string? _scriptFileName;
 		protected bool _preprocessScript;
+		protected bool _savePreprocessedScript;
 
 		/// <summary>
 		/// Logger for the local agent process (as opposed to job logger)
@@ -276,6 +277,10 @@ namespace Horde.Agent.Execution
 				{
 					_preprocessScript = true;
 				}
+				else if (argument.Equals("-SavePreprocessed", StringComparison.OrdinalIgnoreCase))
+				{
+					_savePreprocessedScript = true;
+				}
 				else if (argument.StartsWith(TargetArgumentPrefix, StringComparison.OrdinalIgnoreCase))
 				{
 					_targets.Add(argument.Substring(TargetArgumentPrefix.Length));
@@ -464,7 +469,7 @@ namespace Horde.Agent.Execution
 
 			FileReference? preprocessedScriptFile = null;
 			FileReference? preprocessedSchemaFile = null;
-			if (_preprocessScript)
+			if (_preprocessScript || _savePreprocessedScript)
 			{
 				preprocessedScriptFile = FileReference.Combine(workspaceDir, PreprocessedScript);
 				arguments.AppendArgument("-Preprocess=", preprocessedScriptFile.FullName);
