Horde: Add a /api/v1/server/validateconfig endpoint that allows submitting a set of config files to validate on the server.

diff --git a/Engine/Source/Programs/Horde/Horde.Server/Configuration/ConfigService.cs b/Engine/Source/Programs/Horde/Horde.Server/Configuration/ConfigService.cs
--- a/Engine/Source/Programs/Horde/Horde.Server/Configuration/ConfigService.cs
+++ b/Engine/Source/Programs/Horde/Horde.Server/Configuration/ConfigService.cs
@@ -5,6 +5,7 @@ using System.Collections.Generic;
 using System.IO;
 using System.IO.Compression;
 using System.Linq;
+using System.Security.Cryptography.Xml;
 using System.Text.Json;
 using System.Threading;
 using System.Threading.Tasks;
@@ -13,10 +14,12 @@ using EpicGames.Redis;
 using Horde.Server.Projects;
 using Horde.Server.Server;
 using Horde.Server.Streams;
+using Horde.Server.Users;
 using Horde.Server.Utilities;
 using HordeCommon;
 using Microsoft.Extensions.Hosting;
 using Microsoft.Extensions.Logging;
+using Microsoft.Extensions.Logging.Abstractions;
 using Microsoft.Extensions.Options;
 using Microsoft.Extensions.Primitives;
 using ProtoBuf;
@@ -157,6 +160,99 @@ namespace Horde.Server.Configuration
 			_ticker.Dispose();
 		}
 
+		class OverrideConfigFile : IConfigFile
+		{
+			readonly byte[] _data;
+
+			public Uri Uri { get; }
+			public string Revision => "New";
+			public IUser? Author => null;
+			public bool WasRead { get; private set; }
+
+			public OverrideConfigFile(Uri uri, byte[] data)
+			{
+				Uri = uri;
+				_data = data;
+			}
+
+			public ValueTask<ReadOnlyMemory<byte>> ReadAsync(CancellationToken cancellationToken)
+			{
+				WasRead = true;
+				return new ValueTask<ReadOnlyMemory<byte>>(_data);
+			}
+		}
+
+		class OverrideConfigSource : IConfigSource
+		{
+			readonly IConfigSource _inner;
+			readonly Dictionary<Uri, OverrideConfigFile> _overrides;
+
+			public string Scheme => _inner.Scheme;
+
+			public OverrideConfigSource(IConfigSource inner, Dictionary<Uri, OverrideConfigFile> overrides)
+			{
+				_inner = inner;
+				_overrides = overrides;
+			}
+
+			public void Add(OverrideConfigFile file) => _overrides.Add(file.Uri, file);
+
+			public async Task<IConfigFile[]> GetAsync(Uri[] uris, CancellationToken cancellationToken)
+			{
+				IConfigFile[] files = new IConfigFile[uris.Length];
+				for (int idx = 0; idx < uris.Length; idx++)
+				{
+					IConfigFile? file;
+					if (_overrides.TryGetValue(uris[idx], out OverrideConfigFile? overrideFile))
+					{
+						file = overrideFile;
+					}
+					else 
+					{
+						file = (await _inner.GetAsync(new[] { uris[idx] }, cancellationToken))[0];
+					}
+					files[idx] = file;
+				}
+				return files;
+			}
+		}
+
+		/// <summary>
+		/// Validate a new set of config files. Parses and runs PostLoad methods on them.
+		/// </summary>
+		public async Task<string?> ValidateAsync(Dictionary<Uri, byte[]> files, CancellationToken cancellationToken)
+		{
+			Dictionary<Uri, OverrideConfigFile> overrideFiles = new Dictionary<Uri, OverrideConfigFile>();
+			foreach ((Uri uri, byte[] data) in files)
+			{
+				overrideFiles[uri] = new OverrideConfigFile(uri, data);
+			}
+
+			Dictionary<string, IConfigSource> overrideSources = new Dictionary<string, IConfigSource>(_sources.Count, _sources.Comparer);
+			foreach ((string schema, IConfigSource source) in _sources)
+			{
+				overrideSources.Add(schema, new OverrideConfigSource(source, overrideFiles));
+			}
+		
+			ConfigContext context = new ConfigContext(_jsonOptions, overrideSources, NullLogger.Instance);
+			try
+			{
+				Uri globalConfigUri = GetGlobalConfigUri();
+				GlobalConfig globalConfig = await ConfigType.ReadAsync<GlobalConfig>(globalConfigUri, context, cancellationToken);
+				globalConfig.PostLoad(_serverSettings);
+				return null;
+			}
+			catch (ConfigException ex)
+			{
+				string trace = String.Join("\n", context.IncludeStack.Select(x => $"\n  {x}"));
+				return $"{ex.Message}\nInclude stack:\n{trace}";
+			}
+			catch (Exception ex) when (ex is not ConfigException)
+			{
+				return ex.ToString();
+			}
+		}
+
 		/// <inheritdoc/>
 		public async Task StartAsync(CancellationToken cancellationToken)
 		{
@@ -327,25 +423,29 @@ namespace Horde.Server.Configuration
 		}
 
 		/// <summary>
-		/// Create a new snapshot object
+		/// Gets the path to the root config file
 		/// </summary>
-		/// <param name="cancellationToken">Cancellation token for the operation</param>
-		/// <returns>New config snapshot</returns>
-		async Task<ConfigSnapshot> CreateSnapshotAsync(CancellationToken cancellationToken)
+		Uri GetGlobalConfigUri()
 		{
-			// Get the path to the root config file
-			Uri globalConfigUri;
 			if (Path.IsPathRooted(_serverSettings.ConfigPath) && !_serverSettings.ConfigPath.StartsWith("//", StringComparison.Ordinal))
 			{
 				// absolute path to config
-				globalConfigUri = new Uri(_serverSettings.ConfigPath);
+				return new Uri(_serverSettings.ConfigPath);
 			}
 			else
 			{
 				// relative (development) or perforce path
-				globalConfigUri = ConfigType.CombinePaths(new Uri(FileReference.Combine(Program.AppDir, "_").FullName), _serverSettings.ConfigPath);
+				return ConfigType.CombinePaths(new Uri(FileReference.Combine(Program.AppDir, "_").FullName), _serverSettings.ConfigPath);
 			}
+		}
 
+		/// <summary>
+		/// Create a new snapshot object
+		/// </summary>
+		/// <param name="cancellationToken">Cancellation token for the operation</param>
+		/// <returns>New config snapshot</returns>
+		async Task<ConfigSnapshot> CreateSnapshotAsync(CancellationToken cancellationToken)
+		{
 			// Read the config files
 			ConfigContext context = new ConfigContext(_jsonOptions, _sources, _logger);
 			try
@@ -354,6 +454,7 @@ namespace Horde.Server.Configuration
 				snapshot.ServerVersion = Program.Version.ToString();
 
 				// Read the new config in
+				Uri globalConfigUri = GetGlobalConfigUri();
 				GlobalConfig globalConfig = await ConfigType.ReadAsync<GlobalConfig>(globalConfigUri, context, cancellationToken);
 
 				// Serialize it back out to a byte array
diff --git a/Engine/Source/Programs/Horde/Horde.Server/Server/ServerMessages.cs b/Engine/Source/Programs/Horde/Horde.Server/Server/ServerMessages.cs
--- a/Engine/Source/Programs/Horde/Horde.Server/Server/ServerMessages.cs
+++ b/Engine/Source/Programs/Horde/Horde.Server/Server/ServerMessages.cs
@@ -1,6 +1,7 @@
 // Copyright Epic Games, Inc. All Rights Reserved.
 
 using System;
+using System.Collections.Generic;
 using System.Diagnostics;
 using System.Reflection;
 using System.Runtime.InteropServices;
@@ -12,30 +13,30 @@ namespace Horde.Server.Server
 	/// </summary>
 	public class GetServerInfoResponse
 	{
-        /// <summary>
+		/// <summary>
 		/// Server version info
 		/// </summary>
-        public string ServerVersion { get; set; }
+		public string ServerVersion { get; set; }
 
 		/// <summary>
 		/// The current agent version string
 		/// </summary>
 		public string? AgentVersion { get; set; }
 
-        /// <summary>
-        /// The operating system server is hosted on
-        /// </summary>
-        public string OsDescription { get; set; }
+		/// <summary>
+		/// The operating system server is hosted on
+		/// </summary>
+		public string OsDescription { get; set; }
 
 		/// <summary>
 		/// Constructor
 		/// </summary>
 		public GetServerInfoResponse(string? agentVersion)
-        {
-            FileVersionInfo versionInfo = FileVersionInfo.GetVersionInfo(Assembly.GetExecutingAssembly().Location);		
+		{
+			FileVersionInfo versionInfo = FileVersionInfo.GetVersionInfo(Assembly.GetExecutingAssembly().Location);
 			ServerVersion = versionInfo.ProductVersion ?? String.Empty;
 			AgentVersion = agentVersion;
-			OsDescription = RuntimeInformation.OSDescription;			
+			OsDescription = RuntimeInformation.OSDescription;
 		}
 	}
 
@@ -101,5 +102,48 @@ namespace Horde.Server.Server
 			LocalRedirectUrls = settings.OidcLocalRedirectUrls;
 		}
 	}
+
+	/// <summary>
+	/// Describes an individual config file to validate
+	/// </summary>
+	public class ValidateConfigFileRequest
+	{
+		/// <summary>
+		/// Path for the file to substitute
+		/// </summary>
+		public Uri? Uri { get; set; }
+
+		/// <summary>
+		/// Data for the file
+		/// </summary>
+		public byte[] Data { get; set; } = Array.Empty<byte>();
+	}
+
+	/// <summary>
+	/// Request to validate server configuration with the given files replacing their checked-in counterparts.
+	/// </summary>
+	public class ValidateConfigRequest
+	{
+		/// <summary>
+		/// Files to be validated.
+		/// </summary>
+		public List<ValidateConfigFileRequest> Files { get; set; } = new List<ValidateConfigFileRequest>();
+	}
+
+	/// <summary>
+	/// Response from validating config files
+	/// </summary>
+	public class ValidateConfigResponse
+	{
+		/// <summary>
+		/// Whether the files were validated successfully
+		/// </summary>
+		public bool Result { get; set; }
+
+		/// <summary>
+		/// Output message from validation
+		/// </summary>
+		public string? Message { get; set; }
+	}
 }
 
