Horde: Add endpoint for checking config file syntax.

diff --git a/Engine/Source/Programs/Horde/Horde.Server/Server/ServerController.cs b/Engine/Source/Programs/Horde/Horde.Server/Server/ServerController.cs
--- a/Engine/Source/Programs/Horde/Horde.Server/Server/ServerController.cs
+++ b/Engine/Source/Programs/Horde/Horde.Server/Server/ServerController.cs
@@ -1,11 +1,15 @@
 // Copyright Epic Games, Inc. All Rights Reserved.
 
+using System;
+using System.Collections.Generic;
 using System.Diagnostics;
 using System.Reflection;
 using System.Text.Json;
+using System.Threading;
 using System.Threading.Tasks;
 using Horde.Server.Acls;
 using Horde.Server.Agents;
+using Horde.Server.Configuration;
 using Horde.Server.Projects;
 using Horde.Server.Tools;
 using Horde.Server.Utilities;
@@ -28,15 +32,17 @@ namespace Horde.Server.Server
 	{
 		readonly IToolCollection _toolCollection;
 		readonly IClock _clock;
+		readonly ConfigService _configService;
 		readonly IOptionsSnapshot<GlobalConfig> _globalConfig;
 
 		/// <summary>
 		/// Constructor
 		/// </summary>
-		public ServerController(IToolCollection toolCollection, IClock clock, IOptionsSnapshot<GlobalConfig> globalConfig)
+		public ServerController(IToolCollection toolCollection, IClock clock, ConfigService configService, IOptionsSnapshot<GlobalConfig> globalConfig)
 		{
 			_toolCollection = toolCollection;
 			_clock = clock;
+			_configService = configService;
 			_globalConfig = globalConfig;
 		}
 
@@ -117,5 +123,30 @@ namespace Horde.Server.Server
 		{
 			return new GetAuthConfigResponse(_globalConfig.Value.ServerSettings);
 		}
+
+		/// <summary>
+		/// Returns settings for automating auth against this server
+		/// </summary>
+		[HttpPost]
+		[Route("/api/v1/server/validateconfig")]
+		public async Task<ActionResult<ValidateConfigResponse>> ValidateConfigAsync(ValidateConfigRequest request, CancellationToken cancellationToken)
+		{
+			Dictionary<Uri, byte[]> files = new Dictionary<Uri, byte[]>();
+			foreach (ValidateConfigFileRequest file in request.Files)
+			{
+				if (file.Uri != null)
+				{
+					files.Add(file.Uri, file.Data);
+				}
+			}
+
+			string? message = await _configService.ValidateAsync(files, cancellationToken);
+
+			ValidateConfigResponse response = new ValidateConfigResponse();
+			response.Result = message != null;
+			response.Message = message;
+
+			return response;
+		}
 	}
 }
