Horde: Ensure each validated config file is read by the server.

diff --git a/Engine/Source/Programs/Horde/Horde.Server/Configuration/ConfigService.cs b/Engine/Source/Programs/Horde/Horde.Server/Configuration/ConfigService.cs
--- a/Engine/Source/Programs/Horde/Horde.Server/Configuration/ConfigService.cs
+++ b/Engine/Source/Programs/Horde/Horde.Server/Configuration/ConfigService.cs
@@ -240,6 +240,15 @@ namespace Horde.Server.Configuration
 				Uri globalConfigUri = GetGlobalConfigUri();
 				GlobalConfig globalConfig = await ConfigType.ReadAsync<GlobalConfig>(globalConfigUri, context, cancellationToken);
 				globalConfig.PostLoad(_serverSettings);
+
+				foreach (OverrideConfigFile file in overrideFiles.Values)
+				{
+					if (!file.WasRead)
+					{
+						return $"File {file.Uri} was not read by server";
+					}
+				}
+
 				return null;
 			}
 			catch (ConfigException ex)
diff --git a/Engine/Source/Programs/Horde/Horde.Server/Server/ServerController.cs b/Engine/Source/Programs/Horde/Horde.Server/Server/ServerController.cs
--- a/Engine/Source/Programs/Horde/Horde.Server/Server/ServerController.cs
+++ b/Engine/Source/Programs/Horde/Horde.Server/Server/ServerController.cs
@@ -143,7 +143,7 @@ namespace Horde.Server.Server
 			string? message = await _configService.ValidateAsync(files, cancellationToken);
 
 			ValidateConfigResponse response = new ValidateConfigResponse();
-			response.Result = message != null;
+			response.Result = message == null;
 			response.Message = message;
 
 			return response;
