Horde: Skip syncing if cache file is not dirty

Replicates the behavior that Synced.txt has in PerforceExecutor for syncs that used a cache file, such as AutoSDK.

diff --git a/Engine/Source/Programs/Horde/Horde.Agent/Execution/ManagedWorkspaceMaterializer.cs b/Engine/Source/Programs/Horde/Horde.Agent/Execution/ManagedWorkspaceMaterializer.cs
--- a/Engine/Source/Programs/Horde/Horde.Agent/Execution/ManagedWorkspaceMaterializer.cs
+++ b/Engine/Source/Programs/Horde/Horde.Agent/Execution/ManagedWorkspaceMaterializer.cs
@@ -1,11 +1,9 @@
 // Copyright Epic Games, Inc. All Rights Reserved.
 
 using System.Collections.Generic;
-using System.Text;
 using System.Threading;
 using System.Threading.Tasks;
 using EpicGames.Core;
-using EpicGames.Perforce;
 using Horde.Agent.Utility;
 using HordeCommon.Rpc.Messages;
 using Microsoft.Extensions.Logging;
@@ -103,7 +101,12 @@ public class ManagedWorkspaceMaterializer : IWorkspaceMaterializer
 		FileReference cacheFile = FileReference.Combine(_workspace.MetadataDir, "Contents.dat");
 		if (_useCacheFile)
 		{
-			await _workspace.UpdateLocalCacheMarkerAsync(cacheFile, changeNum, preflightChangeNum);
+			bool isSyncedDataDirty = await _workspace.UpdateLocalCacheMarkerAsync(cacheFile, changeNum, preflightChangeNum);
+			scope.Span.SetTag("IsSyncedDataDirty", isSyncedDataDirty);
+			if (!isSyncedDataDirty)
+			{
+				return;
+			}
 		}
 		else
 		{
@@ -120,6 +123,7 @@ public class ManagedWorkspaceMaterializer : IWorkspaceMaterializer
 		scope.Span.SetTag("Cluster", _agentWorkspace.Cluster);
 		scope.Span.SetTag("Incremental", _agentWorkspace.Incremental);
 		scope.Span.SetTag("Method", _agentWorkspace.Method);
+		scope.Span.SetTag("UseCacheFile", _useCacheFile);
 		return scope;
 	}
 }
\ No newline at end of file
diff --git a/Engine/Source/Programs/Horde/Horde.Agent/Utility/WorkspaceInfo.cs b/Engine/Source/Programs/Horde/Horde.Agent/Utility/WorkspaceInfo.cs
--- a/Engine/Source/Programs/Horde/Horde.Agent/Utility/WorkspaceInfo.cs
+++ b/Engine/Source/Programs/Horde/Horde.Agent/Utility/WorkspaceInfo.cs
@@ -305,8 +305,8 @@ namespace Horde.Agent.Utility
 		/// <param name="cacheFile">Path to the cache file</param>
 		/// <param name="change">The current change being built</param>
 		/// <param name="preflightChange">The preflight changelist number</param>
-		/// <returns>Async task</returns>
-		public async Task UpdateLocalCacheMarkerAsync(FileReference cacheFile, int change, int preflightChange)
+		/// <returns>True if cache file was replaced. That is if changelist or stream view is different from what is already synced</returns>
+		public async Task<bool> UpdateLocalCacheMarkerAsync(FileReference cacheFile, int change, int preflightChange)
 		{
 			// Create the new cache file descriptor
 			string newDescriptor = GetCacheMarkerDescriptor(change, preflightChange);
@@ -320,7 +320,7 @@ namespace Horde.Agent.Utility
 					string oldDescriptor = await FileReference.ReadAllTextAsync(descriptorFile);
 					if (oldDescriptor.Equals(newDescriptor, StringComparison.Ordinal))
 					{
-						return;
+						return false;
 					}
 					else
 					{
@@ -332,6 +332,7 @@ namespace Horde.Agent.Utility
 
 			// Write the new descriptor file
 			await FileReference.WriteAllTextAsync(descriptorFile, newDescriptor);
+			return true;
 		}
 
 		string GetCacheMarkerDescriptor(int change, int preflightChange)
