Horde: Pass a shelved change into config validation command.

diff --git a/Engine/Source/Programs/Horde/Horde.Server/Server/ServerController.cs b/Engine/Source/Programs/Horde/Horde.Server/Server/ServerController.cs
--- a/Engine/Source/Programs/Horde/Horde.Server/Server/ServerController.cs
+++ b/Engine/Source/Programs/Horde/Horde.Server/Server/ServerController.cs
@@ -7,9 +7,11 @@ using System.Reflection;
 using System.Text.Json;
 using System.Threading;
 using System.Threading.Tasks;
+using EpicGames.Perforce;
 using Horde.Server.Acls;
 using Horde.Server.Agents;
 using Horde.Server.Configuration;
+using Horde.Server.Perforce;
 using Horde.Server.Projects;
 using Horde.Server.Tools;
 using Horde.Server.Utilities;
@@ -33,16 +35,18 @@ namespace Horde.Server.Server
 		readonly IToolCollection _toolCollection;
 		readonly IClock _clock;
 		readonly ConfigService _configService;
+		readonly IPerforceService _perforceService;
 		readonly IOptionsSnapshot<GlobalConfig> _globalConfig;
 
 		/// <summary>
 		/// Constructor
 		/// </summary>
-		public ServerController(IToolCollection toolCollection, IClock clock, ConfigService configService, IOptionsSnapshot<GlobalConfig> globalConfig)
+		public ServerController(IToolCollection toolCollection, IClock clock, ConfigService configService, IPerforceService perforceService, IOptionsSnapshot<GlobalConfig> globalConfig)
 		{
 			_toolCollection = toolCollection;
 			_clock = clock;
 			_configService = configService;
+			_perforceService = perforceService;
 			_globalConfig = globalConfig;
 		}
 
@@ -131,13 +135,22 @@ namespace Horde.Server.Server
 		[Route("/api/v1/server/validateconfig")]
 		public async Task<ActionResult<ValidateConfigResponse>> ValidateConfigAsync(ValidateConfigRequest request, CancellationToken cancellationToken)
 		{
+			IPooledPerforceConnection perforce = await _perforceService.ConnectAsync(request.Cluster ?? "default", cancellationToken: cancellationToken);
+			DescribeRecord record = await perforce.DescribeAsync(request.ShelvedChange, cancellationToken);
+
 			Dictionary<Uri, byte[]> files = new Dictionary<Uri, byte[]>();
-			foreach (ValidateConfigFileRequest file in request.Files)
+			foreach (DescribeFileRecord fileRecord in record.Files)
 			{
-				if (file.Uri != null)
+				PerforceResponse<PrintRecord<byte[]>> printRecordResponse = await perforce.TryPrintDataAsync($"{fileRecord.DepotFile}@={request.ShelvedChange}", cancellationToken);
+				if (!printRecordResponse.Succeeded || printRecordResponse.Data.Contents == null)
 				{
-					files.Add(file.Uri, file.Data);
+					return BadRequest($"Unable to print contents of {fileRecord.DepotFile}@={request.ShelvedChange}");
 				}
+
+				PrintRecord<byte[]> printRecord = printRecordResponse.Data;
+
+				Uri uri = new Uri($"perforce://{request.Cluster}{printRecord.DepotFile}");
+				files.Add(uri, printRecord.Contents);
 			}
 
 			string? message = await _configService.ValidateAsync(files, cancellationToken);
diff --git a/Engine/Source/Programs/Horde/Horde.Server/Server/ServerMessages.cs b/Engine/Source/Programs/Horde/Horde.Server/Server/ServerMessages.cs
--- a/Engine/Source/Programs/Horde/Horde.Server/Server/ServerMessages.cs
+++ b/Engine/Source/Programs/Horde/Horde.Server/Server/ServerMessages.cs
@@ -125,9 +125,14 @@ namespace Horde.Server.Server
 	public class ValidateConfigRequest
 	{
 		/// <summary>
-		/// Files to be validated.
+		/// Perforce cluster to retrieve from
 		/// </summary>
-		public List<ValidateConfigFileRequest> Files { get; set; } = new List<ValidateConfigFileRequest>();
+		public string? Cluster { get; set; }
+
+		/// <summary>
+		/// Change to test
+		/// </summary>
+		public int ShelvedChange { get; set; }
 	}
 
 	/// <summary>
diff --git a/Engine/Source/Programs/Shared/EpicGames.Perforce/IPerforceOutput.cs b/Engine/Source/Programs/Shared/EpicGames.Perforce/IPerforceOutput.cs
--- a/Engine/Source/Programs/Shared/EpicGames.Perforce/IPerforceOutput.cs
+++ b/Engine/Source/Programs/Shared/EpicGames.Perforce/IPerforceOutput.cs
@@ -211,6 +211,8 @@ namespace EpicGames.Perforce
 			return result.ToString();
 		}
 
+		static readonly byte[] s_connectionErrorPrefix = Encoding.UTF8.GetBytes("Perforce client error:");
+
 		/// <summary>
 		/// Read a list of responses from the child process
 		/// </summary>
@@ -233,7 +235,14 @@ namespace EpicGames.Perforce
 				ReadOnlyMemory<byte> data = perforce.Data;
 				if (data.Length > 0 && responses.Count == 0 && data.Span[0] != '{')
 				{
-					throw new PerforceException("Unexpected response from server (expected '{{'):{0}", FormatDataAsString(data.Span));
+					if (data.Span.StartsWith(s_connectionErrorPrefix))
+					{
+						throw new PerforceException("Unable to connect to Perforce server:{0}", FormatDataAsString(data.Span.Slice(s_connectionErrorPrefix.Length)).TrimStart('\n', '\r'));
+					}
+					else
+					{
+						throw new PerforceException("Unexpected response from server (expected '{{'):{0}", FormatDataAsString(data.Span));
+					}
 				}
 
 				// Parse the responses from the current buffer
