Horde: Fix empty cluster parameter in Uris for config checks.

diff --git a/Engine/Source/Programs/Horde/Horde.Server/Server/ServerController.cs b/Engine/Source/Programs/Horde/Horde.Server/Server/ServerController.cs
--- a/Engine/Source/Programs/Horde/Horde.Server/Server/ServerController.cs
+++ b/Engine/Source/Programs/Horde/Horde.Server/Server/ServerController.cs
@@ -135,7 +135,9 @@ namespace Horde.Server.Server
 		[Route("/api/v1/server/validateconfig")]
 		public async Task<ActionResult<ValidateConfigResponse>> ValidateConfigAsync(ValidateConfigRequest request, CancellationToken cancellationToken)
 		{
-			IPooledPerforceConnection perforce = await _perforceService.ConnectAsync(request.Cluster ?? "default", cancellationToken: cancellationToken);
+			string cluster = request.Cluster ?? "default";
+
+			IPooledPerforceConnection perforce = await _perforceService.ConnectAsync(cluster, cancellationToken: cancellationToken);
 			DescribeRecord record = await perforce.DescribeAsync(request.ShelvedChange, cancellationToken);
 
 			Dictionary<Uri, byte[]> files = new Dictionary<Uri, byte[]>();
@@ -149,7 +151,7 @@ namespace Horde.Server.Server
 
 				PrintRecord<byte[]> printRecord = printRecordResponse.Data;
 
-				Uri uri = new Uri($"perforce://{request.Cluster}{printRecord.DepotFile}");
+				Uri uri = new Uri($"perforce://{cluster}{printRecord.DepotFile}");
 				files.Add(uri, printRecord.Contents);
 			}
 
