Horde: Fix matching of MSTest failures when test time is reported in seconds rather than milliseconds.

diff --git a/Engine/Source/Programs/AutomationTool/AutomationUtils/Matchers/MsTestEventMatcher.cs b/Engine/Source/Programs/AutomationTool/AutomationUtils/Matchers/MsTestEventMatcher.cs
--- a/Engine/Source/Programs/AutomationTool/AutomationUtils/Matchers/MsTestEventMatcher.cs
+++ b/Engine/Source/Programs/AutomationTool/AutomationUtils/Matchers/MsTestEventMatcher.cs
@@ -10,7 +10,7 @@ namespace AutomationUtils.Matchers
 {
 	class MsTestEventMatcher : ILogEventMatcher
 	{
-		static readonly Regex s_failedPattern = new Regex(@"^  \s*Failed [A-Za-z0-9_]+ \[\d+ ms\]$");
+		static readonly Regex s_failedPattern = new Regex(@"^  \s*Failed [A-Za-z0-9_]+ \[\d+ (?:s|ms)\]$");
 		static readonly Regex s_detailPattern = new Regex(@"^  \s*(?:Error Message|Stack Trace):$");
 
 		public LogEventMatch? Match(ILogCursor cursor)
diff --git a/Engine/Source/Programs/Horde/Horde.Agent.Tests/LogParserTests.cs b/Engine/Source/Programs/Horde/Horde.Agent.Tests/LogParserTests.cs
--- a/Engine/Source/Programs/Horde/Horde.Agent.Tests/LogParserTests.cs
+++ b/Engine/Source/Programs/Horde/Horde.Agent.Tests/LogParserTests.cs
@@ -713,6 +713,25 @@ namespace Horde.Agent.Tests
 			CheckEventGroup(logEvents, 5, 6, LogLevel.Error, KnownLogEvents.MSTest);
 		}
 
+		[TestMethod]
+		public void MsTestEventMatcher2()
+		{
+			string[] lines =
+			{
+				@"  Passed ConditionSimple [154 ms]",
+				@"  Passed GetPoolQueueSizes [1 s]",
+				@"  Failed DowntimeActive [6 s]",
+				@"  Error Message:",
+				@"   Assert.AreEqual failed. Expected:<1>. Actual:<0>.",
+				@"  Stack Trace:",
+				@"     at Horde.Server.Tests.Fleet.JobQueueStrategyTest.DowntimeActive() in /app/Source/Programs/Horde/Horde.Server.Tests/Fleet/JobQueueStrategyTest.cs:line 40",
+			};
+
+			List<LogEvent> logEvents = Parse(lines);
+			Assert.AreEqual(5, logEvents.Count);
+			CheckEventGroup(logEvents, 2, 5, LogLevel.Error, KnownLogEvents.MSTest);
+		}
+
 		static List<LogEvent> Parse(IEnumerable<string> lines)
 		{
 			return Parse(String.Join("\n", lines));
