Horde: Rename endpoint for validating config to preflightconfig. Also improve messaging for errors.

diff --git a/Engine/Source/Programs/Horde/Horde.Server/Configuration/ConfigService.cs b/Engine/Source/Programs/Horde/Horde.Server/Configuration/ConfigService.cs
--- a/Engine/Source/Programs/Horde/Horde.Server/Configuration/ConfigService.cs
+++ b/Engine/Source/Programs/Horde/Horde.Server/Configuration/ConfigService.cs
@@ -251,14 +251,10 @@ namespace Horde.Server.Configuration
 
 				return null;
 			}
-			catch (ConfigException ex)
+			catch (Exception ex)
 			{
 				string trace = String.Join("\n", context.IncludeStack.Select(x => $"\n  {x}"));
-				return $"{ex.Message}\nInclude stack:\n{trace}";
-			}
-			catch (Exception ex) when (ex is not ConfigException)
-			{
-				return ex.ToString();
+				return $"{ex.Message}\n\nInclude stack:\n{trace}\n\nTrace:\n{ex.StackTrace}";
 			}
 		}
 
diff --git a/Engine/Source/Programs/Horde/Horde.Server/Server/ServerController.cs b/Engine/Source/Programs/Horde/Horde.Server/Server/ServerController.cs
--- a/Engine/Source/Programs/Horde/Horde.Server/Server/ServerController.cs
+++ b/Engine/Source/Programs/Horde/Horde.Server/Server/ServerController.cs
@@ -132,8 +132,8 @@ namespace Horde.Server.Server
 		/// Returns settings for automating auth against this server
 		/// </summary>
 		[HttpPost]
-		[Route("/api/v1/server/validateconfig")]
-		public async Task<ActionResult<ValidateConfigResponse>> ValidateConfigAsync(ValidateConfigRequest request, CancellationToken cancellationToken)
+		[Route("/api/v1/server/preflightconfig")]
+		public async Task<ActionResult<PreflightConfigResponse>> PreflightConfigAsync(PreflightConfigRequest request, CancellationToken cancellationToken)
 		{
 			string cluster = request.Cluster ?? "default";
 
@@ -157,7 +157,7 @@ namespace Horde.Server.Server
 
 			string? message = await _configService.ValidateAsync(files, cancellationToken);
 
-			ValidateConfigResponse response = new ValidateConfigResponse();
+			PreflightConfigResponse response = new PreflightConfigResponse();
 			response.Result = message == null;
 			response.Message = message;
 
diff --git a/Engine/Source/Programs/Horde/Horde.Server/Server/ServerMessages.cs b/Engine/Source/Programs/Horde/Horde.Server/Server/ServerMessages.cs
--- a/Engine/Source/Programs/Horde/Horde.Server/Server/ServerMessages.cs
+++ b/Engine/Source/Programs/Horde/Horde.Server/Server/ServerMessages.cs
@@ -5,6 +5,7 @@ using System.Collections.Generic;
 using System.Diagnostics;
 using System.Reflection;
 using System.Runtime.InteropServices;
+using Amazon.S3.Model;
 
 namespace Horde.Server.Server
 {
@@ -103,26 +104,10 @@ namespace Horde.Server.Server
 		}
 	}
 
-	/// <summary>
-	/// Describes an individual config file to validate
-	/// </summary>
-	public class ValidateConfigFileRequest
-	{
-		/// <summary>
-		/// Path for the file to substitute
-		/// </summary>
-		public Uri? Uri { get; set; }
-
-		/// <summary>
-		/// Data for the file
-		/// </summary>
-		public byte[] Data { get; set; } = Array.Empty<byte>();
-	}
-
 	/// <summary>
 	/// Request to validate server configuration with the given files replacing their checked-in counterparts.
 	/// </summary>
-	public class ValidateConfigRequest
+	public class PreflightConfigRequest
 	{
 		/// <summary>
 		/// Perforce cluster to retrieve from
@@ -138,7 +123,7 @@ namespace Horde.Server.Server
 	/// <summary>
 	/// Response from validating config files
 	/// </summary>
-	public class ValidateConfigResponse
+	public class PreflightConfigResponse
 	{
 		/// <summary>
 		/// Whether the files were validated successfully
@@ -149,6 +134,11 @@ namespace Horde.Server.Server
 		/// Output message from validation
 		/// </summary>
 		public string? Message { get; set; }
+
+		/// <summary>
+		/// Detailed response
+		/// </summary>
+		public string? Detail { get; set; }
 	}
 }
 
