Horde: Iterating on error messages when config file validation fails.

diff --git a/Engine/Source/Programs/Horde/Horde.Server/Configuration/ConfigService.cs b/Engine/Source/Programs/Horde/Horde.Server/Configuration/ConfigService.cs
--- a/Engine/Source/Programs/Horde/Horde.Server/Configuration/ConfigService.cs
+++ b/Engine/Source/Programs/Horde/Horde.Server/Configuration/ConfigService.cs
@@ -6,6 +6,7 @@ using System.IO;
 using System.IO.Compression;
 using System.Linq;
 using System.Security.Cryptography.Xml;
+using System.Text;
 using System.Text.Json;
 using System.Threading;
 using System.Threading.Tasks;
@@ -17,6 +18,7 @@ using Horde.Server.Streams;
 using Horde.Server.Users;
 using Horde.Server.Utilities;
 using HordeCommon;
+using Microsoft.CodeAnalysis.CSharp;
 using Microsoft.Extensions.Hosting;
 using Microsoft.Extensions.Logging;
 using Microsoft.Extensions.Logging.Abstractions;
@@ -253,8 +255,18 @@ namespace Horde.Server.Configuration
 			}
 			catch (Exception ex)
 			{
-				string trace = String.Join("\n", context.IncludeStack.Select(x => $"\n  {x}"));
-				return $"{ex.Message}\n\nInclude stack:\n{trace}\n\nTrace:\n{ex.StackTrace}";
+				StringBuilder message = new StringBuilder(ex.Message);
+				message.Append("\n\nLocation:\n");
+				foreach (IConfigFile include in context.IncludeStack)
+				{
+					message.Append($"  {include.Uri}\n");
+				}
+				if (ex is not ConfigException && ex is not JsonException)
+				{
+					message.Append($"\n\nTrace:\n{ex.StackTrace}");
+					_logger.LogWarning(ex, "Unhandled exception while validating config files: {Message}", ex.Message);
+				}
+				return message.ToString();
 			}
 		}
 
