Horde: Tweaking output for config file preflights.

diff --git a/Engine/Source/Programs/Horde/Horde.Server/Configuration/ConfigFile.cs b/Engine/Source/Programs/Horde/Horde.Server/Configuration/ConfigFile.cs
--- a/Engine/Source/Programs/Horde/Horde.Server/Configuration/ConfigFile.cs
+++ b/Engine/Source/Programs/Horde/Horde.Server/Configuration/ConfigFile.cs
@@ -34,4 +34,30 @@ namespace Horde.Server.Configuration
 		/// <returns>UTF-8 encoded data for the config file</returns>
 		ValueTask<ReadOnlyMemory<byte>> ReadAsync(CancellationToken cancellationToken);
 	}
+
+	/// <summary>
+	/// Extension methods for config files
+	/// </summary>
+	public static class ConfigFileExtensions
+	{
+		/// <summary>
+		/// Formats a the URI for a config file in a format that is more readable for users
+		/// </summary>
+		/// <param name="file">Config file to get a path for</param>
+		public static string GetUserFormattedPath(this IConfigFile file)
+		{
+			const string DefaultPerforcePrefix = "perforce://default//";
+
+			string path = file.Uri.ToString();
+			if (path.StartsWith(DefaultPerforcePrefix, StringComparison.OrdinalIgnoreCase))
+			{
+				path = path.Substring(DefaultPerforcePrefix.Length - 2);				
+			}
+			if (!String.IsNullOrEmpty(file.Revision))
+			{
+				path = $"{path}@{file.Revision}";
+			}
+			return path;
+		}
+	}
 }
diff --git a/Engine/Source/Programs/Horde/Horde.Server/Configuration/ConfigService.cs b/Engine/Source/Programs/Horde/Horde.Server/Configuration/ConfigService.cs
--- a/Engine/Source/Programs/Horde/Horde.Server/Configuration/ConfigService.cs
+++ b/Engine/Source/Programs/Horde/Horde.Server/Configuration/ConfigService.cs
@@ -259,7 +259,8 @@ namespace Horde.Server.Configuration
 				message.Append("\n\nLocation:\n");
 				foreach (IConfigFile include in context.IncludeStack)
 				{
-					message.Append($"  {include.Uri}\n");
+					string path = include.GetUserFormattedPath();
+					message.Append($"  {path}\n");
 				}
 				if (ex is not ConfigException && ex is not JsonException)
 				{
