Horde: Fix issue with new agent and lease link query parameter handling
#skipci
#rnx

diff --git a/Engine/Source/Programs/Horde/HordeDashboard/src/base/components/Markdown/MarkdownLink.tsx b/Engine/Source/Programs/Horde/HordeDashboard/src/base/components/Markdown/MarkdownLink.tsx
--- a/Engine/Source/Programs/Horde/HordeDashboard/src/base/components/Markdown/MarkdownLink.tsx
+++ b/Engine/Source/Programs/Horde/HordeDashboard/src/base/components/Markdown/MarkdownLink.tsx
@@ -23,6 +23,17 @@ export const MarkdownLink: React.FunctionComponent<ILinkProps> = props => {
    if (href) {
 
       try {
+
+         if (href.startsWith("?")) {
+            const search = new URLSearchParams(href);
+            const csearch = new URLSearchParams(window.location.search);
+            search.forEach((v, k) => {
+               csearch.set(k, v);
+            });
+
+            href = "?" + csearch.toString();
+         }
+         
          const url = new URL(href);
          if (url.hostname === window.location.hostname) {
             href = url.pathname + url.search
diff --git a/Engine/Source/Programs/Horde/HordeDashboard/src/components/LogRender.tsx b/Engine/Source/Programs/Horde/HordeDashboard/src/components/LogRender.tsx
--- a/Engine/Source/Programs/Horde/HordeDashboard/src/components/LogRender.tsx
+++ b/Engine/Source/Programs/Horde/HordeDashboard/src/components/LogRender.tsx
@@ -168,14 +168,13 @@ const renderTags = (navigate: NavigateFunction, line: LogLine, lineNumber: numbe
          }
 
          return <a key={key} href="/"
-            
-            onAuxClick={(ev) =>
-            {
+
+            onAuxClick={(ev) => {
                ev.preventDefault();
                ev.stopPropagation()
                navigateToLeaseLog(true)
             }}
-         
+
             onClick={(ev) => {
                ev.stopPropagation();
                ev.preventDefault();
@@ -189,9 +188,10 @@ const renderTags = (navigate: NavigateFunction, line: LogLine, lineNumber: numbe
          return <a key={key} target="_blank" rel="noopener noreferrer" href={`https://msdn.microsoft.com/query/dev16.query?appId=Dev16IDEF1&l=EN-US&k=k(${text.toLowerCase()})&rd=true`} onClick={(ev) => ev.stopPropagation()}><Highlight search={search ? search : ""} className={logStyle.logLine}>{text}</Highlight></a>;
 
       } else if (tagType === TagType.AgentId) {
-         const search = window.location.search;
-         const url = `${window.location.pathname}` + (search ? search + `&agentId=${encodeURIComponent(text)}` : `?agentId=${encodeURIComponent(text)}`);
-         return <a key={key} href="/" onClick={async (ev) => { ev.stopPropagation(); ev.preventDefault(); navigate(url) }}><Highlight search={search ? search : ""} className={logStyle.logLine}>{text}</Highlight></a>;
+         const search = new URLSearchParams(window.location.search);
+         search.set("agentId", encodeURIComponent(text));
+         const url = `${window.location.pathname}?` + search.toString();
+         return <a key={key} href="/" onClick={async (ev) => { ev.stopPropagation(); ev.preventDefault(); navigate(url, { replace: true }) }}><Highlight search={search ? search : ""} className={logStyle.logLine}>{text}</Highlight></a>;
 
       } else if (tagType === TagType.SourceFile) {
 
diff --git a/Engine/Source/Programs/Horde/HordeDashboard/src/components/LogView.tsx b/Engine/Source/Programs/Horde/HordeDashboard/src/components/LogView.tsx
--- a/Engine/Source/Programs/Horde/HordeDashboard/src/components/LogView.tsx
+++ b/Engine/Source/Programs/Horde/HordeDashboard/src/components/LogView.tsx
@@ -689,7 +689,13 @@ export const LogList: React.FC<{ logId: string }> = observer(({ logId }) => {
          }
 
          return (
-            <Stack key={`key_log_line_${item.lineNumber}`} style={{ width: "max-content", height: handler.lineHeight }} onClick={() => { navigate(location.pathname + `?lineindex=${item.lineNumber - 1}`, { replace: true }) }}>
+            <Stack key={`key_log_line_${item.lineNumber}`} style={{ width: "max-content", height: handler.lineHeight }} onClick={() => {
+               const search = new URLSearchParams(window.location.search);
+               search.set("lineindex", (item.lineNumber - 1).toString());
+               const url = `${window.location.pathname}?` + search.toString();
+      
+               navigate(url, { replace: true })
+            }}>
                <div style={{ position: "relative" }}>
                   <Stack className={styles.logLine} tokens={{ childrenGap: 8 }} horizontal disableShrink={true}>
                      <Stack styles={{ root: { color: "#c0c0c0", width: 80, textAlign: "right", userSelect: "none", fontSize: handler.fontSize } }}>{prefix + item.lineNumber}</Stack>
