Horde: Support for inheriting settings between agent types, workspace types, and templates.

diff --git a/Engine/Source/Programs/Horde/Horde.Server/Agents/IAgent.cs b/Engine/Source/Programs/Horde/Horde.Server/Agents/IAgent.cs
--- a/Engine/Source/Programs/Horde/Horde.Server/Agents/IAgent.cs
+++ b/Engine/Source/Programs/Horde/Horde.Server/Agents/IAgent.cs
@@ -955,7 +955,7 @@ namespace Horde.Server.Agents
 				{
 					identifier = workspaceConfig.Identifier;
 				}
-				else if (workspaceConfig.Incremental)
+				else if (workspaceConfig.Incremental ?? false)
 				{
 					identifier = $"{streamConfig.GetEscapedName()}+{agentType.Workspace}";
 				}
@@ -965,7 +965,7 @@ namespace Horde.Server.Agents
 				}
 
 				// Create the new workspace
-				workspace = new AgentWorkspace(workspaceConfig.Cluster, workspaceConfig.UserName, identifier, workspaceConfig.Stream ?? streamConfig.Name, workspaceConfig.View, workspaceConfig.Incremental, workspaceConfig.Method);
+				workspace = new AgentWorkspace(workspaceConfig.Cluster, workspaceConfig.UserName, identifier, workspaceConfig.Stream ?? streamConfig.Name, workspaceConfig.View, workspaceConfig.Incremental ?? false, workspaceConfig.Method);
 				autoSdkConfig = GetAutoSdkConfig(workspaceConfig, streamConfig);
 
 				return true;
@@ -974,7 +974,7 @@ namespace Horde.Server.Agents
 			static AutoSdkConfig? GetAutoSdkConfig(WorkspaceConfig workspaceConfig, StreamConfig streamConfig)
 			{
 				AutoSdkConfig? autoSdkConfig = null;
-				if (workspaceConfig.UseAutoSdk)
+				if (workspaceConfig.UseAutoSdk ?? true)
 				{
 					List<string> view = new List<string>();
 					if (streamConfig.AutoSdkView != null)
diff --git a/Engine/Source/Programs/Horde/Horde.Server/Jobs/Templates/TemplateConfig.cs b/Engine/Source/Programs/Horde/Horde.Server/Jobs/Templates/TemplateConfig.cs
--- a/Engine/Source/Programs/Horde/Horde.Server/Jobs/Templates/TemplateConfig.cs
+++ b/Engine/Source/Programs/Horde/Horde.Server/Jobs/Templates/TemplateConfig.cs
@@ -5,6 +5,7 @@ using System.Collections.Generic;
 using System.ComponentModel.DataAnnotations;
 using System.Text.Json.Serialization;
 using EpicGames.Core;
+using Horde.Server.Configuration;
 using Horde.Server.Streams;
 using HordeCommon;
 using HordeCommon.Rpc.Tasks;
@@ -65,21 +66,25 @@ namespace Horde.Server.Jobs.Templates
 		/// <summary>
 		/// Default change to build at. Each object has a condition parameter which can evaluated by the server to determine which change to use.
 		/// </summary>
+		[ConfigMergeStrategy(ConfigMergeStrategy.Append)]
 		public List<ChangeQueryConfig>? DefaultChange { get; set; }
 
 		/// <summary>
 		/// Fixed arguments for the new job
 		/// </summary>
+		[ConfigMergeStrategy(ConfigMergeStrategy.Append)]
 		public List<string> Arguments { get; set; } = new List<string>();
 
 		/// <summary>
 		/// Parameters for this template
 		/// </summary>
+		[ConfigMergeStrategy(ConfigMergeStrategy.Append)]
 		public List<ParameterData> Parameters { get; set; } = new List<ParameterData>();
 
 		/// <summary>
 		/// Default settings for jobs
 		/// </summary>
+		[ConfigMergeStrategy(ConfigMergeStrategy.Recursive)]
 		public JobOptions JobOptions { get; set; } = new JobOptions();
 
 		/// <summary>
diff --git a/Engine/Source/Programs/Horde/Horde.Server/Streams/StreamConfig.cs b/Engine/Source/Programs/Horde/Horde.Server/Streams/StreamConfig.cs
--- a/Engine/Source/Programs/Horde/Horde.Server/Streams/StreamConfig.cs
+++ b/Engine/Source/Programs/Horde/Horde.Server/Streams/StreamConfig.cs
@@ -25,6 +25,10 @@ using Horde.Server.Utilities;
 using Horde.Common;
 using HordeCommon;
 using HordeCommon.Rpc.Tasks;
+using System.Reflection;
+using System.ComponentModel;
+using EpicGames.Serialization;
+using System.Diagnostics;
 
 namespace Horde.Server.Streams
 {
@@ -259,6 +263,15 @@ namespace Horde.Server.Streams
 				template.PostLoad(this);
 			}
 
+			ConfigType.MergeDefaults(AgentTypes.Select(x => (x.Key, x.Value.Base, x.Value)));
+			ConfigType.MergeDefaults(WorkspaceTypes.Select(x => (x.Key, x.Value.Base, x.Value)));
+			ConfigType.MergeDefaults(Templates.Select(x => (x.Id, x.Base, x)));
+
+			foreach (TemplateRefConfig template in Templates)
+			{
+				template.JobOptions.MergeDefaults(JobOptions);
+			}
+
 			if (Environment != null && Environment.Count > 0)
 			{
 				foreach (AgentConfig agentConfig in AgentTypes.Values)
@@ -515,10 +528,14 @@ namespace Horde.Server.Streams
 	/// </summary>
 	public class AgentConfig
 	{
+		/// <summary>
+		/// Base agent config to inherit settings from
+		/// </summary>
+		public string? Base { get; set; }
+
 		/// <summary>
 		/// Pool of agents to use for this agent type
 		/// </summary>
-		[Required]
 		public PoolId Pool { get; set; }
 
 		/// <summary>
@@ -539,6 +556,7 @@ namespace Horde.Server.Streams
 		/// <summary>
 		/// Tokens to allocate for this agent type
 		/// </summary>
+		[ConfigMergeStrategy(ConfigMergeStrategy.Append)]
 		public List<TokenConfig>? Tokens { get; set; }
 
 		/// <summary>
@@ -565,6 +583,11 @@ namespace Horde.Server.Streams
 	/// </summary>
 	public class WorkspaceConfig
 	{
+		/// <summary>
+		/// Base workspace to derive from
+		/// </summary>
+		public string? Base { get; set; }
+
 		/// <summary>
 		/// Name of the Perforce server cluster to use
 		/// </summary>
@@ -598,21 +621,23 @@ namespace Horde.Server.Streams
 		/// <summary>
 		/// Custom view for the workspace
 		/// </summary>
+		[ConfigMergeStrategy(ConfigMergeStrategy.Append)]
 		public List<string>? View { get; set; }
 
 		/// <summary>
 		/// Whether to use an incrementally synced workspace
 		/// </summary>
-		public bool Incremental { get; set; }
+		public bool? Incremental { get; set; }
 
 		/// <summary>
 		/// Whether to use the AutoSDK
 		/// </summary>
-		public bool UseAutoSdk { get; set; } = true;
+		public bool? UseAutoSdk { get; set; }
 
 		/// <summary>
 		/// View for the AutoSDK paths to sync. If null, the whole thing will be synced.
 		/// </summary>
+		[ConfigMergeStrategy(ConfigMergeStrategy.Append)]
 		public List<string>? AutoSdkView { get; set; }
 
 		/// <summary>
@@ -700,6 +725,7 @@ namespace Horde.Server.Streams
 	/// <summary>
 	/// Parameters to create a template within a stream
 	/// </summary>
+	[DebuggerDisplay("{Id}")]
 	public class TemplateRefConfig : TemplateConfig, IAclScope
 	{
 		/// <summary>
@@ -721,6 +747,11 @@ namespace Horde.Server.Streams
 		/// </summary>
 		public TemplateId Id { get; set; }
 
+		/// <summary>
+		/// Base template id to copy from
+		/// </summary>
+		public TemplateId Base { get; set; }
+
 		/// <summary>
 		/// Whether to show badges in UGS for these jobs
 		/// </summary>
@@ -763,16 +794,19 @@ namespace Horde.Server.Streams
 		/// <summary>
 		/// Schedule to execute this template
 		/// </summary>
+		[ConfigMergeStrategy(ConfigMergeStrategy.Recursive)]
 		public ScheduleConfig? Schedule { get; set; }
 
 		/// <summary>
 		/// List of chained job triggers
 		/// </summary>
+		[ConfigMergeStrategy(ConfigMergeStrategy.Append)]
 		public List<ChainedJobTemplateConfig>? ChainedJobs { get; set; }
 
 		/// <summary>
 		/// The ACL for this template
 		/// </summary>
+		[ConfigMergeStrategy(ConfigMergeStrategy.Recursive)]
 		public AclConfig? Acl { get; set; }
 
 		/// <summary>
@@ -782,14 +816,13 @@ namespace Horde.Server.Streams
 		public void PostLoad(StreamConfig streamConfig)
 		{
 			StreamConfig = streamConfig;
-			ScopeName = streamConfig.ScopeName.Append("t", Id.ToString());
-
-			JobOptions.MergeDefaults(streamConfig.JobOptions);
 
 			if (Id.IsEmpty)
 			{
 				Id = new TemplateId(StringId.Sanitize(Name));
 			}
+
+			ScopeName = streamConfig.ScopeName.Append("t", Id.ToString());
 		}
 	}
 
