Horde: Improve server responses when executing an invalid config preflight.

diff --git a/Engine/Source/Programs/Horde/Horde.Server/Server/ServerController.cs b/Engine/Source/Programs/Horde/Horde.Server/Server/ServerController.cs
--- a/Engine/Source/Programs/Horde/Horde.Server/Server/ServerController.cs
+++ b/Engine/Source/Programs/Horde/Horde.Server/Server/ServerController.cs
@@ -7,6 +7,7 @@ using System.Reflection;
 using System.Text.Json;
 using System.Threading;
 using System.Threading.Tasks;
+using EpicGames.Core;
 using EpicGames.Perforce;
 using Horde.Server.Acls;
 using Horde.Server.Agents;
@@ -138,7 +139,14 @@ namespace Horde.Server.Server
 			string cluster = request.Cluster ?? "default";
 
 			IPooledPerforceConnection perforce = await _perforceService.ConnectAsync(cluster, cancellationToken: cancellationToken);
-			DescribeRecord record = await perforce.DescribeAsync(request.ShelvedChange, cancellationToken);
+
+			PerforceResponse<DescribeRecord> describeResponse = await perforce.TryDescribeAsync(request.ShelvedChange, cancellationToken);
+			if (!describeResponse.Succeeded)
+			{
+				return BadRequest(KnownLogEvents.Horde_InvalidPreflight, "CL {Change} does not exist.", request.ShelvedChange);
+			}
+
+			DescribeRecord record = describeResponse.Data;
 
 			Dictionary<Uri, byte[]> files = new Dictionary<Uri, byte[]>();
 			foreach (DescribeFileRecord fileRecord in record.Files)
@@ -155,6 +163,11 @@ namespace Horde.Server.Server
 				files.Add(uri, printRecord.Contents);
 			}
 
+			if (files.Count == 0)
+			{
+				return BadRequest(KnownLogEvents.Horde_InvalidPreflight, "No config files found in CL {Change}.", request.ShelvedChange);
+			}
+
 			string? message = await _configService.ValidateAsync(files, cancellationToken);
 
 			PreflightConfigResponse response = new PreflightConfigResponse();
