Horde: Add preflight config check UI
#rnx
#skipci

diff --git a/Engine/Source/Programs/Horde/HordeDashboard/src/App.tsx b/Engine/Source/Programs/Horde/HordeDashboard/src/App.tsx
--- a/Engine/Source/Programs/Horde/HordeDashboard/src/App.tsx
+++ b/Engine/Source/Programs/Horde/HordeDashboard/src/App.tsx
@@ -33,6 +33,7 @@ import hordePlugins from './Plugins';
 import { modeColors, preloadFonts } from './styles/Styles';
 import { StepIssueReportTest } from './components/test/IssueStepReport';
 import { JobRedirector } from './components/JobRedirector';
+import { PreflightConfigView } from './components/preflights/PreflightConfigCheck';
 
 
 let router: any;
@@ -133,7 +134,7 @@ const Main: React.FC = () => {
                { path: "index", element: (dashboard.user?.dashboardFeatures?.showLandingPage === true) ? <DocView /> : <UserHomeView /> },
                { path: "project/:projectId", element: <ProjectHome /> },
                { path: "pools", element: <PoolView /> },
-               { path: "job/:jobId", element: <JobDetailViewV2 /> },               
+               { path: "job/:jobId", element: <JobDetailViewV2 /> },
                { path: "job", element: <JobRedirector /> },
                { path: "log/:logId", element: <LogView /> },
                { path: "testreport/:testdataId", element: <TestReportView /> },
@@ -142,6 +143,7 @@ const Main: React.FC = () => {
                { path: "admin/token", element: <AdminToken /> },
                { path: "reports/utilization", element: <UtilizationReportView /> },
                { path: "preflight", element: <PreflightRedirector /> },
+               { path: "preflightconfig", element: <PreflightConfigView /> },
                { path: "dashboard", element: <DashboardView /> },
                { path: "perforce/servers", element: <PerforceServerView /> },
                { path: "notices", element: <NoticeView /> },
diff --git a/Engine/Source/Programs/Horde/HordeDashboard/src/backend/Api.ts b/Engine/Source/Programs/Horde/HordeDashboard/src/backend/Api.ts
--- a/Engine/Source/Programs/Horde/HordeDashboard/src/backend/Api.ts
+++ b/Engine/Source/Programs/Horde/HordeDashboard/src/backend/Api.ts
@@ -394,7 +394,7 @@ export type GetDashboardConfigResponse = {
 	/** Help slack channel that users can use for issues */
 	helpSlackChannel?: string;
 
-	/** Device problem cooldown in minutes */ 
+	/** Device problem cooldown in minutes */
 	deviceProblemCooldownMinutes?: number
 
 }
@@ -3265,6 +3265,31 @@ export type GetPerforceServerStatusResponse = {
 	detail: string;
 }
 
+/** Request to validate server configuration with the given files replacing their checked-in counterparts. */
+export type PreflightConfigRequest = {
+
+	/**  Change to test	*/
+	shelvedChange: number;
+
+	/**  Perforce cluster to retrieve from */
+	cluster?: string;
+
+}
+
+/**  Response from validating config files */
+export type PreflightConfigResponse = {
+
+	/** Whether the files were validated successfully */
+	result: boolean;
+
+	/** Output message from validation */
+	message?: string;
+
+	/** Detailed response */
+	detail?: string;
+}
+
+
 /** Get object response which describes a device platform */
 export type GetDevicePlatformResponse = {
 
@@ -4751,7 +4776,7 @@ export type UpdateBisectTaskRequest = {
 	encludeChanges?: number[];
 
 	/// List of jobs to include in the search.	
-	includeJobs?: string[]; 
+	includeJobs?: string[];
 
 	/// List of jobs to exclude from the search.	
 	excludeJobs?: string[];
diff --git a/Engine/Source/Programs/Horde/HordeDashboard/src/backend/Backend.ts b/Engine/Source/Programs/Horde/HordeDashboard/src/backend/Backend.ts
--- a/Engine/Source/Programs/Horde/HordeDashboard/src/backend/Backend.ts
+++ b/Engine/Source/Programs/Horde/HordeDashboard/src/backend/Backend.ts
@@ -1,7 +1,7 @@
 // Copyright Epic Games, Inc. All Rights Reserved.
 
 import templateCache from '../backend/TemplateCache';
-import { AgentData, AgentQuery, ArtifactData, AuditLogEntry, AuditLogQuery, BatchUpdatePoolRequest, ChangeSummaryData, CreateBisectTaskRequest, CreateBisectTaskResponse, CreateDeviceRequest, CreateDeviceResponse, CreateExternalIssueRequest, CreateExternalIssueResponse, CreateJobRequest, CreateJobResponse, CreateNoticeRequest, CreatePoolRequest, CreateSoftwareResponse, CreateSubscriptionRequest, CreateSubscriptionResponse, CreateZipRequest, DashboardPreference, DevicePoolTelemetryQuery, DeviceTelemetryQuery, EventData, FindArtifactsResponse, FindIssueResponse, FindJobTimingsResponse, GetAgentSoftwareChannelResponse, GetArtifactDirectoryResponse, GetArtifactZipRequest, GetBisectTaskResponse, GetDashboardConfigResponse, GetDevicePlatformResponse, GetDevicePoolResponse, GetDevicePoolTelemetryResponse, GetDeviceReservationResponse, GetDeviceResponse, GetDeviceTelemetryResponse, GetExternalIssueProjectResponse, GetExternalIssueResponse, GetGraphResponse, GetIssueStreamResponse, GetJobsTabResponse, GetJobStepRefResponse, GetJobStepTraceResponse, GetJobTimingResponse, GetLogEventResponse, GetLogFileResponse, GetNoticeResponse, GetNotificationResponse, GetPerforceServerStatusResponse, GetPoolResponse, GetServerInfoResponse, GetServerSettingsResponse, GetSoftwareResponse, GetSubscriptionResponse, GetTestDataDetailsResponse, GetTestDataRefResponse, GetTestMetaResponse, GetTestResponse, GetTestsRequest, GetTestStreamResponse, GetToolSummaryResponse, GetUserResponse, GetUtilizationTelemetryResponse, GlobalConfig, IssueData, IssueQuery, IssueQueryV2, JobData, JobQuery, JobsTabColumnType, JobStepOutcome, JobStreamQuery, JobTimingsQuery, LeaseData, LogData, LogLineData, PoolData, ProjectData, ScheduleData, ScheduleQuery, SearchLogFileResponse, ServerUpdateResponse, SessionData, StreamData, TabType, TestData, UpdateAgentRequest, UpdateBisectTaskRequest, UpdateDeviceRequest, UpdateGlobalConfigRequest, UpdateIssueRequest, UpdateJobRequest, UpdateLeaseRequest, UpdateNoticeRequest, UpdateNotificationsRequest, UpdatePoolRequest, UpdateServerSettingsRequest, UpdateStepRequest, UpdateStepResponse, UpdateTemplateRefRequest, UpdateUserRequest, UsersQuery } from './Api';
+import { AgentData, AgentQuery, ArtifactData, AuditLogEntry, AuditLogQuery, BatchUpdatePoolRequest, ChangeSummaryData, CreateBisectTaskRequest, CreateBisectTaskResponse, CreateDeviceRequest, CreateDeviceResponse, CreateExternalIssueRequest, CreateExternalIssueResponse, CreateJobRequest, CreateJobResponse, CreateNoticeRequest, CreatePoolRequest, CreateSoftwareResponse, CreateSubscriptionRequest, CreateSubscriptionResponse, CreateZipRequest, DashboardPreference, DevicePoolTelemetryQuery, DeviceTelemetryQuery, EventData, FindArtifactsResponse, FindIssueResponse, FindJobTimingsResponse, GetAgentSoftwareChannelResponse, GetArtifactDirectoryResponse, GetArtifactZipRequest, GetBisectTaskResponse, GetDashboardConfigResponse, GetDevicePlatformResponse, GetDevicePoolResponse, GetDevicePoolTelemetryResponse, GetDeviceReservationResponse, GetDeviceResponse, GetDeviceTelemetryResponse, GetExternalIssueProjectResponse, GetExternalIssueResponse, GetGraphResponse, GetIssueStreamResponse, GetJobsTabResponse, GetJobStepRefResponse, GetJobStepTraceResponse, GetJobTimingResponse, GetLogEventResponse, GetLogFileResponse, GetNoticeResponse, GetNotificationResponse, GetPerforceServerStatusResponse, GetPoolResponse, GetServerInfoResponse, GetServerSettingsResponse, GetSoftwareResponse, GetSubscriptionResponse, GetTestDataDetailsResponse, GetTestDataRefResponse, GetTestMetaResponse, GetTestResponse, GetTestsRequest, GetTestStreamResponse, GetToolSummaryResponse, GetUserResponse, GetUtilizationTelemetryResponse, GlobalConfig, IssueData, IssueQuery, IssueQueryV2, JobData, JobQuery, JobsTabColumnType, JobStepOutcome, JobStreamQuery, JobTimingsQuery, LeaseData, LogData, LogLineData, PoolData, PreflightConfigResponse, ProjectData, ScheduleData, ScheduleQuery, SearchLogFileResponse, ServerUpdateResponse, SessionData, StreamData, TabType, TestData, UpdateAgentRequest, UpdateBisectTaskRequest, UpdateDeviceRequest, UpdateGlobalConfigRequest, UpdateIssueRequest, UpdateJobRequest, UpdateLeaseRequest, UpdateNoticeRequest, UpdateNotificationsRequest, UpdatePoolRequest, UpdateServerSettingsRequest, UpdateStepRequest, UpdateStepResponse, UpdateTemplateRefRequest, UpdateUserRequest, UsersQuery } from './Api';
 import dashboard from './Dashboard';
 import { ChallengeStatus, Fetch } from './Fetch';
 import graphCache, { GraphQuery } from './GraphCache';
@@ -1686,6 +1686,18 @@ export class Backend {
         });
     }
 
+    checkPreflightConfig(shelvedChange: number): Promise<PreflightConfigResponse> {
+
+        return new Promise<PreflightConfigResponse>((resolve, reject) => {
+            this.backend.post(`/api/v1/server/preflightconfig`, {shelvedChange: shelvedChange}).then((value) => {
+                resolve(value.data as PreflightConfigResponse);                
+            }).catch(reason => {
+                reject(reason);
+            });
+        });
+
+    }
+
     private async update() {
 
         if (this.updateID === "updating") {
diff --git a/Engine/Source/Programs/Horde/HordeDashboard/src/components/preflights/PreflightConfigCheck.tsx b/Engine/Source/Programs/Horde/HordeDashboard/src/components/preflights/PreflightConfigCheck.tsx
new file mode 100644
--- /dev/null
+++ b/Engine/Source/Programs/Horde/HordeDashboard/src/components/preflights/PreflightConfigCheck.tsx
@@ -0,0 +1,121 @@
+import { FontIcon, MaskedTextField, PrimaryButton, Spinner, SpinnerSize, Stack, Text, TextField } from "@fluentui/react";
+import { useWindowSize } from "../../base/utilities/hooks";
+import { hordeClasses, modeColors } from "../../styles/Styles";
+import { Breadcrumbs } from "../Breadcrumbs";
+import { TopNav } from "../TopNav";
+import { useState } from "react";
+import { useNavigate } from "react-router-dom";
+import backend from "../../backend";
+import dashboard, { StatusColor } from "../../backend/Dashboard";
+
+const PreflightConfigPanel: React.FC = () => {
+
+   const search = new URL(window.location.toString()).searchParams;
+   const shelvedChange = search.get("shelvedchange") ? search.get("shelvedchange")! : undefined;
+
+   const navigate = useNavigate();
+   const [state, setState] = useState<{ initialCL?: string, submitting?: boolean, success?: boolean, message?: string }>({initialCL: shelvedChange});
+
+   const maskFormat: { [key: string]: RegExp } = {
+      '*': /[0-9]/,
+   };
+
+   const checkPreflight = async (preflightCL?: string) => {
+
+      if (!preflightCL) {
+
+         if (!shelvedChange) {
+            return;
+         }
+   
+         preflightCL = shelvedChange;
+      }
+
+      setState({ submitting: true });
+
+      try {
+         const response = await backend.checkPreflightConfig(parseInt(preflightCL));
+         setState({ submitting: false, success: response.result, message: response.message });
+      } catch (error) {
+         setState({ submitting: false, success: false, message: error as string });
+      }
+
+   }
+
+   if (state.initialCL) {
+      checkPreflight(state.initialCL);
+      return null;
+   }
+
+
+   return (<Stack>
+      <Stack styles={{ root: { paddingTop: 18, paddingLeft: 12, paddingRight: 12, width: "100%" } }} >
+         <Stack tokens={{ childrenGap: 12 }} >
+            <Stack style={{ width: 800, paddingLeft: 12 }}>
+               <Stack horizontal verticalAlign="center" tokens={{ childrenGap: 24 }}>
+                  <Stack horizontal verticalAlign="center" tokens={{ childrenGap: 18 }}>
+                     <Text styles={{ root: { fontFamily: "Horde Open Sans SemiBold" } }}>Shelved Change</Text>
+                     <MaskedTextField placeholder="Shelved Change" mask="***********" maskFormat={maskFormat} maskChar="" value={shelvedChange} onChange={(ev, newValue) => {
+                        ev.preventDefault();
+
+                        if (!newValue) {
+                           navigate("/preflightconfig", { replace: true });
+                        }
+
+                        if (!isNaN(parseInt(newValue!))) {
+                           navigate(`?shelvedchange=${newValue}`, { replace: true });
+                        }
+
+                     }} />
+                  </Stack>
+                  <Stack grow />
+                  {!!state.submitting && <Stack>
+                     <Spinner size={SpinnerSize.large} />
+                  </Stack>}
+
+                  {state.success === true && <Stack>
+                     <FontIcon style={{ color: dashboard.getStatusColors().get(StatusColor.Success)!, fontSize: 24 }} iconName="Tick" />
+                  </Stack>}
+
+                  {state.success === false && <Stack>
+                     <FontIcon style={{ color: dashboard.getStatusColors().get(StatusColor.Failure)!, fontSize: 24 }} iconName="Cross" />
+                  </Stack>}
+
+                  <Stack>
+                     <PrimaryButton disabled={!shelvedChange || state.submitting} styles={{ root: { fontFamily: "Horde Open Sans SemiBold" } }} onClick={async () => {
+                        checkPreflight(shelvedChange);
+                     }}>Check</PrimaryButton>
+                  </Stack>
+               </Stack>
+               {!!state.message && <Stack style={{ paddingTop: 24, paddingRight: 2 }}>
+                  <TextField readOnly multiline resizable={false} label="Error Message" autoAdjustHeight={false} defaultValue={state.message} style={{ whiteSpace: "pre-wrap", height: 500 }}> </TextField>
+               </Stack>}
+            </Stack>
+         </Stack>
+      </Stack>
+   </Stack>);
+}
+
+
+export const PreflightConfigView: React.FC = () => {
+
+   const windowSize = useWindowSize();
+   const vw = Math.max(document.documentElement.clientWidth, window.innerWidth || 0);
+
+   return <Stack className={hordeClasses.horde}>
+      <TopNav />
+      <Breadcrumbs items={[{ text: 'Preflight Configuration' }]} />
+      <Stack horizontal>
+         <div key={`windowsize_noticeview_${windowSize.width}_${windowSize.height}`} style={{ width: vw / 2 - (1440 / 2), flexShrink: 0, backgroundColor: modeColors.background }} />
+         <Stack tokens={{ childrenGap: 0 }} styles={{ root: { backgroundColor: modeColors.background, width: "100%" } }}>
+            <Stack style={{ maxWidth: 1440, paddingTop: 6, marginLeft: 4, height: 'calc(100vh - 8px)' }}>
+               <Stack horizontal className={hordeClasses.raised}>
+                  <Stack style={{ width: "100%", height: 'calc(100vh - 228px)' }} tokens={{ childrenGap: 18 }}>
+                     <PreflightConfigPanel />
+                  </Stack>
+               </Stack>
+            </Stack>
+         </Stack>
+      </Stack>
+   </Stack>
+};
