Horde: Missing file.

diff --git a/Engine/Source/Programs/Horde/Horde.Server/Configuration/ConfigType.cs b/Engine/Source/Programs/Horde/Horde.Server/Configuration/ConfigType.cs
--- a/Engine/Source/Programs/Horde/Horde.Server/Configuration/ConfigType.cs
+++ b/Engine/Source/Programs/Horde/Horde.Server/Configuration/ConfigType.cs
@@ -12,12 +12,11 @@ using System.Text.Json.Serialization;
 using System.Threading;
 using System.Threading.Tasks;
 using EpicGames.Core;
-using Microsoft.Extensions.Logging;
-
-using JsonObject = System.Text.Json.Nodes.JsonObject;
 
 namespace Horde.Server.Configuration
 {
+	using JsonObject = System.Text.Json.Nodes.JsonObject;
+
 	/// <summary>
 	/// Attribute used to mark <see cref="Uri"/> properties that include other config files
 	/// </summary>
@@ -50,6 +49,44 @@ namespace Horde.Server.Configuration
 	{
 	}
 
+	/// <summary>
+	/// Possible methods for merging config values
+	/// </summary>
+	public enum ConfigMergeStrategy
+	{
+		/// <summary>
+		/// Default strategy; replace with the base value if the current value is null
+		/// </summary>
+		Default,
+
+		/// <summary>
+		/// Append the contents of this list to the base list
+		/// </summary>
+		Append,
+
+		/// <summary>
+		/// Recursively merge object properties
+		/// </summary>
+		Recursive,
+	}
+
+	/// <summary>
+	/// Attribute used to mark object properties whose child properties should be merged individually
+	/// </summary>
+	[AttributeUsage(AttributeTargets.Property)]
+	public sealed class ConfigMergeStrategyAttribute : Attribute
+	{
+		/// <summary>
+		/// Strategy for merging this property
+		/// </summary>
+		public ConfigMergeStrategy Strategy { get; }
+
+		/// <summary>
+		/// Constructor
+		/// </summary>
+		public ConfigMergeStrategyAttribute(ConfigMergeStrategy strategy) => Strategy = strategy;
+	}
+
 	/// <summary>
 	/// Exception thrown when reading config files
 	/// </summary>
@@ -116,6 +153,132 @@ namespace Horde.Server.Configuration
 			return value;
 		}
 
+		/// <summary>
+		/// Merges any unassigned properties from the source object to the target
+		/// </summary>
+		/// <param name="objects">Set of objects to use for merging</param>
+		public static void MergeDefaults<TKey, TValue>(IEnumerable<(TKey Key, TKey? BaseKey, TValue Value)> objects) 
+			where TKey : notnull
+			where TValue : class
+		{
+			// Convert the objects to a dictionary
+			List<(TKey Key, TKey BaseKey, TValue Value)> remainingObjects = new List<(TKey Key, TKey BaseKey, TValue Value)>();
+
+			// Find all the objects with no base
+			Dictionary<TKey, TValue> handledValues = new Dictionary<TKey, TValue>();
+			foreach ((TKey key, TKey? baseKey, TValue value) in objects)
+			{
+				if (baseKey == null || Object.Equals(baseKey, default(TKey)))
+				{
+					handledValues.Add(key, value);
+				}
+				else
+				{
+					remainingObjects.Add((key, baseKey, value));
+				}
+			}
+
+			// Iteratively merge objects with their base
+			for (int lastRemainingObjectCount = 0; remainingObjects.Count != lastRemainingObjectCount; )
+			{
+				lastRemainingObjectCount = remainingObjects.Count;
+				for (int Idx = remainingObjects.Count - 1; Idx >= 0; Idx--)
+				{
+					(TKey key, TKey baseKey, TValue value) = remainingObjects[Idx];
+					if (handledValues.TryGetValue(baseKey, out TValue? baseValue))
+					{
+						MergeDefaults(value, baseValue);
+						handledValues.Add(key, value);
+						remainingObjects.RemoveAt(Idx);
+					}
+				}
+			}
+
+			// Check we were able to merge everything
+			if (remainingObjects.Count > 0)
+			{
+				HashSet<TKey> validKeys = new HashSet<TKey>(objects.Select(x => x.Key));
+				foreach ((TKey key, TKey baseKey, TValue value) in remainingObjects)
+				{
+					if (!validKeys.Contains(baseKey!))
+					{
+						throw new Exception($"{key} has invalid/missing base {baseKey}");
+					}
+				}
+
+				List<TKey> circularKeys = objects.Where(x => !handledValues.ContainsKey(x.Key)).Select(x => x.Key).ToList();
+				if (circularKeys.Count == 1)
+				{
+					throw new Exception($"{circularKeys[0]} has a circular dependency on itself");
+				}
+				else
+				{
+					throw new Exception($"{StringUtils.FormatList(circularKeys.Select(x => x.ToString() ?? "(unknown)"))} have circular dependencies.");
+				}
+			}
+		}
+
+		/// <summary>
+		/// Merges any unassigned properties from the source object to the target
+		/// </summary>
+		/// <param name="target">Target object to merge into</param>
+		/// <param name="source">Source object to merge from</param>
+		public static void MergeDefaults<T>(T target, T source) where T : class
+		{
+			MergeDefaults(typeof(T), target, source);
+		}
+
+		/// <summary>
+		/// Merges any unassigned properties from the source object to the target
+		/// </summary>
+		/// <param name="type">Type of the referenced object</param>
+		/// <param name="target">Target object to merge into</param>
+		/// <param name="source">Source object to merge from</param>
+		public static void MergeDefaults(Type type, object target, object source)
+		{
+			foreach (PropertyInfo propertyInfo in type.GetProperties(BindingFlags.Instance | BindingFlags.Public))
+			{
+				object? targetValue = propertyInfo.GetValue(target);
+				object? sourceValue = propertyInfo.GetValue(source);
+
+				ConfigMergeStrategy strategy = propertyInfo.GetCustomAttribute<ConfigMergeStrategyAttribute>()?.Strategy ?? ConfigMergeStrategy.Default;
+				switch(strategy)
+				{
+					case ConfigMergeStrategy.Default:
+						if (targetValue == null)
+						{
+							propertyInfo.SetValue(target, sourceValue);
+						}
+						break;
+					case ConfigMergeStrategy.Append:
+						if (targetValue == null)
+						{
+							propertyInfo.SetValue(target, sourceValue);
+						}
+						else if (sourceValue != null)
+						{
+							IList sourceList = (IList)sourceValue;
+							IList targetList = (IList)targetValue;
+							for (int idx = 0; idx < sourceList.Count; idx++)
+							{
+								targetList.Insert(idx, sourceList[idx]);
+							}
+						}
+						break;
+					case ConfigMergeStrategy.Recursive:
+						if (targetValue == null)
+						{
+							propertyInfo.SetValue(target, sourceValue);
+						}
+						else if (sourceValue != null)
+						{
+							MergeDefaults(propertyInfo.PropertyType, targetValue, sourceValue);
+						}
+						break;
+				}
+			}
+		}
+
 		/// <summary>
 		/// Reads an object from a particular URL
 		/// </summary>
