Horde: Remove default report time of midnight; config file only appends to this list, not replacing it.

diff --git a/Engine/Source/Programs/Horde/Horde.Server/Issues/WorkflowConfig.cs b/Engine/Source/Programs/Horde/Horde.Server/Issues/WorkflowConfig.cs
--- a/Engine/Source/Programs/Horde/Horde.Server/Issues/WorkflowConfig.cs
+++ b/Engine/Source/Programs/Horde/Horde.Server/Issues/WorkflowConfig.cs
@@ -21,7 +21,7 @@ namespace Horde.Server.Issues
 		/// <summary>
 		/// Times of day at which to send a report
 		/// </summary>
-		public List<TimeSpan> ReportTimes { get; set; } = new List<TimeSpan> { TimeSpan.Zero };
+		public List<TimeSpan> ReportTimes { get; set; } = new List<TimeSpan>();
 
 		/// <summary>
 		/// Name of the tab to post summary data to
