Horde: Introduce a one second delay between messages in reports to prevent Slack rate limiting kicking in.

diff --git a/Engine/Source/Programs/Horde/Horde.Server/Notifications/Sinks/SlackNotificationSink.cs b/Engine/Source/Programs/Horde/Horde.Server/Notifications/Sinks/SlackNotificationSink.cs
--- a/Engine/Source/Programs/Horde/Horde.Server/Notifications/Sinks/SlackNotificationSink.cs
+++ b/Engine/Source/Programs/Horde/Horde.Server/Notifications/Sinks/SlackNotificationSink.cs
@@ -1680,6 +1680,8 @@ namespace Horde.Server.Notifications.Sinks
 			SlackMessage headerMessage = new SlackMessage();
 			headerMessage.AddHeader($"Summary for {streamConfig.Name}");
 
+			TimeSpan rateLimitDelay = TimeSpan.FromSeconds(1.0);
+
 			SlackMessageId? messageId = await SendMessageAsync(channel, headerMessage);
 			if (messageId != null)
 			{
@@ -1695,6 +1697,7 @@ namespace Horde.Server.Notifications.Sinks
 
 				for (int idx = 0; idx < state.Blocks.Count; idx++)
 				{
+					await Task.Delay(rateLimitDelay);
 					string blockEventId = GetReportBlockEventId(messageId.Ts, idx);
 					await UpdateReportBlockAsync(channel, blockEventId, time, streamConfig, state.Blocks[idx].TemplateId, issuesByBlock[idx], report.TriageChannel, state.Blocks[idx].TemplateHeader);
 				}
@@ -1705,6 +1708,8 @@ namespace Horde.Server.Notifications.Sinks
 					string header = $"*{totalPct:0.0}%* of build steps ({report.WorkflowStats.NumPassingSteps:n0}/{report.WorkflowStats.NumSteps:n0}) succeeded since last status update.";
 					await SendMessageAsync(channel, header);
 				}
+
+				await Task.Delay(rateLimitDelay);
 			}
 		}
 
