Horde: Add preflight config modal and hook up to user home
#skipci
#rnx

diff --git a/Engine/Source/Programs/Horde/HordeDashboard/src/App.tsx b/Engine/Source/Programs/Horde/HordeDashboard/src/App.tsx
--- a/Engine/Source/Programs/Horde/HordeDashboard/src/App.tsx
+++ b/Engine/Source/Programs/Horde/HordeDashboard/src/App.tsx
@@ -3,7 +3,8 @@
 import { Image, Spinner, SpinnerSize, Stack, Text } from '@fluentui/react';
 import React, { useState } from 'react';
 import { useDarkreader } from 'react-darkreader';
-import { createBrowserRouter, Navigate, Outlet, RouteObject, RouterProvider } from 'react-router-dom';
+import { Navigate, Outlet, RouteObject, RouterProvider, createBrowserRouter } from 'react-router-dom';
+import hordePlugins from './Plugins';
 import backend from './backend';
 import { DashboardPreference } from './backend/Api';
 import { getSiteConfig } from './backend/Config';
@@ -15,9 +16,8 @@ import { AutomationView } from './components/AutomationView';
 import { DashboardView } from './components/DashboardView';
 import { DebugView } from './components/DebugView';
 import { DeviceView } from './components/DeviceView';
-import { DocView } from './components/docs/DocView';
 import { ErrorDialog, ErrorHandler } from './components/ErrorHandler';
-import { JobDetailViewV2 } from './components/jobDetailsV2/JobDetailViewV2';
+import { JobRedirector } from './components/JobRedirector';
 import { LogView } from './components/LogView';
 import { NoticeView } from './components/NoticeView';
 import { PerforceServerView } from './components/PerforceView';
@@ -29,11 +29,11 @@ import { TestReportView } from './components/TestReportView';
 import { ToolView } from './components/ToolView';
 import { UserHomeView } from './components/UserHome';
 import { UtilizationReportView } from './components/UtilizationReportView';
-import hordePlugins from './Plugins';
-import { modeColors, preloadFonts } from './styles/Styles';
+import { DocView } from './components/docs/DocView';
+import { JobDetailViewV2 } from './components/jobDetailsV2/JobDetailViewV2';
+import { PreflightConfigRedirector } from './components/preflights/PreflightConfigCheckRedirector';
 import { StepIssueReportTest } from './components/test/IssueStepReport';
-import { JobRedirector } from './components/JobRedirector';
-import { PreflightConfigView } from './components/preflights/PreflightConfigCheck';
+import { modeColors, preloadFonts } from './styles/Styles';
 
 
 let router: any;
@@ -143,7 +143,7 @@ const Main: React.FC = () => {
                { path: "admin/token", element: <AdminToken /> },
                { path: "reports/utilization", element: <UtilizationReportView /> },
                { path: "preflight", element: <PreflightRedirector /> },
-               { path: "preflightconfig", element: <PreflightConfigView /> },
+               { path: "preflightconfig", element: <PreflightConfigRedirector /> },
                { path: "dashboard", element: <DashboardView /> },
                { path: "perforce/servers", element: <PerforceServerView /> },
                { path: "notices", element: <NoticeView /> },
diff --git a/Engine/Source/Programs/Horde/HordeDashboard/src/components/UserHome.tsx b/Engine/Source/Programs/Horde/HordeDashboard/src/components/UserHome.tsx
--- a/Engine/Source/Programs/Horde/HordeDashboard/src/components/UserHome.tsx
+++ b/Engine/Source/Programs/Horde/HordeDashboard/src/components/UserHome.tsx
@@ -5,7 +5,7 @@ import { action, makeObservable, observable } from 'mobx';
 import { observer } from 'mobx-react-lite';
 import moment from "moment-timezone";
 import { default as React, useEffect, useState } from 'react';
-import { Link, useLocation } from "react-router-dom";
+import { Link, useLocation, useNavigate } from "react-router-dom";
 import backend, { useBackend } from "../backend";
 import { GetIssueResponse, GetStepResponse, JobData, JobQuery, JobState, JobStepOutcome, LabelData, LabelOutcome, LabelState, ProjectData, StepData, StreamData } from "../backend/Api";
 import dashboard, { StatusColor } from "../backend/Dashboard";
@@ -21,6 +21,7 @@ import { useQuery } from './JobDetailCommon';
 import { JobOperationsContextMenu } from "./JobOperationsContextMenu";
 import { IssueStatusIcon, StepStatusIcon } from "./StatusIcon";
 import { TopNav } from './TopNav';
+import { PreflightConfigModal } from './preflights/PreflightConfigCheck';
 
 
 type JobItem = {
@@ -1059,6 +1060,12 @@ const JobsPanel: React.FC<{ includeOtherPreflights: boolean }> = observer(({ inc
 });
 
 const UserHomeViewInner: React.FC = () => {
+      
+   const search = new URLSearchParams(window.location.search);
+   const change = !search.get("preflightconfig") ? "" : search.get("preflightconfig")!;
+
+   const navigate = useNavigate();
+   
 
    useEffect(() => {
 
@@ -1071,6 +1078,7 @@ const UserHomeViewInner: React.FC = () => {
    handler.startPolling();
 
    return <Stack tokens={{ childrenGap: 0 }} styles={{ root: { backgroundColor: modeColors.background, margin: 0, paddingTop: 8 } }}>
+      {(!!change || !!search.has("preflightconfig")) && <PreflightConfigModal onClose={() => {navigate("/index", {replace: true})}}/>}
       <Stack style={{ padding: 0 }} className={detailClasses.detailsRow}>
          <FocusZone direction={FocusZoneDirection.vertical} style={{ padding: 0 }}>
             <div className={detailClasses.container} style={{ width: "100%", height: 'calc(100vh - 208px)', position: 'relative' }} data-is-scrollable={true}>
diff --git a/Engine/Source/Programs/Horde/HordeDashboard/src/components/preflights/PreflightConfigCheck.tsx b/Engine/Source/Programs/Horde/HordeDashboard/src/components/preflights/PreflightConfigCheck.tsx
--- a/Engine/Source/Programs/Horde/HordeDashboard/src/components/preflights/PreflightConfigCheck.tsx
+++ b/Engine/Source/Programs/Horde/HordeDashboard/src/components/preflights/PreflightConfigCheck.tsx
@@ -1,4 +1,4 @@
-import { FontIcon, MaskedTextField, PrimaryButton, Spinner, SpinnerSize, Stack, Text, TextField } from "@fluentui/react";
+import { FontIcon, IconButton, MaskedTextField, Modal, PrimaryButton, Spinner, SpinnerSize, Stack, Text, TextField } from "@fluentui/react";
 import { useWindowSize } from "../../base/utilities/hooks";
 import { hordeClasses, modeColors } from "../../styles/Styles";
 import { Breadcrumbs } from "../Breadcrumbs";
@@ -11,10 +11,10 @@ import dashboard, { StatusColor } from "../../backend/Dashboard";
 const PreflightConfigPanel: React.FC = () => {
 
    const search = new URL(window.location.toString()).searchParams;
-   const shelvedChange = search.get("shelvedchange") ? search.get("shelvedchange")! : undefined;
+   const shelvedChange = search.get("preflightconfig") ? search.get("preflightconfig")! : undefined;
 
    const navigate = useNavigate();
-   const [state, setState] = useState<{ initialCL?: string, submitting?: boolean, success?: boolean, message?: string }>({initialCL: shelvedChange});
+   const [state, setState] = useState<{ initialCL?: string, submitting?: boolean, success?: boolean, message?: string }>({ initialCL: shelvedChange });
 
    const maskFormat: { [key: string]: RegExp } = {
       '*': /[0-9]/,
@@ -27,7 +27,7 @@ const PreflightConfigPanel: React.FC = () => {
          if (!shelvedChange) {
             return;
          }
-   
+
          preflightCL = shelvedChange;
       }
 
@@ -49,7 +49,7 @@ const PreflightConfigPanel: React.FC = () => {
 
 
    return (<Stack>
-      <Stack styles={{ root: { paddingTop: 18, paddingLeft: 12, paddingRight: 12, width: "100%" } }} >
+      <Stack styles={{ root: { paddingTop: 18, paddingLeft: 0, paddingRight: 0, width: "100%" } }} >
          <Stack tokens={{ childrenGap: 12 }} >
             <Stack style={{ width: 800, paddingLeft: 12 }}>
                <Stack horizontal verticalAlign="center" tokens={{ childrenGap: 24 }}>
@@ -59,11 +59,11 @@ const PreflightConfigPanel: React.FC = () => {
                         ev.preventDefault();
 
                         if (!newValue) {
-                           navigate("/preflightconfig", { replace: true });
+                           navigate("/index?preflightconfig", { replace: true });
                         }
 
                         if (!isNaN(parseInt(newValue!))) {
-                           navigate(`?shelvedchange=${newValue}`, { replace: true });
+                           navigate(`/index?preflightconfig=${newValue}`, { replace: true });
                         }
 
                      }} />
@@ -87,35 +87,46 @@ const PreflightConfigPanel: React.FC = () => {
                      }}>Check</PrimaryButton>
                   </Stack>
                </Stack>
-               {!!state.message && <Stack style={{ paddingTop: 24, paddingRight: 2 }}>
-                  <TextField readOnly multiline resizable={false} label="Error Message" autoAdjustHeight={false} defaultValue={state.message} style={{ whiteSpace: "pre-wrap", height: 500 }}> </TextField>
+               {!!state.message && <Stack style={{ paddingTop: 18, paddingRight: 2 }}>
+                  <Stack style={{ paddingBottom: 12 }}>
+                     <Text styles={{ root: { fontFamily: "Horde Open Sans SemiBold" } }}>Error Message</Text>
+                  </Stack>
+                  <Stack style={{ border: "1px solid #CDCBC9" }}>
+                     <Text style={{ whiteSpace: "pre-wrap", height: "472px", overflowY: "auto", padding: 18 }}> {state.message}</Text>
+                  </Stack>
                </Stack>}
             </Stack>
          </Stack>
       </Stack>
-   </Stack>);
+   </Stack>)
 }
 
 
-export const PreflightConfigView: React.FC = () => {
-
-   const windowSize = useWindowSize();
-   const vw = Math.max(document.documentElement.clientWidth, window.innerWidth || 0);
-
-   return <Stack className={hordeClasses.horde}>
-      <TopNav />
-      <Breadcrumbs items={[{ text: 'Preflight Configuration' }]} />
-      <Stack horizontal>
-         <div key={`windowsize_noticeview_${windowSize.width}_${windowSize.height}`} style={{ width: vw / 2 - (1440 / 2), flexShrink: 0, backgroundColor: modeColors.background }} />
-         <Stack tokens={{ childrenGap: 0 }} styles={{ root: { backgroundColor: modeColors.background, width: "100%" } }}>
-            <Stack style={{ maxWidth: 1440, paddingTop: 6, marginLeft: 4, height: 'calc(100vh - 8px)' }}>
-               <Stack horizontal className={hordeClasses.raised}>
-                  <Stack style={{ width: "100%", height: 'calc(100vh - 228px)' }} tokens={{ childrenGap: 18 }}>
+export const PreflightConfigModal: React.FC<{ onClose: () => void }> = ({ onClose }) => {
+
+   return <Stack>
+      <Modal isOpen={true} isBlocking={true} topOffsetFixed={true} styles={{ main: { padding: 8, width: 880, height: 720, hasBeenOpened: false, top: "80px", position: "absolute" } }} onDismiss={() => onClose()} className={hordeClasses.modal}>
+         <Stack className="horde-no-darktheme" styles={{ root: { paddingTop: 10, paddingRight: 0 } }}>
+            <Stack style={{ paddingLeft: 24, paddingRight: 12 }}>
+               <Stack tokens={{ childrenGap: 12 }} style={{ height: 700 }}>
+                  <Stack horizontal verticalAlign="start">
+                     <Stack style={{ paddingTop: 3 }}>
+                        <Text variant="mediumPlus" styles={{ root: { fontFamily: "Horde Open Sans SemiBold" } }}>Preflight Configuration</Text>
+                     </Stack>
+                     <Stack grow />
+                     <Stack horizontalAlign="end" >
+                        <IconButton
+                           iconProps={{ iconName: 'Cancel' }}
+                           onClick={() => { onClose() }}
+                        />
+                     </Stack>
+                  </Stack>
+                  <Stack styles={{ root: { paddingLeft: 4, paddingRight: 0, paddingTop: 8, paddingBottom: 4 } }}>
                      <PreflightConfigPanel />
                   </Stack>
                </Stack>
             </Stack>
          </Stack>
-      </Stack>
-   </Stack>
+      </Modal>
+   </Stack>;
 };
diff --git a/Engine/Source/Programs/Horde/HordeDashboard/src/components/preflights/PreflightConfigCheckRedirector.tsx b/Engine/Source/Programs/Horde/HordeDashboard/src/components/preflights/PreflightConfigCheckRedirector.tsx
new file mode 100644
--- /dev/null
+++ b/Engine/Source/Programs/Horde/HordeDashboard/src/components/preflights/PreflightConfigCheckRedirector.tsx
@@ -0,0 +1,39 @@
+// Copyright Epic Games, Inc. All Rights Reserved.
+
+import React, { useState } from 'react';
+import { useNavigate } from 'react-router-dom';
+import backend, { useBackend } from '../../backend';
+import { useQuery } from '../JobDetailCommon';
+import ErrorHandler from '../ErrorHandler';
+
+function setError(message: string) {
+
+   console.error(message);
+
+   message = `${message}\n\nReferring URL: ${encodeURI(window.location.href)}`
+
+   ErrorHandler.set({
+
+      title: `Error handling Preflight Configuration`,
+      message: message
+
+   }, true);
+
+}
+
+// redirect from external source, where horde stream id, etc are not known by that application
+export const PreflightConfigRedirector: React.FC = () => {
+
+   const query = useQuery();
+
+   const shelvedCL = !query.get("shelvedchange") ? "" : query.get("shelvedchange")!;
+   
+   if (!shelvedCL) { 
+      setError("No shelved change specified");
+   } else {
+      window.location.assign(`/index?preflightconfig=${shelvedCL}`);
+   }
+
+   return null;
+
+}
\ No newline at end of file
