Horde: Avoid updating reactions for issues if they haven't changed. Hopefully address rate limiting observed in logs.

diff --git a/Engine/Source/Programs/Horde/Horde.Server/Notifications/Sinks/SlackNotificationSink.cs b/Engine/Source/Programs/Horde/Horde.Server/Notifications/Sinks/SlackNotificationSink.cs
--- a/Engine/Source/Programs/Horde/Horde.Server/Notifications/Sinks/SlackNotificationSink.cs
+++ b/Engine/Source/Programs/Horde/Horde.Server/Notifications/Sinks/SlackNotificationSink.cs
@@ -946,6 +946,16 @@ namespace Horde.Server.Notifications.Sinks
 
 		static string GetTriageThreadEventId(int issueId) => $"issue_triage_{issueId}";
 
+		[Flags]
+		enum ReactionFlags
+		{
+			None = 0,
+			Acknowledged = 1,
+			Quarantined = 2,
+			FixFailed = 4,
+			Resolved = 8,
+		}
+
 		async Task CreateOrUpdateWorkflowThreadAsync(string triageChannel, IIssue issue, IIssueSpan span, IReadOnlyList<IIssueSpan> spans, WorkflowConfig workflow)
 		{
 			Uri issueUrl = GetIssueUrl(issue, span.FirstFailure);
@@ -1027,8 +1037,8 @@ namespace Horde.Server.Notifications.Sinks
 					catch (Exception Ex)
 					{
 						_issueService.Collection.GetLogger(issue.Id).LogInformation("Error associating workflow thread with issue, bad URI format? {ErrorMessage}", Ex.Message);
-					}					
-				}				
+					}
+				}
 			}
 
 			// Post a message containing the controls and status
@@ -1154,7 +1164,7 @@ namespace Horde.Server.Notifications.Sinks
 			if ((workflow.TriageAlias != null || workflow.TriageTypeAliases != null) && issue.OwnerId == null && (suspects.All(x => x.DeclinedAt != null) || notifyTriageAlias) && !closed)
 			{
 				string? triageAlias;
-					
+
 				if (workflow.TriageTypeAliases == null || issue.Fingerprints.Count == 0 || !workflow.TriageTypeAliases.TryGetValue(issue.Fingerprints[0].Type, out triageAlias))
 				{
 					triageAlias = workflow.TriageAlias;
@@ -1167,24 +1177,6 @@ namespace Horde.Server.Notifications.Sinks
 				}
 			}
 
-			if (issue.AcknowledgedAt != null)
-			{
-				await _slackClient.AddReactionAsync(threadId, "eyes");
-			}
-			else
-			{
-				await _slackClient.RemoveReactionAsync(threadId, "eyes");
-			}
-
-			if (issue.QuarantinedByUserId != null)
-			{
-				await _slackClient.AddReactionAsync(threadId, "mask");
-			}
-			else
-			{
-				await _slackClient.RemoveReactionAsync(threadId, "mask");
-			}
-
 			IIssueSpan? fixFailedSpan = null;
 			if (issue.FixChange != null)
 			{
@@ -1223,22 +1215,70 @@ namespace Horde.Server.Notifications.Sinks
 				}
 			}
 
-			if (fixFailedSpan != null)
+			// Reactions
 			{
-				await _slackClient.AddReactionAsync(threadId, "x");
-			}
-			else
-			{
-				await _slackClient.RemoveReactionAsync(threadId, "x");
-			}
+				ReactionFlags reactions = ReactionFlags.None;
+				if (issue.AcknowledgedAt != null)
+				{
+					reactions |= ReactionFlags.Acknowledged;
+				}
+				if (issue.QuarantinedByUserId != null)
+				{
+					reactions |= ReactionFlags.Quarantined;
+				}
+				if (fixFailedSpan != null)
+				{
+					reactions |= ReactionFlags.FixFailed;
+				}
+				if (issue.ResolvedAt != null && fixFailedSpan == null)
+				{
+					reactions |= ReactionFlags.Resolved;
+				}
 
-			if (issue.ResolvedAt != null && fixFailedSpan == null)
-			{
-				await _slackClient.AddReactionAsync(threadId, "tick");
-			}
-			else
-			{
-				await _slackClient.RemoveReactionAsync(threadId, "tick");
+				string reactionEventId = $"{eventId}_reactions";
+				string reactionStateDigest = ((int)reactions).ToString();
+
+				MessageStateDocument? reactionState = await GetMessageStateAsync(triageChannel, reactionEventId);
+				if ((reactionState == null && reactions != ReactionFlags.None) || (reactionState != null && reactionState.Digest != reactionStateDigest))
+				{
+					if ((reactions & ReactionFlags.Acknowledged) != 0)
+					{
+						await _slackClient.AddReactionAsync(threadId, "eyes");
+					}
+					else
+					{
+						await _slackClient.RemoveReactionAsync(threadId, "eyes");
+					}
+
+					if ((reactions & ReactionFlags.Quarantined) != 0)
+					{
+						await _slackClient.AddReactionAsync(threadId, "mask");
+					}
+					else
+					{
+						await _slackClient.RemoveReactionAsync(threadId, "mask");
+					}
+
+					if ((reactions & ReactionFlags.FixFailed) != 0)
+					{
+						await _slackClient.AddReactionAsync(threadId, "x");
+					}
+					else
+					{
+						await _slackClient.RemoveReactionAsync(threadId, "x");
+					}
+
+					if ((reactions & ReactionFlags.Resolved) != 0)
+					{
+						await _slackClient.AddReactionAsync(threadId, "tick");
+					}
+					else
+					{
+						await _slackClient.RemoveReactionAsync(threadId, "tick");
+					}
+
+					await AddOrUpdateMessageStateAsync(triageChannel, reactionEventId, null, reactionStateDigest, threadId);
+				}
 			}
 
 			if (issue.ExternalIssueKey != null)
