Horde: Fix OpenTelemetry gauges not reporting

diff --git a/Engine/Source/Programs/Horde/Horde.Server/Agents/AgentService.cs b/Engine/Source/Programs/Horde/Horde.Server/Agents/AgentService.cs
--- a/Engine/Source/Programs/Horde/Horde.Server/Agents/AgentService.cs
+++ b/Engine/Source/Programs/Horde/Horde.Server/Agents/AgentService.cs
@@ -91,7 +91,7 @@ namespace Horde.Server.Agents
 		// All the agents currently performing a long poll for work on this server
 		readonly Dictionary<AgentId, CancellationTokenSource> _waitingAgents = new Dictionary<AgentId, CancellationTokenSource>();
 
-		readonly Gauge<int> _agentTotalCount;
+		IEnumerable<Measurement<int>> _measurements = new List<Measurement<int>>();
 
 		// Subscription for update events
 		IAsyncDisposable? _subscription;
@@ -117,7 +117,7 @@ namespace Horde.Server.Agents
 			_meter = meter;
 			_logger = logger;
 
-			_agentTotalCount = _meter.CreateGauge<int>("horde.agent.count");
+			_meter.CreateObservableGauge("horde.agent.count", () => _measurements);
 		}
 
 		/// <inheritdoc/>
@@ -1087,15 +1087,20 @@ namespace Horde.Server.Agents
 			int numAgentsTotalUnhealthy = agentList.Count(a => a.Enabled && a.Status == AgentStatus.Unhealthy);
 			int numAgentsTotalPaused = agentList.Count(a => a.Enabled && a.Status == AgentStatus.Paused);
 			int numAgentsTotalUnspecified = agentList.Count(a => a.Enabled && a.Status == AgentStatus.Unspecified);
-			
-			_agentTotalCount.Record(numAgentsTotal);
-			_agentTotalCount.Record(numAgentsTotalEnabled, new KeyValuePair<string, object?>("status", "enabled"));
-			_agentTotalCount.Record(numAgentsTotalDisabled, new KeyValuePair<string, object?>("status", "disabled"));
-			_agentTotalCount.Record(numAgentsTotalOk, new KeyValuePair<string, object?>("status", "ok"));
-			_agentTotalCount.Record(numAgentsTotalStopping, new KeyValuePair<string, object?>("status", "stopping"));
-			_agentTotalCount.Record(numAgentsTotalUnhealthy, new KeyValuePair<string, object?>("status", "unhealthy"));
-			_agentTotalCount.Record(numAgentsTotalPaused, new KeyValuePair<string, object?>("status", "paused"));
-			_agentTotalCount.Record(numAgentsTotalUnspecified, new KeyValuePair<string, object?>("status", "unspecified"));
+
+			List<Measurement<int>> newMeasurements = new()
+			{
+				new Measurement<int>(numAgentsTotal),
+				new Measurement<int>(numAgentsTotalEnabled, new KeyValuePair<string, object?>("status", "enabled")),
+				new Measurement<int>(numAgentsTotalDisabled, new KeyValuePair<string, object?>("status", "disabled")),
+				new Measurement<int>(numAgentsTotalOk, new KeyValuePair<string, object?>("status", "ok")),
+				new Measurement<int>(numAgentsTotalStopping, new KeyValuePair<string, object?>("status", "stopping")),
+				new Measurement<int>(numAgentsTotalUnhealthy, new KeyValuePair<string, object?>("status", "unhealthy")),
+				new Measurement<int>(numAgentsTotalPaused, new KeyValuePair<string, object?>("status", "paused")),
+				new Measurement<int>(numAgentsTotalUnspecified, new KeyValuePair<string, object?>("status", "unspecified")),
+			};
+
+			_measurements = newMeasurements;
 		}
 
 		/// <summary>
diff --git a/Engine/Source/Programs/Horde/Horde.Server/Agents/Fleet/FleetService.cs b/Engine/Source/Programs/Horde/Horde.Server/Agents/Fleet/FleetService.cs
--- a/Engine/Source/Programs/Horde/Horde.Server/Agents/Fleet/FleetService.cs
+++ b/Engine/Source/Programs/Horde/Horde.Server/Agents/Fleet/FleetService.cs
@@ -81,8 +81,6 @@ namespace Horde.Server.Agents.Fleet
 		private readonly IOptionsMonitor<GlobalConfig> _globalConfig;
 		private readonly Tracer _tracer;
 		private readonly ILogger<FleetService> _logger;
-		private readonly Gauge<int> _agentCountCurrentGauge;
-		private readonly Gauge<int> _agentCountTargetGauge;
 		
 		/// <summary>
 		/// Constructor
@@ -123,9 +121,6 @@ namespace Horde.Server.Agents.Fleet
 			_settings = settings;
 			_defaultScaleOutCooldown = TimeSpan.FromSeconds(settings.Value.AgentPoolScaleOutCooldownSeconds);
 			_defaultScaleInCooldown = TimeSpan.FromSeconds(settings.Value.AgentPoolScaleInCooldownSeconds);
-
-			_agentCountCurrentGauge = _meter.CreateGauge<int>("horde.fleet.agentCountCurrent");
-			_agentCountTargetGauge = _meter.CreateGauge<int>("horde.fleet.agentCountTarget");
 		}
 
 		/// <inheritdoc/>
@@ -238,9 +233,6 @@ namespace Horde.Server.Agents.Fleet
 					pool.Name, currentAgentCount, desiredAgentCount, deltaAgentCount);	
 			}
 
-			_agentCountCurrentGauge.Record(currentAgentCount, KeyValuePair.Create("pool", (object?)pool.Name));
-			_agentCountTargetGauge.Record(desiredAgentCount, KeyValuePair.Create("pool", (object?)pool.Name));
-
 			DateTime? scaleOutTime = null;
 			DateTime? scaleInTime = null;
 			
diff --git a/Engine/Source/Programs/Horde/Horde.Server/Server/MetricService.cs b/Engine/Source/Programs/Horde/Horde.Server/Server/MetricService.cs
--- a/Engine/Source/Programs/Horde/Horde.Server/Server/MetricService.cs
+++ b/Engine/Source/Programs/Horde/Horde.Server/Server/MetricService.cs
@@ -4,10 +4,7 @@ using System;
 using System.Diagnostics.Metrics;
 using System.Threading;
 using System.Threading.Tasks;
-using Horde.Server.Utilities;
-using HordeCommon;
 using Microsoft.Extensions.Hosting;
-using Microsoft.Extensions.Logging;
 
 namespace Horde.Server.Server
 {
@@ -16,101 +13,61 @@ namespace Horde.Server.Server
 	/// </summary>
 	public sealed class MetricService : IHostedService, IDisposable
 	{
-		readonly Meter _meter;
-		readonly ITicker _ticker;
-		
-		private readonly Gauge<int> _threadPoolIoMax;
-		private readonly Gauge<int> _threadPoolIoMin;
-		private readonly Gauge<int> _threadPoolIoFree;
-		private readonly Gauge<int> _threadPoolIoBusy;
-		private readonly Gauge<int> _threadPoolWorkerMax;
-		private readonly Gauge<int> _threadPoolWorkerMin;
-		private readonly Gauge<int> _threadPoolWorkerFree;
-		private readonly Gauge<int> _threadPoolWorkerBusy;
-		private readonly Gauge<long> _clrGcTotalMem;
-		private readonly Gauge<long> _clrGcTotalAllocated;
-		private readonly Gauge<long> _clrGcHeapSize;
-		private readonly Gauge<long> _clrGcFragmented;
-		private readonly Gauge<long> _clrGcMemoryLoad;
-		private readonly Gauge<long> _clrGcTotalAvailableMem;
-		private readonly Gauge<long> _clrGcHighMemThreshold;
-
 		/// <summary>
 		/// Constructor
 		/// </summary>
-		public MetricService(Meter meter, IClock clock, ILogger<MetricService> logger)
+		public MetricService(Meter meter)
 		{
-			_meter = meter;
-			_ticker = clock.AddTicker<MetricService>(TimeSpan.FromSeconds(20.0), TickAsync, logger);
-
-			_threadPoolIoMax = _meter.CreateGauge<int>("horde.clr.threadpool.io.max");
-			_threadPoolIoMin = _meter.CreateGauge<int>("horde.clr.threadpool.io.min");
-			_threadPoolIoFree = _meter.CreateGauge<int>("horde.clr.threadpool.io.free");
-			_threadPoolIoBusy = _meter.CreateGauge<int>("horde.clr.threadpool.io.busy");
-			_threadPoolWorkerMax = _meter.CreateGauge<int>("horde.clr.threadpool.worker.max");
-			_threadPoolWorkerMin = _meter.CreateGauge<int>("horde.clr.threadpool.worker.min");
-			_threadPoolWorkerFree = _meter.CreateGauge<int>("horde.clr.threadpool.worker.free");
-			_threadPoolWorkerBusy = _meter.CreateGauge<int>("horde.clr.threadpool.worker.busy");
-			_clrGcTotalMem = _meter.CreateGauge<long>("horde.clr.gc.totalMemory");
-			_clrGcTotalAllocated = _meter.CreateGauge<long>("horde.clr.gc.totalAllocated");
-			_clrGcHeapSize = _meter.CreateGauge<long>("horde.clr.gc.heapSize");
-			_clrGcFragmented = _meter.CreateGauge<long>("horde.clr.gc.fragmented");
-			_clrGcMemoryLoad = _meter.CreateGauge<long>("horde.clr.gc.memoryLoad");
-			_clrGcTotalAvailableMem = _meter.CreateGauge<long>("horde.clr.gc.totalAvailableMemory");
-			_clrGcHighMemThreshold = _meter.CreateGauge<long>("horde.clr.gc.highMemoryLoadThreshold");
+			// Force a collection to ensure metrics are populated as some of them require a GC to have been run
+			GC.Collect();
+			
+			meter.CreateObservableGauge("horde.clr.threadpool.io.min", () => new ThreadMetrics().IoMin);
+			meter.CreateObservableGauge("horde.clr.threadpool.io.max", () => new ThreadMetrics().IoMax);
+			meter.CreateObservableGauge("horde.clr.threadpool.io.free", () => new ThreadMetrics().IoFree);
+			meter.CreateObservableGauge("horde.clr.threadpool.io.busy", () => new ThreadMetrics().IoBusy);
+			
+			meter.CreateObservableGauge("horde.clr.threadpool.worker.min", () => new ThreadMetrics().WorkerMin);
+			meter.CreateObservableGauge("horde.clr.threadpool.worker.max", () => new ThreadMetrics().WorkerMax);
+			meter.CreateObservableGauge("horde.clr.threadpool.worker.free", () => new ThreadMetrics().WorkerFree);
+			meter.CreateObservableGauge("horde.clr.threadpool.worker.busy", () => new ThreadMetrics().WorkerBusy);
+			
+			meter.CreateObservableGauge("horde.clr.gc.totalMemory", () => GC.GetTotalMemory(false));
+			meter.CreateObservableGauge("horde.clr.gc.totalAllocated", () => GC.GetTotalAllocatedBytes());
+			meter.CreateObservableGauge("horde.clr.gc.heapSize", () => GC.GetGCMemoryInfo().HeapSizeBytes);
+			meter.CreateObservableGauge("horde.clr.gc.fragmented", () => GC.GetGCMemoryInfo().FragmentedBytes);
+			meter.CreateObservableGauge("horde.clr.gc.memoryLoad", () => GC.GetGCMemoryInfo().MemoryLoadBytes);
+			meter.CreateObservableGauge("horde.clr.gc.totalAvailableMemory", () => GC.GetGCMemoryInfo().TotalAvailableMemoryBytes);
+			meter.CreateObservableGauge("horde.clr.gc.highMemoryLoadThreshold", () => GC.GetGCMemoryInfo().HighMemoryLoadThresholdBytes);
 		}
 
 		/// <inheritdoc/>
-		public Task StartAsync(CancellationToken cancellationToken) => _ticker.StartAsync();
+		public Task StartAsync(CancellationToken cancellationToken) { return Task.CompletedTask; }
 
 		/// <inheritdoc/>
-		public Task StopAsync(CancellationToken cancellationToken) => _ticker.StopAsync();
+		public Task StopAsync(CancellationToken cancellationToken) { return Task.CompletedTask; }
 
 		/// <inheritdoc/>
-		public void Dispose() => _ticker.Dispose();
-
-		/// <summary>
-		/// Execute the background task
-		/// </summary>
-		/// <param name="stoppingToken">Cancellation token for the async task</param>
-		/// <returns>Async task</returns>
-		ValueTask TickAsync(CancellationToken stoppingToken)
-		{
-			ReportGcMetrics();
-			ReportThreadMetrics();
-			return new ValueTask();
-		}
-
-		private void ReportThreadMetrics()
-		{
-			ThreadPool.GetMaxThreads(out int maxWorkerThreads, out int maxIoThreads);
-			ThreadPool.GetAvailableThreads(out int freeWorkerThreads, out int freeIoThreads);
-			ThreadPool.GetMinThreads(out int minWorkerThreads, out int minIoThreads);
-
-			int busyIoThreads = maxIoThreads - freeIoThreads;
-			int busyWorkerThreads = maxWorkerThreads - freeWorkerThreads;
-			
-			_threadPoolIoMax.Record(maxIoThreads);
-			_threadPoolIoMin.Record(minIoThreads);
-			_threadPoolIoFree.Record(freeIoThreads);
-			_threadPoolIoBusy.Record(busyIoThreads);
-			_threadPoolWorkerMax.Record(maxWorkerThreads);
-			_threadPoolWorkerMin.Record(minWorkerThreads);
-			_threadPoolWorkerFree.Record(freeWorkerThreads);
-			_threadPoolWorkerBusy.Record(busyWorkerThreads);
-		}
+		public void Dispose() { }
 
-		private void ReportGcMetrics()
+		private struct ThreadMetrics
 		{
-			GCMemoryInfo gcMemoryInfo = GC.GetGCMemoryInfo();
+			public readonly int WorkerMin;
+			public readonly int WorkerMax;
+			public readonly int WorkerFree;
+			public readonly int WorkerBusy;
+			public readonly int IoMin;
+			public readonly int IoMax;
+			public readonly int IoFree;
+			public readonly int IoBusy;
 
-			_clrGcTotalMem.Record(GC.GetTotalMemory(false));
-			_clrGcTotalAllocated.Record(GC.GetTotalAllocatedBytes());
-			_clrGcHeapSize.Record(gcMemoryInfo.HeapSizeBytes);
-			_clrGcFragmented.Record(gcMemoryInfo.FragmentedBytes);
-			_clrGcMemoryLoad.Record(gcMemoryInfo.MemoryLoadBytes);
-			_clrGcTotalAvailableMem.Record(gcMemoryInfo.TotalAvailableMemoryBytes);
-			_clrGcHighMemThreshold.Record(gcMemoryInfo.HighMemoryLoadThresholdBytes);
+			public ThreadMetrics()
+			{
+				ThreadPool.GetMinThreads(out WorkerMin, out IoMin);
+				ThreadPool.GetMaxThreads(out WorkerMax, out IoMax);
+				ThreadPool.GetAvailableThreads(out WorkerFree, out IoFree);
+				WorkerBusy = WorkerMax - WorkerFree;
+				IoBusy = WorkerMax - WorkerFree;
+			}
 		}
 	}
 }
diff --git a/Engine/Source/Programs/Horde/Horde.Server/Utilities/OpenTelemetry.cs b/Engine/Source/Programs/Horde/Horde.Server/Utilities/OpenTelemetry.cs
--- a/Engine/Source/Programs/Horde/Horde.Server/Utilities/OpenTelemetry.cs
+++ b/Engine/Source/Programs/Horde/Horde.Server/Utilities/OpenTelemetry.cs
@@ -14,7 +14,6 @@ using Microsoft.AspNetCore.Http;
 using Microsoft.Extensions.DependencyInjection;
 using Microsoft.Extensions.Logging;
 using MongoDB.Bson;
-using MongoDB.Bson.Serialization;
 using MongoDB.Driver;
 using OpenTelemetry.Exporter;
 using OpenTelemetry.Logs;
@@ -343,18 +342,6 @@ public static class OpenTelemetrySpanExtensions
 
 		return span;
 	}
-	
-	/// <summary>
-	/// Create a metrics Gauge object.
-	/// </summary>
-	/// <param name="meter">Current meter being extended</param>
-	/// <param name="name">The instrument name. cannot be null.</param>
-	/// <param name="unit">Optional instrument unit of measurements.</param>
-	/// <param name="description">Optional instrument description.</param>
-	/// <remarks>
-	/// Gauge is an Instrument which records absolute values
-	/// </remarks>
-	public static Gauge<T> CreateGauge<T>(this Meter meter, string name, string? unit = null, string? description = null) where T : struct => new (meter, name, unit, description);
 }
 
 /// <summary>
@@ -378,28 +365,3 @@ public static class OpenTelemetryMeters
 	/// </summary>
 	public static readonly Meter Horde = new (HordeName);
 }
-
-/// <summary>
-/// A non-observable gauge
-/// </summary>
-/// <typeparam name="T">Type of value being measured</typeparam>
-public sealed class Gauge<T> : Instrument<T> where T : struct
-{
-    internal Gauge(Meter meter, string name, string? unit, string? description) : base(meter, name, unit, description)
-    {
-        Publish();
-    }
-
-    /// <summary>
-    /// Record the current value of the measurement.
-    /// </summary>
-    /// <param name="value">Current value.</param>
-    public void Record(T value) => RecordMeasurement(value);
-
-    /// <summary>
-    /// Record the current value of the measurement.
-    /// </summary>
-    /// <param name="value">The increment measurement.</param>
-    /// <param name="tag">A key-value pair tag associated with the measurement.</param>
-    public void Record(T value, KeyValuePair<string, object?> tag) => RecordMeasurement(value, tag);
-}
\ No newline at end of file
