Horde: Send reports listing agents stuck in a conform loop to Slack.

diff --git a/Engine/Source/Programs/Horde/Horde.Server.Tests/NotificationServiceTests.cs b/Engine/Source/Programs/Horde/Horde.Server.Tests/NotificationServiceTests.cs
--- a/Engine/Source/Programs/Horde/Horde.Server.Tests/NotificationServiceTests.cs
+++ b/Engine/Source/Programs/Horde/Horde.Server.Tests/NotificationServiceTests.cs
@@ -11,13 +11,13 @@ using Horde.Server.Jobs.Graphs;
 using Horde.Server.Notifications;
 using Horde.Server.Users;
 using Horde.Server.Streams;
-using Horde.Server.Utilities;
 using HordeCommon;
 using Microsoft.Extensions.DependencyInjection;
 using Microsoft.VisualStudio.TestTools.UnitTesting;
 using Moq;
 using Horde.Server.Logs;
 using Microsoft.Extensions.Caching.Memory;
+using Horde.Server.Agents;
 
 namespace Horde.Server.Tests
 {
@@ -41,6 +41,7 @@ namespace Horde.Server.Tests
 		public Task NotifyConfigUpdateAsync(Exception? ex) => Task.CompletedTask;
 		public Task NotifyConfigUpdateFailureAsync(string errorMessage, string fileName, int? change = null, IUser? author = null, string? description = null) { throw new NotImplementedException(); }
 		public Task NotifyDeviceServiceAsync(string message, IDevice? device = null, IDevicePool? pool = null, StreamConfig? stream = null, IJob? job = null, IJobStep? step = null, INode? node = null, IUser? user = null) { throw new NotImplementedException(); }
+		public Task SendAgentReportAsync(AgentReport report) => throw new NotImplementedException();
 		public Task SendIssueReportAsync(IssueReportGroup report) => throw new NotImplementedException();
 		public Task SendDeviceIssueReportAsync(DeviceIssueReport report) => throw new NotImplementedException();
 	}
diff --git a/Engine/Source/Programs/Horde/Horde.Server/Agents/AgentReportService.cs b/Engine/Source/Programs/Horde/Horde.Server/Agents/AgentReportService.cs
new file mode 100644
--- /dev/null
+++ b/Engine/Source/Programs/Horde/Horde.Server/Agents/AgentReportService.cs
@@ -0,0 +1,113 @@
+// Copyright Epic Games, Inc. All Rights Reserved.
+
+using System;
+using System.Collections.Generic;
+using System.Threading;
+using System.Threading.Tasks;
+using EpicGames.Horde.Api;
+using Horde.Server.Jobs;
+using Horde.Server.Notifications;
+using Horde.Server.Server;
+using Horde.Server.Utilities;
+using HordeCommon;
+using Microsoft.Extensions.Hosting;
+using Microsoft.Extensions.Logging;
+using Microsoft.Extensions.Options;
+
+namespace Horde.Server.Agents
+{
+	/// <summary>
+	/// Report for the state of agents in the farm
+	/// </summary>
+	public class AgentReport
+	{
+		/// <summary>
+		/// List of agents stuck in a conform loop, plus the number of times they've attempted to conform
+		/// </summary>
+		public List<(AgentId, int)> ConformLoop { get; } = new List<(AgentId, int)>();
+	}
+
+	[SingletonDocument("agent-report-state", "6268871c211d05611b3e4fd8")]
+	class AgentReportState : SingletonBase
+	{
+		public DateTime LastUpdateUtc { get; set; }
+	}
+
+	/// <summary>
+	/// Posts summaries for all the open issues in different streams to Slack channels
+	/// </summary>
+	public sealed class AgentReportService : IHostedService, IDisposable
+	{
+		readonly SingletonDocument<AgentReportState> _state;
+		readonly IAgentCollection _agentCollection;
+		readonly INotificationService _notificationService;
+		readonly IClock _clock;
+		readonly ITicker _ticker;
+		readonly IOptionsMonitor<GlobalConfig> _globalConfig;
+		readonly ILogger _logger;
+
+		/// <summary>
+		/// Constructor
+		/// </summary>
+		public AgentReportService(MongoService mongoService, IAgentCollection agentCollection, INotificationService notificationService, IClock clock, IOptionsMonitor<GlobalConfig> globalConfig, ILogger<AgentReportService> logger)
+		{
+			_state = new SingletonDocument<AgentReportState>(mongoService);
+			_agentCollection = agentCollection;
+			_notificationService = notificationService;
+			_clock = clock;
+			_ticker = clock.AddSharedTicker<AgentReportService>(TimeSpan.FromMinutes(5.0), TickAsync, logger);
+			_globalConfig = globalConfig;
+			_logger = logger;
+		}
+
+		/// <inheritdoc/>
+		public void Dispose()
+		{
+			_ticker.Dispose();
+		}
+
+		/// <inheritdoc/>
+		public async Task StartAsync(CancellationToken cancellationToken)
+		{
+			await _ticker.StartAsync();
+		}
+
+		/// <inheritdoc/>
+		public async Task StopAsync(CancellationToken cancellationToken)
+		{
+			await _ticker.StopAsync();
+		}
+
+		async ValueTask TickAsync(CancellationToken cancellationToken)
+		{
+			AgentReportState state = await _state.GetAsync();
+			DateTime utcNow = _clock.UtcNow;
+
+			TimeSpan agentReportTime = TimeSpan.FromHours(9.0); // 9am
+
+			DateTime nextUpdateTime = _clock.TimeZone.GetStartOfDayUtc(state.LastUpdateUtc + TimeSpan.FromHours(24.0)) + agentReportTime;
+			if (utcNow > nextUpdateTime)
+			{
+				AgentReport report = new AgentReport();
+
+				List<IAgent> agents = await _agentCollection.FindAsync();
+				foreach (IAgent agent in agents)
+				{
+					if (agent.ConformAttemptCount.HasValue && agent.ConformAttemptCount.Value > 3)
+					{
+						report.ConformLoop.Add((agent.Id, agent.ConformAttemptCount.Value));
+					}
+				}
+
+				await _notificationService.SendAgentReportAsync(report);
+
+				DateTime updateTime = _clock.TimeZone.GetStartOfDayUtc(utcNow) + agentReportTime;
+				if (updateTime > utcNow)
+				{
+					updateTime -= TimeSpan.FromDays(1.0);
+				}
+				await _state.UpdateAsync(x => x.LastUpdateUtc = updateTime);
+			}
+		}
+	}
+}
diff --git a/Engine/Source/Programs/Horde/Horde.Server/Notifications/INotificationService.cs b/Engine/Source/Programs/Horde/Horde.Server/Notifications/INotificationService.cs
--- a/Engine/Source/Programs/Horde/Horde.Server/Notifications/INotificationService.cs
+++ b/Engine/Source/Programs/Horde/Horde.Server/Notifications/INotificationService.cs
@@ -12,6 +12,7 @@ using Horde.Server.Issues;
 using Horde.Server.Users;
 using MongoDB.Bson;
 using Horde.Server.Streams;
+using Horde.Server.Agents;
 
 namespace Horde.Server.Notifications
 {
@@ -156,6 +157,12 @@ namespace Horde.Server.Notifications
 		/// <param name="report">The report data to send</param>
 		Task SendDeviceIssueReportAsync(DeviceIssueReport report);
 
+		/// <summary>
+		/// Post a notification for any agents encountering issues
+		/// </summary>
+		/// <param name="report">The report data to send</param>
+		Task SendAgentReportAsync(AgentReport report);
+
 		/// <summary>
 		/// Post a notification for the open issues in a stream
 		/// </summary>
diff --git a/Engine/Source/Programs/Horde/Horde.Server/Notifications/INotificationSink.cs b/Engine/Source/Programs/Horde/Horde.Server/Notifications/INotificationSink.cs
--- a/Engine/Source/Programs/Horde/Horde.Server/Notifications/INotificationSink.cs
+++ b/Engine/Source/Programs/Horde/Horde.Server/Notifications/INotificationSink.cs
@@ -3,6 +3,7 @@
 using System;
 using System.Collections.Generic;
 using System.Threading.Tasks;
+using Horde.Server.Agents;
 using Horde.Server.Configuration;
 using Horde.Server.Devices;
 using Horde.Server.Issues;
@@ -126,6 +127,10 @@ namespace Horde.Server.Notifications
 		/// <returns></returns>
 		Task SendDeviceIssueReportAsync(DeviceIssueReport report);
 
+		/// <summary>
+		/// Notify the status of agents in the farm
+		/// </summary>
+		Task SendAgentReportAsync(AgentReport report);
 	}
 }
 
diff --git a/Engine/Source/Programs/Horde/Horde.Server/Notifications/NotificationService.cs b/Engine/Source/Programs/Horde/Horde.Server/Notifications/NotificationService.cs
--- a/Engine/Source/Programs/Horde/Horde.Server/Notifications/NotificationService.cs
+++ b/Engine/Source/Programs/Horde/Horde.Server/Notifications/NotificationService.cs
@@ -11,6 +11,7 @@ using System.Threading.Tasks;
 using EpicGames.Core;
 using EpicGames.Horde.Api;
 using EpicGames.Redis;
+using Horde.Server.Agents;
 using Horde.Server.Agents.Pools;
 using Horde.Server.Configuration;
 using Horde.Server.Devices;
@@ -793,6 +794,22 @@ namespace Horde.Server.Notifications
 			_logger.LogDebug("Finished sending label notifications for label {DashboardName}/{UgsName} in job {JobId}", label.DashboardName, label.UgsName, job.Id);
 		}
 
+		/// <inheritdoc/>
+		public async Task SendAgentReportAsync(AgentReport report)
+		{
+			foreach (INotificationSink sink in _sinks)
+			{
+				try
+				{
+					await sink.SendAgentReportAsync(report);
+				}
+				catch (Exception e)
+				{
+					_logger.LogError(e, "Failed sending agent report to {Channel}", _settings.CurrentValue.AgentNotificationChannel);
+				}
+			}
+		}
+
 		/// <inheritdoc/>
 		public async Task SendIssueReportAsync(IssueReportGroup report)
 		{
diff --git a/Engine/Source/Programs/Horde/Horde.Server/Notifications/Sinks/SlackNotificationSink.cs b/Engine/Source/Programs/Horde/Horde.Server/Notifications/Sinks/SlackNotificationSink.cs
--- a/Engine/Source/Programs/Horde/Horde.Server/Notifications/Sinks/SlackNotificationSink.cs
+++ b/Engine/Source/Programs/Horde/Horde.Server/Notifications/Sinks/SlackNotificationSink.cs
@@ -43,6 +43,7 @@ using MongoDB.Bson;
 using MongoDB.Bson.Serialization.Attributes;
 using MongoDB.Driver;
 using EpicGames.Horde;
+using Horde.Server.Agents;
 
 namespace Horde.Server.Notifications.Sinks
 {
@@ -2996,5 +2997,39 @@ namespace Horde.Server.Notifications.Sinks
 				return;
 			}
 		}
+
+		/// <inheritdoc/>
+		public async Task SendAgentReportAsync(AgentReport report)
+		{
+			if (!String.IsNullOrEmpty(_settings.AgentNotificationChannel))
+			{
+				_logger.LogInformation("Sending agent report to {Channel}", _settings.AgentNotificationChannel);
+
+				SlackMessage headerMessage = new SlackMessage();
+				headerMessage.AddHeader($"Agent status");
+				await _slackClient.PostMessageAsync(_settings.AgentNotificationChannel, headerMessage);
+
+				{
+					StringBuilder conformMessage = new StringBuilder("**Conform issues:**\n");
+					if (report.ConformLoop.Count == 0)
+					{
+						conformMessage.Append("None.\n");
+						await _slackClient.PostMessageAsync(_settings.AgentNotificationChannel, conformMessage.ToString());
+					}
+					else
+					{
+						foreach (IReadOnlyList<(AgentId, int)> conformBatch in report.ConformLoop.Batch(10))
+						{
+							foreach ((AgentId agentId, int conformCount) in conformBatch)
+							{
+								Uri agentUrl = new Uri(_settings.DashboardUrl, $"agents?agentId={agentId}");
+								conformMessage.Append($"* **[{agentId}]({agentUrl})** has run conform {conformCount} times\n");
+							}
+							await _slackClient.PostMessageAsync(_settings.AgentNotificationChannel, conformMessage.ToString());
+						}
+					}
+				}
+			}
+		}
 	}
 }
diff --git a/Engine/Source/Programs/Horde/Horde.Server/ServerSettings.cs b/Engine/Source/Programs/Horde/Horde.Server/ServerSettings.cs
--- a/Engine/Source/Programs/Horde/Horde.Server/ServerSettings.cs
+++ b/Engine/Source/Programs/Horde/Horde.Server/ServerSettings.cs
@@ -694,6 +694,11 @@ namespace Horde.Server
 		/// </summary>
 		public string? JobNotificationChannel { get; set; }
 
+		/// <summary>
+		/// Slack channel to send agent related notifications to.
+		/// </summary>
+		public string? AgentNotificationChannel { get; set; }
+
 		/// <summary>
 		/// The URl to use for generating links back to the dashboard.
 		/// </summary>
@@ -744,7 +749,6 @@ namespace Horde.Server
 		/// </summary>
 		public string? DeviceReportChannel { get; set; }
 
-
 		/// <summary>
 		/// Default agent pool sizing strategy for pools that doesn't have one explicitly configured
 		/// </summary>
