Horde: Fix issue with redirection not working with navigation hook on initial render, and return a Redirect instead
#skipci
#rnx

diff --git a/Engine/Source/Programs/Horde/HordeDashboard/src/components/preflights/PreflightConfigCheckRedirector.tsx b/Engine/Source/Programs/Horde/HordeDashboard/src/components/preflights/PreflightConfigCheckRedirector.tsx
--- a/Engine/Source/Programs/Horde/HordeDashboard/src/components/preflights/PreflightConfigCheckRedirector.tsx
+++ b/Engine/Source/Programs/Horde/HordeDashboard/src/components/preflights/PreflightConfigCheckRedirector.tsx
@@ -1,8 +1,7 @@
 // Copyright Epic Games, Inc. All Rights Reserved.
 
-import React, { useState } from 'react';
-import { useNavigate } from 'react-router-dom';
-import backend, { useBackend } from '../../backend';
+import React from 'react';
+import { Navigate } from 'react-router-dom';
 import { useQuery } from '../JobDetailCommon';
 import ErrorHandler from '../ErrorHandler';
 
@@ -31,7 +30,7 @@ export const PreflightConfigRedirector: React.FC = () => {
    if (!shelvedCL) { 
       setError("No shelved change specified");
    } else {
-      window.location.assign(`/index?preflightconfig=${shelvedCL}`);
+      return <Navigate to={`/index?preflightconfig=${shelvedCL}`} replace={true} />
    }
 
    return null;
