Horde: Fix incorrect variable in auto-scale log statement

diff --git a/Engine/Source/Programs/Horde/Horde.Server/Agents/Fleet/FleetService.cs b/Engine/Source/Programs/Horde/Horde.Server/Agents/Fleet/FleetService.cs
--- a/Engine/Source/Programs/Horde/Horde.Server/Agents/Fleet/FleetService.cs
+++ b/Engine/Source/Programs/Horde/Horde.Server/Agents/Fleet/FleetService.cs
@@ -294,7 +294,7 @@ namespace Horde.Server.Agents.Fleet
 			if (isResultDifferentFromLastTime && !isCooldownResult)
 			{
 				_logger.LogInformation("Scale result: Outcome={Outcome} AgentsAdded={AgentsAdded} AgentsRemoved={AgentsRemoved} Message={Message}", 
-					result.Outcome, result.AgentsAddedCount, result.AgentsAddedCount, result.Message);
+					result.Outcome, result.AgentsAddedCount, result.AgentsRemovedCount, result.Message);
 			}
 			
 			span.SetAttribute("resultOutcome", result.Outcome.ToString());
