Horde: Fix delayed update on device updates using stale data
#rnx
#skipci

diff --git a/Engine/Source/Programs/Horde/HordeDashboard/src/backend/PollBase.ts b/Engine/Source/Programs/Horde/HordeDashboard/src/backend/PollBase.ts
--- a/Engine/Source/Programs/Horde/HordeDashboard/src/backend/PollBase.ts
+++ b/Engine/Source/Programs/Horde/HordeDashboard/src/backend/PollBase.ts
@@ -34,6 +34,11 @@ export abstract class PollBase {
 
     abstract poll(): Promise<void>;
 
+    async forceUpdate() {        
+        this.stop(); 
+        this.update();
+    }
+
     async update() {
 
         clearTimeout(this.timeoutId);
diff --git a/Engine/Source/Programs/Horde/HordeDashboard/src/components/DeviceEditor.tsx b/Engine/Source/Programs/Horde/HordeDashboard/src/components/DeviceEditor.tsx
--- a/Engine/Source/Programs/Horde/HordeDashboard/src/components/DeviceEditor.tsx
+++ b/Engine/Source/Programs/Horde/HordeDashboard/src/components/DeviceEditor.tsx
@@ -217,7 +217,9 @@ export const DeviceEditor: React.FC<{ handler: DeviceHandler, deviceIn?: GetDevi
         } catch (reason:any) {
             setError(reason.toString());
             setSubmitting(false);
-        }
+       }
+       
+       await handler.forceUpdate();
 
     }
 
@@ -391,7 +393,9 @@ export const DeviceEditor: React.FC<{ handler: DeviceHandler, deviceIn?: GetDevi
                     setError(`Problem modifying device,: ${reason}`);
                     setSubmitting(false);
                 })
-            }
+           }
+           
+           await handler.forceUpdate();
 
         } catch (reason:any) {
             setError(reason.toString());
