Horde: Ensure only ephemeral agents are permanently deleted

diff --git a/Engine/Source/Programs/Horde/Horde.Server/Agents/AgentService.cs b/Engine/Source/Programs/Horde/Horde.Server/Agents/AgentService.cs
--- a/Engine/Source/Programs/Horde/Horde.Server/Agents/AgentService.cs
+++ b/Engine/Source/Programs/Horde/Horde.Server/Agents/AgentService.cs
@@ -1068,7 +1068,7 @@ namespace Horde.Server.Agents
 			{
 				cancellationToken.ThrowIfCancellationRequested();
 				bool noStatusChangeDuringPeriod = _clock.UtcNow > agent.LastStatusChange + TimeSpan.FromDays(7);
-				if (agent.Status == AgentStatus.Stopped && noStatusChangeDuringPeriod)
+				if (agent.Status == AgentStatus.Stopped && agent.Ephemeral && noStatusChangeDuringPeriod)
 				{
 					_logger.LogDebug("Deleting ephemeral agent {Agent}", agent.Id);
 					await DeleteAgentAsync(agent, true);
