Horde: Fix issue with state device data submission, better
#rnx
#skipci

diff --git a/Engine/Source/Programs/Horde/HordeDashboard/src/components/DeviceEditor.tsx b/Engine/Source/Programs/Horde/HordeDashboard/src/components/DeviceEditor.tsx
--- a/Engine/Source/Programs/Horde/HordeDashboard/src/components/DeviceEditor.tsx
+++ b/Engine/Source/Programs/Horde/HordeDashboard/src/components/DeviceEditor.tsx
@@ -8,99 +8,99 @@ import { PollBase } from "../backend/PollBase";
 import { hordeClasses } from "../styles/Styles";
 
 type DeviceEditData = {
-    // if defined, existing device
-    id?: string;
-    platformId: string;
-    poolId: string;
-    modelId?: string;
-    name: string;
-    address?: string;
-    enabled: boolean;
-    maintenance: boolean;
-    notes?: string;
+   // if defined, existing device
+   id?: string;
+   platformId: string;
+   poolId: string;
+   modelId?: string;
+   name: string;
+   address?: string;
+   enabled: boolean;
+   maintenance: boolean;
+   notes?: string;
 };
 
 // in sort order
 export enum DeviceStatus {
-    Available = 0,
-    Reserved = 1,
-    Maintenance = 2,
-    Problem = 3,
-    Disabled = 4
+   Available = 0,
+   Reserved = 1,
+   Maintenance = 2,
+   Problem = 3,
+   Disabled = 4
 }
 
 export class DeviceHandler extends PollBase {
 
-    constructor(pollTime = 5000) {
-        super(pollTime);
-    }
+   constructor(pollTime = 5000) {
+      super(pollTime);
+   }
 
-    clear() {
-        super.stop();
-        this.devices = [];
-        this.reservations = [];
-        this.pools.clear();
-        this.platforms.clear();
-    }
+   clear() {
+      super.stop();
+      this.devices = [];
+      this.reservations = [];
+      this.pools.clear();
+      this.platforms.clear();
+   }
 
-    async poll(): Promise<void> {
+   async poll(): Promise<void> {
 
-        try {
+      try {
 
-            this.devices = await backend.getDevices();
-            this.reservations = await backend.getDeviceReservations();
+         this.devices = await backend.getDevices();
+         this.reservations = await backend.getDeviceReservations();
 
-            // filter out utilization to entries which contain job id (legacy reservations don't)
-            this.devices.forEach(d => d.utilization = d.utilization?.filter(u => !!u.jobId));
+         // filter out utilization to entries which contain job id (legacy reservations don't)
+         this.devices.forEach(d => d.utilization = d.utilization?.filter(u => !!u.jobId));
 
-            const userIds = new Set<string>();
-            this.devices.forEach(d => {
-                if (d.checkedOutByUserId && !this.users.has(d.checkedOutByUserId)) {
-                    userIds.add(d.checkedOutByUserId);
-                }
-            });
-
-            if (!this.users.get(dashboard.userId)) {
-                userIds.add(dashboard.userId);
+         const userIds = new Set<string>();
+         this.devices.forEach(d => {
+            if (d.checkedOutByUserId && !this.users.has(d.checkedOutByUserId)) {
+               userIds.add(d.checkedOutByUserId);
             }
+         });
+
+         if (!this.users.get(dashboard.userId)) {
+            userIds.add(dashboard.userId);
+         }
+
+         if (userIds.size) {
+            const users = await backend.getUsers({ ids: Array.from(userIds) });
+            users.forEach(u => this.users.set(u.id, u));
+         }
+
+         if (!this.devices.length || this.devices.find(d => !this.platforms.has(d.platformId))) {
+            // we need to get platforms
+            const platforms = await backend.getDevicePlatforms();
+            this.platforms.clear();
+            platforms.forEach(p => {
+               this.platforms.set(p.id, p);
+            });
 
-            if (userIds.size) {
-                const users = await backend.getUsers({ ids: Array.from(userIds) });
-                users.forEach(u => this.users.set(u.id, u));
-            }
+            const pools = await backend.getDevicePools();
+            this.pools.clear();
+            pools.forEach(p => {
+               this.pools.set(p.id, p);
+            });
 
-            if (!this.devices.length || this.devices.find(d => !this.platforms.has(d.platformId))) {
-                // we need to get platforms
-                const platforms = await backend.getDevicePlatforms();
-                this.platforms.clear();
-                platforms.forEach(p => {
-                    this.platforms.set(p.id, p);
-                });
+         }
 
-                const pools = await backend.getDevicePools();
-                this.pools.clear();
-                pools.forEach(p => {
-                    this.pools.set(p.id, p);
-                });
+         this.setUpdated();
 
-            }
+      } catch (err) {
 
-            this.setUpdated();
+      }
 
-        } catch (err) {
+   }
 
-        }
+   getUserDeviceCheckouts(userId: string): GetDeviceResponse[] {
+      return this.devices.filter(d => d.checkedOutByUserId === userId);
+   }
 
-    }
-    
-    getUserDeviceCheckouts(userId: string): GetDeviceResponse[] {
-        return this.devices.filter(d => d.checkedOutByUserId === userId);
-    }
- 
-    getReservation(device: GetDeviceResponse): GetDeviceReservationResponse | undefined {
-        return this.reservations.find(r => !!r.devices.find(d => d === device.id));
+   getReservation(device: GetDeviceResponse): GetDeviceReservationResponse | undefined {
+      return this.reservations.find(r => !!r.devices.find(d => d === device.id));
    }
-   
+
    getDeviceWriteAccess(device: GetDeviceResponse): boolean {
 
       const pool = this.pools.get(device.poolId);
@@ -113,397 +113,395 @@ export class DeviceHandler extends PollBase {
 
    }
 
-    getDeviceStatus(device: GetDeviceResponse): DeviceStatus {
+   getDeviceStatus(device: GetDeviceResponse): DeviceStatus {
 
-        if (this.getReservation(device)) {
-            return DeviceStatus.Reserved;
-        }
+      if (this.getReservation(device)) {
+         return DeviceStatus.Reserved;
+      }
 
-        if (!device.enabled) {
-            return DeviceStatus.Disabled;
-        }
+      if (!device.enabled) {
+         return DeviceStatus.Disabled;
+      }
 
-        if (device.maintenanceTime) {
-            return DeviceStatus.Maintenance;
-        }
+      if (device.maintenanceTime) {
+         return DeviceStatus.Maintenance;
+      }
 
-        if (device.problemTime) {
+      if (device.problemTime) {
 
-            const end = moment.utc();
-            const d = moment.duration(end.diff(moment(device.problemTime)));
+         const end = moment.utc();
+         const d = moment.duration(end.diff(moment(device.problemTime)));
 
-            // note this must match reservation selection for problem time in backend
-            if (d.asMinutes() < dashboard.deviceProblemCooldownMinutes) {
-                return DeviceStatus.Problem;
-            }
-        }
+         // note this must match reservation selection for problem time in backend
+         if (d.asMinutes() < dashboard.deviceProblemCooldownMinutes) {
+            return DeviceStatus.Problem;
+         }
+      }
 
-        return DeviceStatus.Available;
-    }
+      return DeviceStatus.Available;
+   }
 
-    getDevices(platformId?: string) {
-        return this.devices.filter(d => !platformId || d.platformId === platformId);
-    }
+   getDevices(platformId?: string) {
+      return this.devices.filter(d => !platformId || d.platformId === platformId);
+   }
 
 
-    private devices: GetDeviceResponse[] = [];
+   private devices: GetDeviceResponse[] = [];
 
-    platforms: Map<string, GetDevicePlatformResponse> = new Map();
+   platforms: Map<string, GetDevicePlatformResponse> = new Map();
 
-    pools: Map<string, GetDevicePoolResponse> = new Map();
+   pools: Map<string, GetDevicePoolResponse> = new Map();
 
-    users: Map<string, GetUserResponse> = new Map();
+   users: Map<string, GetUserResponse> = new Map();
 
-    private reservations: GetDeviceReservationResponse[] = [];
+   private reservations: GetDeviceReservationResponse[] = [];
 }
 
 
 export const DeviceEditor: React.FC<{ handler: DeviceHandler, deviceIn?: GetDeviceResponse | undefined, editNote?: boolean, onClose: () => void }> = ({ handler, deviceIn, editNote, onClose }) => {
 
-    const [state, setState] = useState<{ device?: DeviceEditData, title?: string }>({});
-    const [error, setError] = useState<string | undefined>();
-    const [submitting, setSubmitting] = useState(false);
-    const [confirmDelete, setConfirmDelete] = useState(false);
-
-    if (submitting) {
-
-        return <Modal isOpen={true} isBlocking={true} styles={{ main: { padding: 8, width: 400 } }} >
-            <Stack style={{ paddingTop: 32 }}>
-                <Stack tokens={{ childrenGap: 24 }} styles={{ root: { padding: 8 } }}>
-                    <Stack horizontalAlign="center">
-                        <Text variant="large">Please wait...</Text>
-                    </Stack>
-                    <Stack verticalAlign="center" style={{ paddingBottom: 32 }}>
-                        <Spinner size={SpinnerSize.large} />
-                    </Stack>
-                </Stack>
+   const [state, setState] = useState<{ device?: DeviceEditData, title?: string }>({});
+   const [error, setError] = useState<string | undefined>();
+   const [submitting, setSubmitting] = useState(false);
+   const [confirmDelete, setConfirmDelete] = useState(false);
+
+   if (submitting) {
+
+      return <Modal isOpen={true} isBlocking={true} styles={{ main: { padding: 8, width: 400 } }} >
+         <Stack style={{ paddingTop: 32 }}>
+            <Stack tokens={{ childrenGap: 24 }} styles={{ root: { padding: 8 } }}>
+               <Stack horizontalAlign="center">
+                  <Text variant="large">Please wait...</Text>
+               </Stack>
+               <Stack verticalAlign="center" style={{ paddingBottom: 32 }}>
+                  <Spinner size={SpinnerSize.large} />
+               </Stack>
             </Stack>
-        </Modal>
+         </Stack>
+      </Modal>
 
-    }
+   }
 
-    const platforms = handler.platforms;
-    const pools = handler.pools;
+   const platforms = handler.platforms;
+   const pools = handler.pools;
 
-    const device = state.device!;
-    const existing = !!deviceIn;
-
-    let deviceStatus = DeviceStatus.Available;
-
-    if (deviceIn) {
-        deviceStatus = handler.getDeviceStatus(deviceIn);
-    }
-
-    const reservation = deviceStatus === DeviceStatus.Reserved;
-
-    const close = () => {
-        // @todo: state should be squashed
-        setState({});
-        setSubmitting(false);
-        setError("");
-        setConfirmDelete(false);
-        onClose();
-    };
-
-    const onDelete = async () => {
-
-        try {
-
-            setSubmitting(true);
-            await backend.deleteDevice(device.id!);
-            setSubmitting(false);
-            close();
-
-        } catch (reason:any) {
-            setError(reason.toString());
-            setSubmitting(false);
-       }
-       
-       await handler.forceUpdate();
-
-    }
-
-
-    if (confirmDelete) {
-        return <Dialog
-            hidden={false}
-            onDismiss={() => setConfirmDelete(false)}
-            minWidth={400}
-            dialogContentProps={{
-                type: DialogType.normal,
-                title: `Delete device ${state.device?.name} ?`,
-            }}
-            modalProps={{ isBlocking: true }} >
-            <DialogFooter>
-                <PrimaryButton onClick={() => { setConfirmDelete(false); onDelete() }} text="Delete" />
-                <DefaultButton onClick={() => setConfirmDelete(false)} text="Cancel" />
-            </DialogFooter>
-        </Dialog>
-    }
-
-    const platformOptions = Array.from(platforms.values()).sort((a, b) => {
-
-        return a.name.localeCompare(b.name);
-
-    }).map(p => {
-        return { key: `platform_${p.id}`, text: p.name, selected: device?.platformId === p.id, data: p.id } as IDropdownOption
-    });
-
-    let modelOptions: IDropdownOption[] = [];
-
-    const platform = platforms.get(device?.platformId);
-
-    if (platform) {
-        modelOptions = platform.modelIds.sort((a, b) => {
-            return a.localeCompare(b);
-        }).map(model => {
-            return { key: `model_${model}`, text: model, selected: device?.modelId === model, data: model } as IDropdownOption
-        });
-    }
-
-    modelOptions.unshift({ key: `model_base`, text: "Base", selected: device?.modelId === undefined, data: undefined });
-
-    let defaultModelKey = "model_base";
-    if (device?.modelId) {
-        defaultModelKey = `model_${device.modelId}`;
-    }
-
-
-    const poolOptions: IDropdownOption[] = [];
-    pools.forEach(p => {
-        poolOptions.push({ key: `pool_${p.id}`, text: p.name, selected: device?.poolId === p.id, data: p.id });
-    });
-
-    if (!state.device) {
-
-        if (deviceIn) {
-            setState({
-                device: {
-                    id: deviceIn?.id,
-                    platformId: deviceIn.platformId,
-                    enabled: deviceIn.enabled,
-                    address: deviceIn.address,
-                    name: deviceIn.name,
-                    poolId: deviceIn.poolId,
-                    modelId: deviceIn.modelId,
-                    notes: deviceIn.notes,
-                    maintenance: !!deviceIn.maintenanceTime
-                }, title: `Edit ${deviceIn.name}`
-            });
-        } else {
-            setState({
-                device: {
-                    enabled: true,
-                    maintenance: false,
-                    address: "",
-                    name: "",
-                    platformId: platformOptions.length > 0 ? platformOptions[0].data : "",
-                    poolId: poolOptions.length > 0 ? poolOptions[0].data : ""
-                },
-                title: "Add New Device"
-            });
-        }
+   const device = state.device!;
+   const existing = !!deviceIn;
+
+   let deviceStatus = DeviceStatus.Available;
+
+   if (deviceIn) {
+      deviceStatus = handler.getDeviceStatus(deviceIn);
+   }
 
-        return null;
-    }
+   const reservation = deviceStatus === DeviceStatus.Reserved;
 
+   const close = async () => {
+      // @todo: state should be squashed
 
-    const statusOptions: IDropdownOption[] = [];
+      await handler.forceUpdate();
 
-    if (reservation) {
-        statusOptions.push({ key: 'status_reserved', text: "Reserved", selected: true, data: "Reserved" });
-    } else {
+      setState({});
+      setSubmitting(false);
+      setError("");
+      setConfirmDelete(false);
+      onClose();
+   };
 
-        statusOptions.push({ key: 'status_available', text: "Available", selected: device.enabled && !device.maintenance, data: "Available" });
-        statusOptions.push({ key: 'status_disabled', text: "Disabled", selected: !device.enabled, data: "Disabled" });
-        statusOptions.push({ key: 'status_maintenance', text: "Maintenance", selected: device.enabled && device.maintenance, data: "Maintenance" });
-    };
+   const onDelete = async () => {
 
-    // need to handle delete reservation case
-    const onSave = async () => {
+      try {
 
-        if (!device.name) {
-            setError("Device must have a name");
-            return;
-        }
+         setSubmitting(true);
+         await backend.deleteDevice(device.id!);
+         await close();
+
+      } catch (reason: any) {
+         setError(reason.toString());
+         setSubmitting(false);
+      }
+
+   }
 
-        if (!device.address) {
-            setError("Device must have an address");
-            return;
-        }
 
-        setSubmitting(true);
+   if (confirmDelete) {
+      return <Dialog
+         hidden={false}
+         onDismiss={() => setConfirmDelete(false)}
+         minWidth={400}
+         dialogContentProps={{
+            type: DialogType.normal,
+            title: `Delete device ${state.device?.name} ?`,
+         }}
+         modalProps={{ isBlocking: true }} >
+         <DialogFooter>
+            <PrimaryButton onClick={() => { setConfirmDelete(false); onDelete() }} text="Delete" />
+            <DefaultButton onClick={() => setConfirmDelete(false)} text="Cancel" />
+         </DialogFooter>
+      </Dialog>
+   }
+
+   const platformOptions = Array.from(platforms.values()).sort((a, b) => {
 
-        try {
+      return a.name.localeCompare(b.name);
 
-            if (reservation) {
+   }).map(p => {
+      return { key: `platform_${p.id}`, text: p.name, selected: device?.platformId === p.id, data: p.id } as IDropdownOption
+   });
 
-                if (editNote) {
+   let modelOptions: IDropdownOption[] = [];
 
-                    await backend.modifyDevice(device.id!, {
-                        notes: device.notes ? device.notes.trim() : ""
-                    }).then(() => {
-                        close();
-                    }).catch((reason) => {
-                        setError(`Problem modifying device, ${reason}`);
-                    });
+   const platform = platforms.get(device?.platformId);
 
-                } else {
+   if (platform) {
+      modelOptions = platform.modelIds.sort((a, b) => {
+         return a.localeCompare(b);
+      }).map(model => {
+         return { key: `model_${model}`, text: model, selected: device?.modelId === model, data: model } as IDropdownOption
+      });
+   }
 
-                    await backend.modifyDevice(device.id!, {
-                        problem: false,
-                        maintenance: device.maintenance,
-                        enabled: true
-                    }).then(() => {
-                        close();
-                    }).catch((reason) => {
-                        setError(`Problem modifying device, ${reason}`);
-                    });
-                }
+   modelOptions.unshift({ key: `model_base`, text: "Base", selected: device?.modelId === undefined, data: undefined });
 
+   let defaultModelKey = "model_base";
+   if (device?.modelId) {
+      defaultModelKey = `model_${device.modelId}`;
+   }
+
+
+   const poolOptions: IDropdownOption[] = [];
+   pools.forEach(p => {
+      poolOptions.push({ key: `pool_${p.id}`, text: p.name, selected: device?.poolId === p.id, data: p.id });
+   });
+
+   if (!state.device) {
+
+      if (deviceIn) {
+         setState({
+            device: {
+               id: deviceIn?.id,
+               platformId: deviceIn.platformId,
+               enabled: deviceIn.enabled,
+               address: deviceIn.address,
+               name: deviceIn.name,
+               poolId: deviceIn.poolId,
+               modelId: deviceIn.modelId,
+               notes: deviceIn.notes,
+               maintenance: !!deviceIn.maintenanceTime
+            }, title: `Edit ${deviceIn.name}`
+         });
+      } else {
+         setState({
+            device: {
+               enabled: true,
+               maintenance: false,
+               address: "",
+               name: "",
+               platformId: platformOptions.length > 0 ? platformOptions[0].data : "",
+               poolId: poolOptions.length > 0 ? poolOptions[0].data : ""
+            },
+            title: "Add New Device"
+         });
+      }
+
+      return null;
+   }
+
+
+   const statusOptions: IDropdownOption[] = [];
+
+   if (reservation) {
+      statusOptions.push({ key: 'status_reserved', text: "Reserved", selected: true, data: "Reserved" });
+   } else {
+
+      statusOptions.push({ key: 'status_available', text: "Available", selected: device.enabled && !device.maintenance, data: "Available" });
+      statusOptions.push({ key: 'status_disabled', text: "Disabled", selected: !device.enabled, data: "Disabled" });
+      statusOptions.push({ key: 'status_maintenance', text: "Maintenance", selected: device.enabled && device.maintenance, data: "Maintenance" });
+   };
+
+   // need to handle delete reservation case
+   const onSave = async () => {
+
+      if (!device.name) {
+         setError("Device must have a name");
+         return;
+      }
+
+      if (!device.address) {
+         setError("Device must have an address");
+         return;
+      }
+
+      setSubmitting(true);
+
+      try {
+
+         if (reservation) {
+
+            if (editNote) {
+
+               await backend.modifyDevice(device.id!, {
+                  notes: device.notes ? device.notes.trim() : ""
+               }).then(() => {
+                  close();
+               }).catch((reason) => {
+                  setError(`Problem modifying device, ${reason}`);
+               });
+
+            } else {
+
+               await backend.modifyDevice(device.id!, {
+                  problem: false,
+                  maintenance: device.maintenance,
+                  enabled: true
+               }).then(() => {
+                  close();
+               }).catch((reason) => {
+                  setError(`Problem modifying device, ${reason}`);
+               });
             }
-            else if (!existing) {
-                const result = await backend.addDevice({
-                    name: device.name.trim(),
-                    address: device.address,
-                    poolId: device.poolId,
-                    platformId: device.platformId,
-                    modelId: device.modelId
-                });
-
-                if (!result) {
-                    setError("Problem adding device");
-                } else {
-                    close();
-                }
+
+         }
+         else if (!existing) {
+            const result = await backend.addDevice({
+               name: device.name.trim(),
+               address: device.address,
+               poolId: device.poolId,
+               platformId: device.platformId,
+               modelId: device.modelId
+            });
+
+            if (!result) {
+               setError("Problem adding device");
             } else {
+               close();
+            }
+         } else {
+
+            await backend.modifyDevice(device.id!, {
+               name: device.name.trim(),
+               address: device.address,
+               poolId: device.poolId,
+               modelId: device.modelId,
+               notes: device.notes,
+               problem: false, // always clear automatically generated problem state when saving
+               maintenance: device.maintenance,
+               enabled: device.enabled
+            }).then(() => {
+               close();
+            }).catch((reason) => {
+               setError(`Problem modifying device,: ${reason}`);
+               setSubmitting(false);
+            })
+         }
+         
+      } catch (reason: any) {
+         setError(reason.toString());
+         setSubmitting(false);
+      }
 
-                await backend.modifyDevice(device.id!, {
-                    name: device.name.trim(),
-                    address: device.address,
-                    poolId: device.poolId,
-                    modelId: device.modelId,
-                    notes: device.notes,
-                    problem: false, // always clear automatically generated problem state when saving
-                    maintenance: device.maintenance,
-                    enabled: device.enabled
-                }).then(() => {
-                    close();
-                }).catch((reason) => {
-                    setError(`Problem modifying device,: ${reason}`);
-                    setSubmitting(false);
-                })
-           }
-           
-           await handler.forceUpdate();
-
-        } catch (reason:any) {
-            setError(reason.toString());
-            setSubmitting(false);
-        }
-
-    }
-
-    return <Modal className={hordeClasses.modal} isOpen={true} isBlocking={true} styles={{ main: { padding: 8, width: 700 } }}>
-        {!!error && <MessageBar
-            messageBarType={MessageBarType.error}
-            isMultiline={false}
-            onDismiss={() => setError("")}
-        >
-            <Text>{error}</Text>
-        </MessageBar>}
-
-        <Stack style={{ padding: 8 }}>
-            <Stack style={{ paddingBottom: 16 }}>
-                <Text variant="large" style={{ fontFamily: "Horde Open Sans SemiBold" }}>{state.title}</Text>
-            </Stack>
-            <Stack tokens={{ childrenGap: 8 }} style={{ padding: 8 }}>
-                <TextField label="Name" disabled={existing} defaultValue={device.name} onChange={(ev, value) => {
-                    device.name = value ?? "";
-                }} />
-                <Dropdown label="Platform" disabled={existing} options={platformOptions} onChange={(ev, option) => {
+   }
 
-                    if (option) {
-                        device.platformId = option.data as string;
-                        device.modelId = undefined;
+   return <Modal className={hordeClasses.modal} isOpen={true} isBlocking={true} styles={{ main: { padding: 8, width: 700 } }}>
+      {!!error && <MessageBar
+         messageBarType={MessageBarType.error}
+         isMultiline={false}
+         onDismiss={() => setError("")}
+      >
+         <Text>{error}</Text>
+      </MessageBar>}
 
-                        setState({ device: device, title: state.title });
-                    }
+      <Stack style={{ padding: 8 }}>
+         <Stack style={{ paddingBottom: 16 }}>
+            <Text variant="large" style={{ fontFamily: "Horde Open Sans SemiBold" }}>{state.title}</Text>
+         </Stack>
+         <Stack tokens={{ childrenGap: 8 }} style={{ padding: 8 }}>
+            <TextField label="Name" disabled={existing} defaultValue={device.name} onChange={(ev, value) => {
+               device.name = value ?? "";
+            }} />
+            <Dropdown label="Platform" disabled={existing} options={platformOptions} onChange={(ev, option) => {
 
-                }} />
+               if (option) {
+                  device.platformId = option.data as string;
+                  device.modelId = undefined;
 
+                  setState({ device: device, title: state.title });
+               }
 
-                {(!editNote && !reservation) && <Dropdown label="Pool" options={poolOptions} onChange={(ev, option) => {
+            }} />
 
-                    if (option) {
-                        device.poolId = option.data as string;
-                    }
 
-                }} />}
+            {(!editNote && !reservation) && <Dropdown label="Pool" options={poolOptions} onChange={(ev, option) => {
 
-                {(!editNote && !reservation) && <Dropdown label="Model" defaultSelectedKey={defaultModelKey} options={modelOptions} onChange={(ev, option) => {
+               if (option) {
+                  device.poolId = option.data as string;
+               }
 
-                    if (option?.data === undefined) {
-                        device.modelId = undefined; // "Base"
-                    }
-                    else if (option) {
-                        device.modelId = option.data as string;
-                    }
+            }} />}
 
-                }} />}
+            {(!editNote && !reservation) && <Dropdown label="Model" defaultSelectedKey={defaultModelKey} options={modelOptions} onChange={(ev, option) => {
 
+               if (option?.data === undefined) {
+                  device.modelId = undefined; // "Base"
+               }
+               else if (option) {
+                  device.modelId = option.data as string;
+               }
 
-                <TextField label="Address" defaultValue={device.address} disabled={(editNote || reservation)} onChange={(ev, value) => {
-                    device.address = value ?? "";
-                }} />
+            }} />}
 
-                {(!editNote && !reservation) && <Dropdown label="Status" options={statusOptions} onChange={(ev, option) => {
 
-                    if (option) {
+            <TextField label="Address" defaultValue={device.address} disabled={(editNote || reservation)} onChange={(ev, value) => {
+               device.address = value ?? "";
+            }} />
 
-                        if (option.data === "Available") {
-                            device.maintenance = false;
-                            device.enabled = true;
-                        }
+            {(!editNote && !reservation) && <Dropdown label="Status" options={statusOptions} onChange={(ev, option) => {
 
-                        if (option.data === "Maintenance") {
-                            device.maintenance = true;
-                            device.enabled = true;
-                        }
+               if (option) {
 
-                        if (option.data === "Disabled") {
-                            device.maintenance = false;
-                            device.enabled = false;
-                        }
+                  if (option.data === "Available") {
+                     device.maintenance = false;
+                     device.enabled = true;
+                  }
 
-                    }
+                  if (option.data === "Maintenance") {
+                     device.maintenance = true;
+                     device.enabled = true;
+                  }
 
-                }} />}
+                  if (option.data === "Disabled") {
+                     device.maintenance = false;
+                     device.enabled = false;
+                  }
 
-                {!!editNote && <TextField label="Notes" defaultValue={device.notes} multiline rows={5} resizable={false} onChange={(ev, newValue) => {
-                    device.notes = newValue;
-                }} />}
+               }
 
-                {!!reservation && !editNote && <Stack style={{ paddingTop: 24 }}><Checkbox label="Disable device for maintenance once reservation finishes"
-                    checked={!!device.maintenance}
-                    onChange={(ev, checked) => {
-                        device.maintenance = checked ? true : false;
-                        setState({ device: device, title: state.title });
-                    }} />
-                </Stack>}
+            }} />}
 
+            {!!editNote && <TextField label="Notes" defaultValue={device.notes} multiline rows={5} resizable={false} onChange={(ev, newValue) => {
+               device.notes = newValue;
+            }} />}
 
-            </Stack>
+            {!!reservation && !editNote && <Stack style={{ paddingTop: 24 }}><Checkbox label="Disable device for maintenance once reservation finishes"
+               checked={!!device.maintenance}
+               onChange={(ev, checked) => {
+                  device.maintenance = checked ? true : false;
+                  setState({ device: device, title: state.title });
+               }} />
+            </Stack>}
+
+
+         </Stack>
 
-            <Stack horizontal style={{ paddingTop: 64 }}>
-                {!!existing && !editNote && !reservation && <PrimaryButton style={{ backgroundColor: "red", border: 0 }} onClick={() => setConfirmDelete(true)} text="Delete Device" />}
-                <Stack grow />
-                <Stack horizontal tokens={{ childrenGap: 28 }}>
-                    <PrimaryButton onClick={() => onSave()} text="Save" />
-                    <DefaultButton onClick={() => close()} text="Cancel" />
-                </Stack>
+         <Stack horizontal style={{ paddingTop: 64 }}>
+            {!!existing && !editNote && !reservation && <PrimaryButton style={{ backgroundColor: "red", border: 0 }} onClick={() => setConfirmDelete(true)} text="Delete Device" />}
+            <Stack grow />
+            <Stack horizontal tokens={{ childrenGap: 28 }}>
+               <PrimaryButton onClick={() => onSave()} text="Save" />
+               <DefaultButton onClick={() => close()} text="Cancel" />
             </Stack>
-        </Stack>
-    </Modal>
+         </Stack>
+      </Stack>
+   </Modal>
 
 };
\ No newline at end of file
