Horde: Fix issue with trends panel rendering without history data (and a stuck spinner)
#skipci

diff --git a/Engine/Source/Programs/Horde/HordeDashboard/src/components/jobDetailsV2/JobDetailStepTrendsV2.tsx b/Engine/Source/Programs/Horde/HordeDashboard/src/components/jobDetailsV2/JobDetailStepTrendsV2.tsx
--- a/Engine/Source/Programs/Horde/HordeDashboard/src/components/jobDetailsV2/JobDetailStepTrendsV2.tsx
+++ b/Engine/Source/Programs/Horde/HordeDashboard/src/components/jobDetailsV2/JobDetailStepTrendsV2.tsx
@@ -717,6 +717,10 @@ export const StepTrendsPanelV2: React.FC<{ jobDetails: JobDetailsV2; stepId: str
       return null;
    }
 
+   if (dataView.initialized && !dataView.history?.length) {
+      return null;
+   }
+
    dataView.set(stepId);
 
    return (<Stack id={sideRail.url} styles={{ root: { paddingTop: 18, paddingRight: 12 } }}>
