Horde: Add filter to (legacy) step artifacts in log modal
#skipci
#rnx

diff --git a/Engine/Source/Programs/Horde/HordeDashboard/src/components/JobDetailArtifacts.tsx b/Engine/Source/Programs/Horde/HordeDashboard/src/components/JobDetailArtifacts.tsx
--- a/Engine/Source/Programs/Horde/HordeDashboard/src/components/JobDetailArtifacts.tsx
+++ b/Engine/Source/Programs/Horde/HordeDashboard/src/components/JobDetailArtifacts.tsx
@@ -1,7 +1,7 @@
 // Copyright Epic Games, Inc. All Rights Reserved.
 
-import { Stack, Text, IColumn, mergeStyleSets, Icon, DetailsList, Selection, SelectionMode, DetailsListLayoutMode, ScrollablePane, ScrollbarVisibility, StickyPositionType, IDetailsListProps, IDetailsHeaderStyles, Sticky, DetailsHeader, PrimaryButton, SpinnerSize, Spinner, Link } from '@fluentui/react';
-import React from 'react';
+import { Stack, Text, IColumn, mergeStyleSets, Icon, DetailsList, Selection, SelectionMode, DetailsListLayoutMode, ScrollablePane, ScrollbarVisibility, StickyPositionType, IDetailsListProps, IDetailsHeaderStyles, Sticky, DetailsHeader, PrimaryButton, SpinnerSize, Spinner, Link, TextField } from '@fluentui/react';
+import React, { useState } from 'react';
 import { ArtifactData, GetArtifactZipRequest } from '../backend/Api';
 import { JobDetails } from '../backend/JobDetails';
 import { hordeClasses } from '../styles/Styles';
@@ -149,6 +149,9 @@ class ArtifactState {
 
 const artifactState = new ArtifactState();
 export const JobDetailArtifacts: React.FC<{ jobDetails: JobDetails; stepId?: string, topPadding?: number }> = observer(({ jobDetails, stepId, topPadding }) => {
+
+   const [state, setState] = useState<{ filter?: string }>({});
+
    artifactState.setDetails(jobDetails, stepId);
 
    function getKey(item: any, index?: number): string {
@@ -183,6 +186,12 @@ export const JobDetailArtifacts: React.FC<{ jobDetails: JobDetails; stepId?: str
       artifacts = artifacts.filter(artifact => artifact.stepId === stepId);
    }
 
+   if (state.filter) {
+      const f = state.filter.toLowerCase();
+      artifacts = artifacts.filter(a => a.name?.toLowerCase().indexOf(f) !== -1);
+   }
+
+
    let height = Math.min(36 * artifacts.length + 60, 500);
 
    return (<Stack styles={{ root: { paddingTop: topPadding ?? 18, paddingRight: 12 } }}>
@@ -190,9 +199,39 @@ export const JobDetailArtifacts: React.FC<{ jobDetails: JobDetails; stepId?: str
          <Stack tokens={{ childrenGap: 12 }}>
             <Stack horizontal horizontalAlign="space-between" styles={{ root: { minHeight: 32 } }}>
                <Text variant="mediumPlus" styles={{ root: { fontFamily: "Horde Open Sans SemiBold" } }}>Artifacts</Text>
-               <Stack horizontal>
-                  {artifactState.isDownloading && <Spinner styles={{ root: { marginRight: 10 } }} size={SpinnerSize.medium}></Spinner>}
-                  <PrimaryButton styles={{ root: { fontFamily: 'Horde Open Sans Semibold !important' } }} disabled={artifactState.isDownloading} onClick={artifactState.downloadItems.bind(artifactState, true, false, undefined)}>{buttonText}</PrimaryButton>
+               <Stack horizontal tokens={{ childrenGap: 24 }}>
+                  <Stack>
+                     <TextField
+                        spellCheck={false}
+                        deferredValidationTime={500}
+                        validateOnLoad={false}
+
+                        styles={{
+                           root: { width: 320, fontSize: 12 }, fieldGroup: {
+                              borderWidth: 1
+                           }
+                        }}
+
+                        placeholder="Filter"
+
+                        onGetErrorMessage={(newValue) => {
+
+                           const filter = newValue ? newValue : undefined;
+                           if (filter !== state.filter) {
+                              setState({ ...state, filter: filter });
+                           }
+
+                           return undefined;
+                        }}
+
+                     />
+
+                  </Stack>
+
+                  <Stack>
+                     {artifactState.isDownloading && <Spinner styles={{ root: { marginRight: 10 } }} size={SpinnerSize.medium}></Spinner>}
+                     <PrimaryButton styles={{ root: { fontFamily: 'Horde Open Sans Semibold !important' } }} disabled={artifactState.isDownloading} onClick={artifactState.downloadItems.bind(artifactState, true, false, undefined)}>{buttonText}</PrimaryButton>
+                  </Stack>
                </Stack>
             </Stack>
             <Stack>
