Horde: Reduce timeout for remote execution agents to 20 seconds. Pings should be sent every 5 seconds.

diff --git a/Engine/Source/Programs/Horde/Horde.Agent/Leases/Handlers/ComputeHandler.cs b/Engine/Source/Programs/Horde/Horde.Agent/Leases/Handlers/ComputeHandler.cs
--- a/Engine/Source/Programs/Horde/Horde.Agent/Leases/Handlers/ComputeHandler.cs
+++ b/Engine/Source/Programs/Horde/Horde.Agent/Leases/Handlers/ComputeHandler.cs
@@ -58,7 +58,7 @@ namespace Horde.Agent.Leases.Handlers
 			}
 		}
 
-		const int NoDataTimeoutMinutes = 10;
+		static TimeSpan NoDataTimeout { get; } = TimeSpan.FromSeconds(20);
 
 		readonly ComputeListenerService _listenerService;
 		readonly IMemoryCache _memoryCache;
@@ -135,9 +135,9 @@ namespace Horde.Agent.Leases.Handlers
 							}
 						}
 					}
-					catch (OperationCanceledException ex) when (cts.IsCancellationRequested && transport.TimeSinceActivity > TimeSpan.FromMinutes(NoDataTimeoutMinutes))
+					catch (OperationCanceledException ex) when (cts.IsCancellationRequested && transport.TimeSinceActivity > NoDataTimeout)
 					{
-						logger.LogError(ex, "Lease was terminated due to no data being received for {Time} minutes", NoDataTimeoutMinutes);
+						logger.LogError(ex, "Lease was terminated due to no data being received for {Time} seconds", (int)NoDataTimeout.TotalSeconds);
 						return LeaseResult.Failed;
 					}
 				}
@@ -157,7 +157,7 @@ namespace Horde.Agent.Leases.Handlers
 		{
 			while(!cancellationToken.IsCancellationRequested)
 			{
-				if (transport.TimeSinceActivity > TimeSpan.FromMinutes(NoDataTimeoutMinutes))
+				if (transport.TimeSinceActivity > NoDataTimeout)
 				{
 					logger.LogWarning("Terminating compute task due to timeout (last tick at {Time})", DateTime.UtcNow - transport.TimeSinceActivity);
 					cts.Cancel();
