Horde: Fix default port for agent connections.

diff --git a/Engine/Source/Programs/Horde/Horde.Agent/appsettings.json b/Engine/Source/Programs/Horde/Horde.Agent/appsettings.json
--- a/Engine/Source/Programs/Horde/Horde.Agent/appsettings.json
+++ b/Engine/Source/Programs/Horde/Horde.Agent/appsettings.json
@@ -7,7 +7,7 @@
 			{
 				"Name": "Local",
 				"Environment": "dev",
-				"Url": "https://localhost:5001",
+				"Url": "http://localhost:5000",
 				"Storage": {
 					"Url": "https://localhost:57001"
 				}
