Horde: Add timestamp tooltip on log issue buttons
#rnx
#skipci

diff --git a/Engine/Source/Programs/Horde/HordeDashboard/src/components/LogView.tsx b/Engine/Source/Programs/Horde/HordeDashboard/src/components/LogView.tsx
--- a/Engine/Source/Programs/Horde/HordeDashboard/src/components/LogView.tsx
+++ b/Engine/Source/Programs/Horde/HordeDashboard/src/components/LogView.tsx
@@ -1,10 +1,10 @@
 // Copyright Epic Games, Inc. All Rights Reserved.
 
-import { DefaultButton, DetailsList, DetailsListLayoutMode, DetailsRow, DirectionalHint, Dropdown, FocusZone, FocusZoneDirection, IColumn, Icon, IconButton, IContextualMenuItem, IContextualMenuProps, IDetailsListProps, ITextField, List, Modal, ProgressIndicator, ScrollToMode, Selection, SelectionMode, SelectionZone, Separator, Spinner, SpinnerSize, Stack, Text, TextField } from '@fluentui/react';
+import { DefaultButton, DetailsList, DetailsListLayoutMode, DetailsRow, DirectionalHint, Dropdown, FocusZone, FocusZoneDirection, IColumn, Icon, IconButton, IContextualMenuItem, IContextualMenuProps, IDetailsListProps, ITextField, List, Modal, ProgressIndicator, ScrollToMode, Selection, SelectionMode, SelectionZone, Separator, Spinner, SpinnerSize, Stack, Text, TextField, TooltipHost } from '@fluentui/react';
 import { action, makeObservable, observable } from 'mobx';
 import { observer } from 'mobx-react-lite';
 import moment from 'moment-timezone';
-import React, { useEffect, useState } from 'react';
+import React, { useEffect, useId, useState } from 'react';
 import { Link, useNavigate, useLocation, useParams } from 'react-router-dom';
 import backend from '../backend';
 import { ArtifactContextType, ArtifactData, EventSeverity, GetChangeSummaryResponse, GetJobStepRefResponse, GetLogEventResponse, LogLevel } from '../backend/Api';
@@ -633,6 +633,8 @@ export const LogList: React.FC<{ logId: string }> = observer(({ logId }) => {
 
          const IssueButton: React.FC<{ item: LogItem, event: GetLogEventResponse }> = ({ item, event }) => {
 
+            const tooltipId = useId();
+
             const error = event.severity === EventSeverity.Error;
 
             const issueId = item!.issueId!.toString();
@@ -648,17 +650,25 @@ export const LogList: React.FC<{ logId: string }> = observer(({ logId }) => {
             }
 
             return <Stack>
-               <DefaultButton className={error ? styles.errorButton : styles.warningButton}
-                  href={href}
-                  style={{ padding: 0, margin: 0, width: tsWidth, paddingLeft: 8, paddingRight: 8, height: "100%", fontWeight: "unset" }}
-
-                  onClick={(ev) => {
-                     ev.preventDefault();
-                     ev.stopPropagation();
-                     location.search = `?issue=${issueId}`;
-                     setIssueHistory(true);
-                     navigate(location);
-                  }}><Text variant="small" style={{ ...fontStyle }}>Issue</Text><div style={{ ...fontStyle }}>&nbsp;</div><Text variant="small" style={{ ...fontStyle }}>{`${issueId}`}</Text></DefaultButton>
+               <TooltipHost
+                  content={timestamp}
+                  id={tooltipId}
+                  calloutProps={{ gapSpace: 0 }}
+                  styles={{ root: { display: 'inline-block' } }}>
+                  <DefaultButton className={error ? styles.errorButton : styles.warningButton}
+                     href={href}
+                     style={{ padding: 0, margin: 0, width: tsWidth, paddingLeft: 8, paddingRight: 8, height: "100%", fontWeight: "unset" }}
+
+                     onClick={(ev) => {
+                        ev.preventDefault();
+                        ev.stopPropagation();
+                        location.search = `?issue=${issueId}`;
+                        setIssueHistory(true);
+                        navigate(location);
+                     }}>
+                     <Text variant="small" style={{ ...fontStyle }}>Issue</Text><div style={{ ...fontStyle }}>&nbsp;</div><Text variant="small" style={{ ...fontStyle }}>{`${issueId}`}</Text>
+                  </DefaultButton>
+               </TooltipHost>
             </Stack>
          }
 
@@ -693,7 +703,7 @@ export const LogList: React.FC<{ logId: string }> = observer(({ logId }) => {
                const search = new URLSearchParams(window.location.search);
                search.set("lineindex", (item.lineNumber - 1).toString());
                const url = `${window.location.pathname}?` + search.toString();
-      
+
                navigate(url, { replace: true })
             }}>
                <div style={{ position: "relative" }}>
