Horde: Add a more descriptive error for a shelved change not existing due to perforce edge server replication delay
#skipci
#rnx

diff --git a/Engine/Source/Programs/Horde/HordeDashboard/src/components/NewBuild.tsx b/Engine/Source/Programs/Horde/HordeDashboard/src/components/NewBuild.tsx
--- a/Engine/Source/Programs/Horde/HordeDashboard/src/components/NewBuild.tsx
+++ b/Engine/Source/Programs/Horde/HordeDashboard/src/components/NewBuild.tsx
@@ -1448,6 +1448,10 @@ export const NewBuild: React.FC<{ streamId: string; show: boolean; onClose: (new
 
                   if (errorReason) {
 
+                     if (errorReason?.trim().endsWith("does not exist")) {
+                        errorReason += ".  Perforce edge server replication for the change may be in progress."
+                     }
+
                      ErrorHandler.set({
 
                         reason: `${errorReason}`,
