Horde: Improve message on temp storage artifacts missing from direct links to storage
#skipci
#rnx

diff --git a/Engine/Source/Programs/Horde/HordeDashboard/src/components/artifacts/ArtifactsModal.tsx b/Engine/Source/Programs/Horde/HordeDashboard/src/components/artifacts/ArtifactsModal.tsx
--- a/Engine/Source/Programs/Horde/HordeDashboard/src/components/artifacts/ArtifactsModal.tsx
+++ b/Engine/Source/Programs/Horde/HordeDashboard/src/components/artifacts/ArtifactsModal.tsx
@@ -543,7 +543,11 @@ const JobDetailArtifactsInner: React.FC<{ jobId: string; stepId: string, artifac
 
    if (handler.artifactMissing) {
 
-      const text = `${handler.contextName} artifact was not found on server.`;
+      let text = `${handler.contextName} artifact was not found on server.`;
+
+      if (handler.contextName === "Temp Storage") {
+         text += "  Temporary storage artifacts expire due to high volume.";
+      }
 
       return <Stack>
          <Stack horizontal verticalAlign="center" verticalFill tokens={{ childrenGap: 12 }}>
