Horde: Load missing pools when loading a single agent
#rnx
#skipci

diff --git a/Engine/Source/Programs/Horde/HordeDashboard/src/backend/AgentStore.ts b/Engine/Source/Programs/Horde/HordeDashboard/src/backend/AgentStore.ts
--- a/Engine/Source/Programs/Horde/HordeDashboard/src/backend/AgentStore.ts
+++ b/Engine/Source/Programs/Horde/HordeDashboard/src/backend/AgentStore.ts
@@ -66,6 +66,11 @@ export class AgentStore {
                 this._agents = this._agents.filter(a => a.id !== agentId);
                 this._agents.push(agent);
 
+                const load = agent.pools?.filter(p => !this._pools.find(pool => pool.id === p));
+                if (load?.length) {
+                    this._pools = await backend.getPools()
+                }                
+
                 resolve(agent);
             } catch (error) {
                 reject(error);
