Horde: Support resolving latest change number at sync time for WorkspaceMaterializer

Also changes AutoSDK syncing to use latest change number of its stream, rather than the job's change number.

diff --git a/Engine/Source/Programs/Horde/Horde.Agent.Tests/Execution/FakeWorkspaceMaterializer.cs b/Engine/Source/Programs/Horde/Horde.Agent.Tests/Execution/FakeWorkspaceMaterializer.cs
--- a/Engine/Source/Programs/Horde/Horde.Agent.Tests/Execution/FakeWorkspaceMaterializer.cs
+++ b/Engine/Source/Programs/Horde/Horde.Agent.Tests/Execution/FakeWorkspaceMaterializer.cs
@@ -3,6 +3,7 @@
 using System;
 using System.Collections.Generic;
 using System.IO;
+using System.Linq;
 using System.Threading;
 using System.Threading.Tasks;
 using EpicGames.Core;
@@ -88,6 +89,11 @@ public class FakeWorkspaceMaterializer : IWorkspaceMaterializer
 			throw new Exception("Cannot sync before initialization");
 		}
 
+		if (changeNum == IWorkspaceMaterializer.LatestChangeNumber)
+		{
+			changeNum = _changeToFiles.Keys.Max();
+		}
+		
 		if (!_changeToFiles.ContainsKey(changeNum))
 		{
 			throw new WorkspaceMaterializationException($"Change {changeNum} could not be found");
diff --git a/Engine/Source/Programs/Horde/Horde.Agent/Execution/ManagedWorkspaceMaterializer.cs b/Engine/Source/Programs/Horde/Horde.Agent/Execution/ManagedWorkspaceMaterializer.cs
--- a/Engine/Source/Programs/Horde/Horde.Agent/Execution/ManagedWorkspaceMaterializer.cs
+++ b/Engine/Source/Programs/Horde/Horde.Agent/Execution/ManagedWorkspaceMaterializer.cs
@@ -97,6 +97,13 @@ public class ManagedWorkspaceMaterializer : IWorkspaceMaterializer
 		{
 			throw new WorkspaceMaterializationException("Workspace not initialized");
 		}
+		
+		if (changeNum == IWorkspaceMaterializer.LatestChangeNumber)
+		{
+			int latestChangeNum = await _workspace.GetLatestChangeAsync(cancellationToken);
+			scope.Span.SetTag("LatestChangeNum", latestChangeNum);
+			changeNum = latestChangeNum;
+		}
 
 		FileReference cacheFile = FileReference.Combine(_workspace.MetadataDir, "Contents.dat");
 		if (_useCacheFile)
diff --git a/Engine/Source/Programs/Horde/Horde.Agent/Execution/WorkspaceExecutor.cs b/Engine/Source/Programs/Horde/Horde.Agent/Execution/WorkspaceExecutor.cs
--- a/Engine/Source/Programs/Horde/Horde.Agent/Execution/WorkspaceExecutor.cs
+++ b/Engine/Source/Programs/Horde/Horde.Agent/Execution/WorkspaceExecutor.cs
@@ -46,11 +46,8 @@ namespace Horde.Agent.Execution
 
 				autoSdkWorkspaceSettings = await _autoSdkWorkspace.InitializeAsync(logger, cancellationToken);
 
-				// Match change for AutoSDK and actual job change
-				int autoSdkChangeNumber = _batch.Change;
-
 				SyncOptions syncOptions = new();
-				await _autoSdkWorkspace.SyncAsync(autoSdkChangeNumber, -1, syncOptions, cancellationToken);
+				await _autoSdkWorkspace.SyncAsync(IWorkspaceMaterializer.LatestChangeNumber, -1, syncOptions, cancellationToken);
 			}
 			
 			// Sync the regular workspace
diff --git a/Engine/Source/Programs/Horde/Horde.Agent/Execution/WorkspaceMaterializer.cs b/Engine/Source/Programs/Horde/Horde.Agent/Execution/WorkspaceMaterializer.cs
--- a/Engine/Source/Programs/Horde/Horde.Agent/Execution/WorkspaceMaterializer.cs
+++ b/Engine/Source/Programs/Horde/Horde.Agent/Execution/WorkspaceMaterializer.cs
@@ -107,6 +107,11 @@ public class SyncOptions
 /// </summary>
 public interface IWorkspaceMaterializer
 {
+	/// <summary>
+	/// Placeholder for resolving the latest available change number of stream during sync
+	/// </summary>
+	public const int LatestChangeNumber = -2;
+	
 	/// <summary>
 	/// Prepare file system for syncing
 	/// </summary>
