Horde: Fix config preflight not correctly enumerating shelved files.

diff --git a/Engine/Source/Programs/Horde/Horde.Server/Server/ServerController.cs b/Engine/Source/Programs/Horde/Horde.Server/Server/ServerController.cs
--- a/Engine/Source/Programs/Horde/Horde.Server/Server/ServerController.cs
+++ b/Engine/Source/Programs/Horde/Horde.Server/Server/ServerController.cs
@@ -140,7 +140,7 @@ namespace Horde.Server.Server
 
 			IPooledPerforceConnection perforce = await _perforceService.ConnectAsync(cluster, cancellationToken: cancellationToken);
 
-			PerforceResponse<DescribeRecord> describeResponse = await perforce.TryDescribeAsync(request.ShelvedChange, cancellationToken);
+			PerforceResponse<DescribeRecord> describeResponse = await perforce.TryDescribeAsync(DescribeOptions.Shelved, -1, request.ShelvedChange, cancellationToken);
 			if (!describeResponse.Succeeded)
 			{
 				return BadRequest(KnownLogEvents.Horde_InvalidPreflight, "CL {Change} does not exist.", request.ShelvedChange);
@@ -151,16 +151,19 @@ namespace Horde.Server.Server
 			Dictionary<Uri, byte[]> files = new Dictionary<Uri, byte[]>();
 			foreach (DescribeFileRecord fileRecord in record.Files)
 			{
-				PerforceResponse<PrintRecord<byte[]>> printRecordResponse = await perforce.TryPrintDataAsync($"{fileRecord.DepotFile}@={request.ShelvedChange}", cancellationToken);
-				if (!printRecordResponse.Succeeded || printRecordResponse.Data.Contents == null)
+				if (fileRecord.DepotFile.EndsWith("/globals.json", StringComparison.OrdinalIgnoreCase) || fileRecord.DepotFile.EndsWith(".project.json", StringComparison.OrdinalIgnoreCase) || fileRecord.DepotFile.EndsWith(".stream.json", StringComparison.OrdinalIgnoreCase))
 				{
-					return BadRequest($"Unable to print contents of {fileRecord.DepotFile}@={request.ShelvedChange}");
-				}
+					PerforceResponse<PrintRecord<byte[]>> printRecordResponse = await perforce.TryPrintDataAsync($"{fileRecord.DepotFile}@={request.ShelvedChange}", cancellationToken);
+					if (!printRecordResponse.Succeeded || printRecordResponse.Data.Contents == null)
+					{
+						return BadRequest($"Unable to print contents of {fileRecord.DepotFile}@={request.ShelvedChange}");
+					}
 
-				PrintRecord<byte[]> printRecord = printRecordResponse.Data;
+					PrintRecord<byte[]> printRecord = printRecordResponse.Data;
 
-				Uri uri = new Uri($"perforce://{cluster}{printRecord.DepotFile}");
-				files.Add(uri, printRecord.Contents);
+					Uri uri = new Uri($"perforce://{cluster}{printRecord.DepotFile}");
+					files.Add(uri, printRecord.Contents);
+				}
 			}
 
 			if (files.Count == 0)
