Horde: Catch exceptions from failed login attempts in CLI.

diff --git a/Engine/Source/Programs/Horde/Horde/Settings.cs b/Engine/Source/Programs/Horde/Horde/Settings.cs
--- a/Engine/Source/Programs/Horde/Horde/Settings.cs
+++ b/Engine/Source/Programs/Horde/Horde/Settings.cs
@@ -121,7 +121,15 @@ namespace Horde
 			using ITokenStore tokenStore = TokenStoreFactory.CreateTokenStore();
 			OidcTokenManager oidcTokenManager = OidcTokenManager.CreateTokenManager(configuration, tokenStore, new List<string>() { OidcProvider });
 
-			OidcTokenInfo? result = await oidcTokenManager.TryGetAccessToken(OidcProvider, cancellationToken);
+			OidcTokenInfo? result = null;
+			try
+			{
+				result = await oidcTokenManager.TryGetAccessToken(OidcProvider, cancellationToken);
+			}
+			catch(Exception ex)
+			{
+				logger.LogTrace(ex, "Unable to get access token; attempting login: {Message}", ex.Message);
+			}
 			if(result == null)
 			{
 				result = await oidcTokenManager.Login(OidcProvider, cancellationToken);
