Horde: Try to release conform leases immediately if the lease creation is cancelled.

diff --git a/Engine/Source/Programs/Horde/Horde.Server/Tasks/ConformTaskSource.cs b/Engine/Source/Programs/Horde/Horde.Server/Tasks/ConformTaskSource.cs
--- a/Engine/Source/Programs/Horde/Horde.Server/Tasks/ConformTaskSource.cs
+++ b/Engine/Source/Programs/Horde/Horde.Server/Tasks/ConformTaskSource.cs
@@ -177,13 +177,21 @@ namespace Horde.Server.Tasks
 					LeaseId leaseId = new LeaseId(BinaryIdUtils.CreateNew());
 					if (await AllocateConformLeaseAsync(agent.Id, task.Workspaces, leaseId))
 					{
-						ILogFile log = await _logService.CreateLogFileAsync(JobId.Empty, leaseId, agent.SessionId, LogType.Json, useNewStorageBackend: false, cancellationToken: cancellationToken);
-						task.LogId = log.Id.ToString();
-						task.RemoveUntrackedFiles = agent.RequestFullConform;
+						try
+						{
+							ILogFile log = await _logService.CreateLogFileAsync(JobId.Empty, leaseId, agent.SessionId, LogType.Json, useNewStorageBackend: false, cancellationToken: cancellationToken);
+							task.LogId = log.Id.ToString();
+							task.RemoveUntrackedFiles = agent.RequestFullConform;
 
-						byte[] payload = Any.Pack(task).ToByteArray();
+							byte[] payload = Any.Pack(task).ToByteArray();
 
-						return Lease(new AgentLease(leaseId, null, "Updating workspaces", null, null, log.Id, LeaseState.Pending, null, true, payload));
+							return Lease(new AgentLease(leaseId, null, "Updating workspaces", null, null, log.Id, LeaseState.Pending, null, true, payload));
+						}
+						catch
+						{
+							await ReleaseConformLeaseAsync(leaseId);
+							throw;
+						}
 					}
 				}
 			}
