Horde: Include deleted agents for AWS-based lifecycle termination

diff --git a/Engine/Source/Programs/Horde/Horde.Server.Tests/Fleet/AwsAutoScalingLifecycleTest.cs b/Engine/Source/Programs/Horde/Horde.Server.Tests/Fleet/AwsAutoScalingLifecycleTest.cs
--- a/Engine/Source/Programs/Horde/Horde.Server.Tests/Fleet/AwsAutoScalingLifecycleTest.cs
+++ b/Engine/Source/Programs/Horde/Horde.Server.Tests/Fleet/AwsAutoScalingLifecycleTest.cs
@@ -71,7 +71,8 @@ public class AwsAutoScalingLifecycleServiceTest : TestSetup
 	{
 		// Arrange
 		LifecycleActionEvent lae = new() { Ec2InstanceId = "i-1234", LifecycleActionToken = "action-token-test", Origin = "AutoScalingGroup" };
-		IAgent agent = await CreateAgentAsync(new PoolId("pool1"), properties: new() { KnownPropertyNames.AwsInstanceId + "=" + lae.Ec2InstanceId });
+		IAgent agent = await CreateAgentAsync(new PoolId("pool1"), properties: new() { KnownPropertyNames.AwsInstanceId + "=" + lae.Ec2InstanceId }, ephemeral: true);
+		await AgentService.DeleteAgentAsync(agent);
 		
 		// Act
 		await _fakeSqs.SendMessageAsync(_queueUrl, JsonSerializer.Serialize(lae));
diff --git a/Engine/Source/Programs/Horde/Horde.Server/Agents/AgentCollection.cs b/Engine/Source/Programs/Horde/Horde.Server/Agents/AgentCollection.cs
--- a/Engine/Source/Programs/Horde/Horde.Server/Agents/AgentCollection.cs
+++ b/Engine/Source/Programs/Horde/Horde.Server/Agents/AgentCollection.cs
@@ -194,11 +194,15 @@ namespace Horde.Server.Agents
 		}
 
 		/// <inheritdoc/>
-		public async Task<List<IAgent>> FindAsync(PoolId? poolId, DateTime? modifiedAfter, string? property, AgentStatus? status, bool? enabled, int? index, int? count)
+		public async Task<List<IAgent>> FindAsync(PoolId? poolId, DateTime? modifiedAfter, string? property, AgentStatus? status, bool? enabled, bool includeDeleted, int? index, int? count)
 		{
 			FilterDefinitionBuilder<AgentDocument> filterBuilder = new FilterDefinitionBuilder<AgentDocument>();
 
-			FilterDefinition<AgentDocument> filter = filterBuilder.Ne(x => x.Deleted, true);
+			FilterDefinition<AgentDocument> filter = filterBuilder.Empty;
+			if (!includeDeleted)
+			{
+				filter &= filterBuilder.Ne(x => x.Deleted, true);
+			}
 			
 			if (poolId != null)
 			{
diff --git a/Engine/Source/Programs/Horde/Horde.Server/Agents/AgentService.cs b/Engine/Source/Programs/Horde/Horde.Server/Agents/AgentService.cs
--- a/Engine/Source/Programs/Horde/Horde.Server/Agents/AgentService.cs
+++ b/Engine/Source/Programs/Horde/Horde.Server/Agents/AgentService.cs
@@ -212,12 +212,13 @@ namespace Horde.Server.Agents
 		/// <param name="poolId">The pool containing the agent</param>
 		/// <param name="modifiedAfter">If set, only returns agents modified after this time</param>
 		/// <param name="property">If set, only return agents matching this property</param>
+		/// <param name="includeDeleted">If set, include agents marked as deleted</param>
 		/// <param name="index">Index within the list of results</param>
 		/// <param name="count">Number of results to return</param>
 		/// <returns>List of agents matching the given criteria</returns>
-		public Task<List<IAgent>> FindAgentsAsync(PoolId? poolId, DateTime? modifiedAfter, string? property, int? index, int? count)
+		public Task<List<IAgent>> FindAgentsAsync(PoolId? poolId, DateTime? modifiedAfter, string? property, bool includeDeleted, int? index, int? count)
 		{
-			return Agents.FindAsync(poolId, modifiedAfter, property, null, null, index, count);
+			return Agents.FindAsync(poolId, modifiedAfter, property, null, null, includeDeleted, index, count);
 		}
 
 		/// <summary>
diff --git a/Engine/Source/Programs/Horde/Horde.Server/Agents/AgentsController.cs b/Engine/Source/Programs/Horde/Horde.Server/Agents/AgentsController.cs
--- a/Engine/Source/Programs/Horde/Horde.Server/Agents/AgentsController.cs
+++ b/Engine/Source/Programs/Horde/Horde.Server/Agents/AgentsController.cs
@@ -44,6 +44,7 @@ namespace Horde.Server.Agents
 		/// Finds the agents matching specified criteria.
 		/// </summary>
 		/// <param name="poolId">The pool containing the agent</param>
+		/// <param name="includeDeleted">Whether to include agents marked as deleted</param>
 		/// <param name="index">First result to return</param>
 		/// <param name="count">Number of results to return</param>
 		/// <param name="modifiedAfter">If set, only returns agents modified after this time</param>
@@ -52,14 +53,14 @@ namespace Horde.Server.Agents
 		[HttpGet]
 		[Route("/api/v1/agents")]
 		[ProducesResponseType(typeof(List<GetAgentResponse>), 200)]
-		public async Task<ActionResult<List<object>>> FindAgentsAsync([FromQuery] PoolId? poolId = null, [FromQuery] int? index = null, [FromQuery] int? count = null, [FromQuery] DateTimeOffset? modifiedAfter = null, [FromQuery] PropertyFilter? filter = null)
+		public async Task<ActionResult<List<object>>> FindAgentsAsync([FromQuery] PoolId? poolId = null, [FromQuery] bool includeDeleted = false, [FromQuery] int? index = null, [FromQuery] int? count = null, [FromQuery] DateTimeOffset? modifiedAfter = null, [FromQuery] PropertyFilter? filter = null)
 		{
 			if (!_globalConfig.Value.Authorize(AgentAclAction.ListAgents, User))
 			{
 				return Forbid(AgentAclAction.ListAgents);
 			}
 
-			List<IAgent> agents = await _agentService.FindAgentsAsync(poolId, modifiedAfter?.UtcDateTime, null, index, count);
+			List<IAgent> agents = await _agentService.FindAgentsAsync(poolId, modifiedAfter?.UtcDateTime, null, includeDeleted, index, count);
 
 			List<object> responses = new List<object>();
 			foreach (IAgent agent in agents)
diff --git a/Engine/Source/Programs/Horde/Horde.Server/Agents/Fleet/AwsAutoScalingLifecycle.cs b/Engine/Source/Programs/Horde/Horde.Server/Agents/Fleet/AwsAutoScalingLifecycle.cs
--- a/Engine/Source/Programs/Horde/Horde.Server/Agents/Fleet/AwsAutoScalingLifecycle.cs
+++ b/Engine/Source/Programs/Horde/Horde.Server/Agents/Fleet/AwsAutoScalingLifecycle.cs
@@ -231,7 +231,7 @@ public sealed class AwsAutoScalingLifecycleService : IHostedService, IDisposable
 		if (e.Origin == OriginAsg)
 		{
 			string instanceIdProp = KnownPropertyNames.AwsInstanceId + "=" + e.Ec2InstanceId;
-			IAgent? agent = (await _agentService.FindAgentsAsync(null, null, instanceIdProp, null, null)).FirstOrDefault();
+			IAgent? agent = (await _agentService.FindAgentsAsync(null, null, instanceIdProp, true, null, null)).FirstOrDefault();
 			if (agent == null)
 			{
 				_logger.LogWarning("Lifecycle action received but no agent with instance ID {InstanceId} found", e.Ec2InstanceId);
@@ -282,7 +282,7 @@ public sealed class AwsAutoScalingLifecycleService : IHostedService, IDisposable
 		}
 		
 		List<string> validInstanceIds = new();
-		List<IAgent> agents = await _agentService.FindAgentsAsync(null, null, null, null, null);
+		List<IAgent> agents = await _agentService.FindAgentsAsync(null, null, null, true, null, null);
 		foreach (IAgent agent in agents)
 		{
 			if (IsAgentSuggestedByAsg(agent, out string? instanceId))
diff --git a/Engine/Source/Programs/Horde/Horde.Server/Agents/IAgentCollection.cs b/Engine/Source/Programs/Horde/Horde.Server/Agents/IAgentCollection.cs
--- a/Engine/Source/Programs/Horde/Horde.Server/Agents/IAgentCollection.cs
+++ b/Engine/Source/Programs/Horde/Horde.Server/Agents/IAgentCollection.cs
@@ -63,10 +63,11 @@ namespace Horde.Server.Agents
 		/// <param name="property">Property to look for</param>
 		/// <param name="status">Status to look for</param>
 		/// <param name="enabled">Enabled/disabled status to look for</param>
+		/// <param name="includeDeleted">Whether agents marked as deleted should be included</param>
 		/// <param name="index">Index of the first result</param>
 		/// <param name="count">Number of results to return</param>
 		/// <returns>List of agents matching the given criteria</returns>
-		Task<List<IAgent>> FindAsync(PoolId? poolId = null, DateTime? modifiedAfter = null, string? property = null, AgentStatus? status = null, bool? enabled = null, int? index = null, int? count = null);
+		Task<List<IAgent>> FindAsync(PoolId? poolId = null, DateTime? modifiedAfter = null, string? property = null, AgentStatus? status = null, bool? enabled = null, bool includeDeleted = false, int? index = null, int? count = null);
 
 		/// <summary>
 		/// Finds all agents with an expired session
