Horde: Ensure that all cases where the server transitions an agent to the "stopped" state exits the lease handler.

diff --git a/Engine/Source/Programs/Horde/Horde.Agent/Leases/Handlers/RestartHandler.cs b/Engine/Source/Programs/Horde/Horde.Agent/Leases/Handlers/RestartHandler.cs
--- a/Engine/Source/Programs/Horde/Horde.Agent/Leases/Handlers/RestartHandler.cs
+++ b/Engine/Source/Programs/Horde/Horde.Agent/Leases/Handlers/RestartHandler.cs
@@ -21,7 +21,7 @@ namespace Horde.Agent.Leases.Handlers
 		/// <inheritdoc/>
 		public override Task<LeaseResult> ExecuteAsync(ISession session, string leaseId, RestartTask task, CancellationToken cancellationToken)
 		{
-			_logger.LogInformation("Scheduling restart task");
+			_logger.LogInformation("Scheduling restart task for agent {AgentId}", session.AgentId);
 			SessionResult result = new SessionResult((logger, ctx) => Shutdown.ExecuteAsync(true, logger, ctx));
 			return Task.FromResult(new LeaseResult(result));
 		}
diff --git a/Engine/Source/Programs/Horde/Horde.Agent/Leases/Handlers/ShutdownHandler.cs b/Engine/Source/Programs/Horde/Horde.Agent/Leases/Handlers/ShutdownHandler.cs
--- a/Engine/Source/Programs/Horde/Horde.Agent/Leases/Handlers/ShutdownHandler.cs
+++ b/Engine/Source/Programs/Horde/Horde.Agent/Leases/Handlers/ShutdownHandler.cs
@@ -21,7 +21,7 @@ namespace Horde.Agent.Leases.Handlers
 		/// <inheritdoc/>
 		public override Task<LeaseResult> ExecuteAsync(ISession session, string leaseId, ShutdownTask task, CancellationToken cancellationToken)
 		{
-			_logger.LogInformation("Scheduling shutdown task");
+			_logger.LogInformation("Scheduling shutdown task for agent {AgentId}", session.AgentId);
 			SessionResult result = new SessionResult((logger, ctx) => Shutdown.ExecuteAsync(false, logger, ctx));
 			return Task.FromResult(new LeaseResult(result));
 		}
diff --git a/Engine/Source/Programs/Horde/Horde.Agent/Leases/LeaseManager.cs b/Engine/Source/Programs/Horde/Horde.Agent/Leases/LeaseManager.cs
--- a/Engine/Source/Programs/Horde/Horde.Agent/Leases/LeaseManager.cs
+++ b/Engine/Source/Programs/Horde/Horde.Agent/Leases/LeaseManager.cs
@@ -7,6 +7,7 @@ using System.Linq;
 using System.Threading;
 using System.Threading.Tasks;
 using EpicGames.Core;
+using EpicGames.Horde.Storage;
 using Google.Protobuf.WellKnownTypes;
 using Grpc.Core;
 using Horde.Agent.Services;
@@ -240,13 +241,13 @@ namespace Horde.Agent.Leases
 				Task waitTask = _updateLeasesEvent.Task;
 
 				// Flag for whether the service is stopping
-				bool stopping = false;
-				if (stoppingToken.IsCancellationRequested)
+				if (stoppingToken.IsCancellationRequested && _sessionResult == null)
 				{
-					_logger.LogInformation("Cancellation from token requested");
-					stopping = true;
+					_logger.LogInformation("Cancellation from token requested; setting session result to terminate.");
+					_sessionResult = new SessionResult(SessionOutcome.Terminate);
 				}
 
+				bool stopping = false;
 				if (_sessionResult != null)
 				{
 					_logger.LogInformation("Session termination requested (result: {Result})", _sessionResult.Outcome);
@@ -363,12 +364,19 @@ namespace Horde.Agent.Leases
 									OnLeaseActive?.Invoke(serverLease);
 								}
 							}
+
+							// Update the session result if we've transitioned to stopped
+							if (updateSessionResponse.Status == AgentStatus.Stopped)
+							{
+								_logger.LogInformation("Agent status is stopped; returning from session update loop.");
+								return _sessionResult ?? new SessionResult(SessionOutcome.BackOff);
+							}
 						}
 
 						// If there's nothing still running and cancellation was requested, exit
 						if (_activeLeases.Count == 0 && _sessionResult != null)
 						{
-							_logger.LogInformation("No leases are active. Agent is stopping.");
+							_logger.LogInformation("No leases are active. Agent is stopping."); // TODO: Should not really hit this any more; server should report stopping state in lease update above.
 							return _sessionResult;
 						}
 					}
diff --git a/Engine/Source/Programs/Horde/Horde.Server/Agents/AgentService.cs b/Engine/Source/Programs/Horde/Horde.Server/Agents/AgentService.cs
--- a/Engine/Source/Programs/Horde/Horde.Server/Agents/AgentService.cs
+++ b/Engine/Source/Programs/Horde/Horde.Server/Agents/AgentService.cs
@@ -755,6 +755,7 @@ namespace Horde.Server.Agents
 				IAgent? terminatedAgent = await TryTerminateSessionAsync(agent);
 				if (terminatedAgent != null)
 				{
+					_logger.LogInformation("Terminated session {SessionId} for {AgentId}; agent is stopping", agent.SessionId, agent.Id);
 					agent = terminatedAgent;
 					break;
 				}
diff --git a/Engine/Source/Programs/Horde/Horde.Server/Server/RpcService.cs b/Engine/Source/Programs/Horde/Horde.Server/Server/RpcService.cs
--- a/Engine/Source/Programs/Horde/Horde.Server/Server/RpcService.cs
+++ b/Engine/Source/Programs/Horde/Horde.Server/Server/RpcService.cs
@@ -416,6 +416,7 @@ namespace Horde.Server.Server
 					UpdateSessionResponse response = new UpdateSessionResponse();
 					response.Leases.Add(agent.Leases.Select(x => x.ToRpcMessage()));
 					response.ExpiryTime = (agent.SessionExpiresAt == null) ? new Timestamp() : Timestamp.FromDateTime(agent.SessionExpiresAt.Value);
+					response.Status = agent.Status;
 					await writer.WriteAsync(response);
 				}
 
