Horde: Increase time before warning about lock on issue collection.

diff --git a/Engine/Source/Programs/Horde/Horde.Server/Issues/IssueCollection.cs b/Engine/Source/Programs/Horde/Horde.Server/Issues/IssueCollection.cs
--- a/Engine/Source/Programs/Horde/Horde.Server/Issues/IssueCollection.cs
+++ b/Engine/Source/Programs/Horde/Horde.Server/Issues/IssueCollection.cs
@@ -462,7 +462,7 @@ namespace Horde.Server.Issues
 		public async Task<IAsyncDisposable> EnterCriticalSectionAsync()
 		{
 			Stopwatch timer = Stopwatch.StartNew();
-			TimeSpan nextNotifyTime = TimeSpan.FromSeconds(2.0);
+			TimeSpan nextNotifyTime = TimeSpan.FromSeconds(10.0);
 
 			RedisLock issueLock = new (_redisService.GetDatabase(), "issues/lock");
 			while (!await issueLock.AcquireAsync(TimeSpan.FromMinutes(1)))
