Horde: Increase frequency of compute timeout checks.

diff --git a/Engine/Source/Programs/Horde/Horde.Agent/Leases/Handlers/ComputeHandler.cs b/Engine/Source/Programs/Horde/Horde.Agent/Leases/Handlers/ComputeHandler.cs
--- a/Engine/Source/Programs/Horde/Horde.Agent/Leases/Handlers/ComputeHandler.cs
+++ b/Engine/Source/Programs/Horde/Horde.Agent/Leases/Handlers/ComputeHandler.cs
@@ -157,13 +157,14 @@ namespace Horde.Agent.Leases.Handlers
 		{
 			while(!cancellationToken.IsCancellationRequested)
 			{
-				if (transport.TimeSinceActivity > NoDataTimeout)
+				TimeSpan reaminingTime = NoDataTimeout - transport.TimeSinceActivity;
+				if (reaminingTime < TimeSpan.Zero)
 				{
 					logger.LogWarning("Terminating compute task due to timeout (last tick at {Time})", DateTime.UtcNow - transport.TimeSinceActivity);
 					cts.Cancel();
 					break;
 				}
-				await Task.Delay(TimeSpan.FromSeconds(20), cancellationToken);
+				await Task.Delay(reaminingTime + TimeSpan.FromSeconds(0.2), cancellationToken);
 			}
 		}
 	}
