Horde: Reduce log spam for users that don't have an email address when attempting to map to Slack users.

diff --git a/Engine/Source/Programs/Horde/Horde.Server/Notifications/Sinks/SlackNotificationSink.cs b/Engine/Source/Programs/Horde/Horde.Server/Notifications/Sinks/SlackNotificationSink.cs
--- a/Engine/Source/Programs/Horde/Horde.Server/Notifications/Sinks/SlackNotificationSink.cs
+++ b/Engine/Source/Programs/Horde/Horde.Server/Notifications/Sinks/SlackNotificationSink.cs
@@ -2408,12 +2408,17 @@ namespace Horde.Server.Notifications.Sinks
 			return (await GetSlackUserAsync(user))?.SlackUserId;
 		}
 
+		private readonly HashSet<UserId> _usersWithoutEmail = new HashSet<UserId>();
+
 		private async Task<SlackUserDocument?> GetSlackUserAsync(IUser user)
 		{
 			string? email = user.Email;
 			if (email == null)
 			{
-				_logger.LogWarning("Unable to find Slack user id for {UserId} ({Name}): No email address in user profile", user.Id, user.Name);
+				if (_usersWithoutEmail.Add(user.Id))
+				{
+					_logger.LogInformation("Unable to find Slack user id for {UserId} ({Name}): No email address in user profile", user.Id, user.Name);
+				}
 				return null;
 			}
 
