Horde: Fix markdown syntax for agent report.

diff --git a/Engine/Source/Programs/Horde/Horde.Server/Notifications/Sinks/SlackNotificationSink.cs b/Engine/Source/Programs/Horde/Horde.Server/Notifications/Sinks/SlackNotificationSink.cs
--- a/Engine/Source/Programs/Horde/Horde.Server/Notifications/Sinks/SlackNotificationSink.cs
+++ b/Engine/Source/Programs/Horde/Horde.Server/Notifications/Sinks/SlackNotificationSink.cs
@@ -3015,7 +3015,7 @@ namespace Horde.Server.Notifications.Sinks
 				await _slackClient.PostMessageAsync(_settings.AgentNotificationChannel, headerMessage);
 
 				{
-					StringBuilder conformMessage = new StringBuilder("**Conform issues:**\n");
+					StringBuilder conformMessage = new StringBuilder("*Conform issues:*\n");
 					if (report.ConformLoop.Count == 0)
 					{
 						conformMessage.Append("None.\n");
@@ -3028,7 +3028,7 @@ namespace Horde.Server.Notifications.Sinks
 							foreach ((AgentId agentId, int conformCount) in conformBatch)
 							{
 								Uri agentUrl = new Uri(_settings.DashboardUrl, $"agents?agentId={agentId}");
-								conformMessage.Append($"* **[{agentId}]({agentUrl})** has run conform {conformCount} times\n");
+								conformMessage.Append($"* *[{agentId}]({agentUrl})* has run conform {conformCount} times\n");
 							}
 							await _slackClient.PostMessageAsync(_settings.AgentNotificationChannel, conformMessage.ToString());
 						}
