Horde: Suppress log spam when creating a unique CL number can't submit due to a conflict.

diff --git a/Engine/Source/Programs/Horde/Horde.Server/Perforce/PerforceService.cs b/Engine/Source/Programs/Horde/Horde.Server/Perforce/PerforceService.cs
--- a/Engine/Source/Programs/Horde/Horde.Server/Perforce/PerforceService.cs
+++ b/Engine/Source/Programs/Horde/Horde.Server/Perforce/PerforceService.cs
@@ -1002,7 +1002,15 @@ namespace Horde.Server.Perforce
 						change.Files.Add(depotPath);
 						change = await perforce.CreateChangeAsync(change, cancellationToken);
 
-						SubmitRecord submit = await perforce.SubmitAsync(change.Number, SubmitOptions.SubmitUnchanged, cancellationToken);
+						PerforceResponse<SubmitRecord> submitResponse = await perforce.TrySubmitAsync(change.Number, SubmitOptions.SubmitUnchanged, cancellationToken);
+						if (attempt < MaxAttempts && submitResponse.Error != null && submitResponse.Error.Generic == PerforceGenericCode.NotYet)
+						{
+							// File needs resolving; sync and retry.
+							_logger.LogDebug("Unable to submit new changelist (file: {File}, depotPath: {DepotPath}, description: \"{Description}\", attempt: {Attempt}/{MaxAttempts}, error: {Message}", filePath, depotPath, description, attempt, MaxAttempts, submitResponse.Error);
+							continue;
+						}
+
+						SubmitRecord submit = submitResponse.Data;
 						_logger.LogInformation("Submitted new changelist with {DepotPath}: CL {Change}", depotPath, submit.SubmittedChangeNumber);
 						return submit.SubmittedChangeNumber;
 					}
