Horde: Suppress warnings inviting users to a channel when the error code indicates they are already *in* the channel.

diff --git a/Engine/Source/Programs/Horde/Horde.Server/Notifications/Sinks/SlackNotificationSink.cs b/Engine/Source/Programs/Horde/Horde.Server/Notifications/Sinks/SlackNotificationSink.cs
--- a/Engine/Source/Programs/Horde/Horde.Server/Notifications/Sinks/SlackNotificationSink.cs
+++ b/Engine/Source/Programs/Horde/Horde.Server/Notifications/Sinks/SlackNotificationSink.cs
@@ -934,7 +934,10 @@ namespace Horde.Server.Notifications.Sinks
 						}
 						catch (SlackException ex)
 						{
-							_logger.LogWarning(ex, "Unable to invite user {UserId} to {Channel} (as admin): {Error}", slackUserId, channel, ex.Code);
+							if (!String.Equals(ex.Code, "already_in_channel", StringComparison.Ordinal))
+							{
+								_logger.LogWarning(ex, "Unable to invite user {UserId} to {Channel} (as admin): {Error}", slackUserId, channel, ex.Code);
+							}
 						}
 					}
 					else
