Horde: Another pass on Bisection workflow for server and dashboard
#rnx

diff --git a/Engine/Source/Programs/Horde/Horde.Server.Tests/BisectTests.cs b/Engine/Source/Programs/Horde/Horde.Server.Tests/BisectTests.cs
--- a/Engine/Source/Programs/Horde/Horde.Server.Tests/BisectTests.cs
+++ b/Engine/Source/Programs/Horde/Horde.Server.Tests/BisectTests.cs
@@ -48,6 +48,11 @@ namespace Horde.Server.Tests
 			Assert.AreEqual(failedJob.Change, response.CurrentChange);
 			Assert.AreEqual(failedJob.Id, response.CurrentJobId);
 
+			IUser user = await UserCollection.FindOrAddUserByLoginAsync("TestUser");
+			IUserSettings settings = await UserCollection.GetSettingsAsync(user.Id);
+			
+			Assert.AreEqual(task.BisectTaskId, settings.PinnedBisectTaskIds[0]);
+
 			Assert.AreEqual(task.BisectTaskId, response.Id);
 			List<IJob> jobs = await JobCollection.FindBisectTaskJobsAsync(response!.Id, null).ToListAsync();
 			Assert.AreEqual(0, jobs.Count);
@@ -65,7 +70,7 @@ namespace Horde.Server.Tests
 			Assert.AreEqual(14, jobs[0].Change);
 			await SetJobOutcomeAsync(jobs[0], graph, JobStepOutcome.Failure);
 
-			List<GetBisectTaskResponse> bisectTasks = Deref(await BisectTasksController!.FindBisectTasksAsync(bisectTask.Owner?.Id.ToString() ?? ""));
+			List<GetBisectTaskResponse> bisectTasks = Deref(await BisectTasksController!.FindBisectTasksAsync(null, bisectTask.Owner?.Id.ToString() ?? ""));
 			Assert.AreEqual(1, bisectTasks.Count);
 			Assert.AreEqual(bisectTask.Owner?.Id ?? default, bisectTasks[0].Owner?.Id ?? default);
 
diff --git a/Engine/Source/Programs/Horde/Horde.Server/Jobs/Bisect/BisectController.cs b/Engine/Source/Programs/Horde/Horde.Server/Jobs/Bisect/BisectController.cs
--- a/Engine/Source/Programs/Horde/Horde.Server/Jobs/Bisect/BisectController.cs
+++ b/Engine/Source/Programs/Horde/Horde.Server/Jobs/Bisect/BisectController.cs
@@ -268,7 +268,15 @@ namespace Horde.Server.Jobs.Bisect
 			options.IgnoreChanges = create.IgnoreChanges;
 			options.IgnoreJobs = create.IgnoreJobs;
 
-			IBisectTask bisectTask = await _bisectTaskCollection.CreateAsync(job, initialBatch.Id, jobStep.Id, create.NodeName, jobStep.Outcome, User.GetUserId() ?? UserId.Empty, options, cancellationToken);
+			UserId? userId = User.GetUserId();
+
+			IBisectTask bisectTask = await _bisectTaskCollection.CreateAsync(job, initialBatch.Id, jobStep.Id, create.NodeName, jobStep.Outcome, userId ?? UserId.Empty, options, cancellationToken);
+
+			if (userId != null)
+			{
+				await _userCollection.UpdateSettingsAsync(userId.Value, addBisectTaskIds: new[] { bisectTask.Id });
+			}
+
 			return new CreateBisectTaskResponse(bisectTask);
 		}
 
@@ -366,7 +374,7 @@ namespace Horde.Server.Jobs.Bisect
 
 			List<IJobStepRef> steps = await _jobStepRefs.GetStepsForNodeAsync(initialJob.StreamId, initialJob.TemplateId, task.NodeName, null, true, 1024, task.Id, cancellationToken);
 			IJobStepRef? initialStep = await _jobStepRefs.FindAsync(task.InitialJobId, task.InitialBatchId, task.InitialStepId);
-			if (initialStep != null) 
+			if (initialStep != null)
 			{
 				steps.Add(initialStep);
 			}
@@ -391,7 +399,7 @@ namespace Horde.Server.Jobs.Bisect
 				return NotFound(jobId);
 			}
 
-			IReadOnlyList<IBisectTask> tasks = await _bisectTaskCollection.FindAsync(jobId, null, null, null, null, null, cancellationToken);
+			IReadOnlyList<IBisectTask> tasks = await _bisectTaskCollection.FindAsync(null, jobId, null, null, null, null, null, cancellationToken);
 
 			List<GetBisectTaskResponse> response = new List<GetBisectTaskResponse>();
 
@@ -421,6 +429,7 @@ namespace Horde.Server.Jobs.Bisect
 		[Route("/api/v1/bisect")]
 		[ProducesResponseType(typeof(List<GetBisectTaskResponse>), 200)]
 		public async Task<ActionResult<List<GetBisectTaskResponse>>> FindBisectTasksAsync(
+		[FromQuery(Name = "id")] string[]? ids = null,
 		[FromQuery] string? ownerId = null,
 		[FromQuery] string? jobId = null,
 		[FromQuery] DateTimeOffset? minCreateTime = null,
@@ -431,12 +440,13 @@ namespace Horde.Server.Jobs.Bisect
 		{
 			List<GetBisectTaskResponse> responses = new List<GetBisectTaskResponse>();
 
+			BisectTaskId[]? bisectTaskIdValues = (ids == null) ? (BisectTaskId[]?)null : Array.ConvertAll(ids, x => BisectTaskId.Parse(x));
 			JobId? JobIdValue = !String.IsNullOrEmpty(jobId) ? JobId.Parse(jobId) : null;
 			UserId? OwnerIdValue = !String.IsNullOrEmpty(ownerId) ? UserId.Parse(ownerId) : null;
 
-			IReadOnlyList<IBisectTask> tasks = await _bisectTaskCollection.FindAsync(JobIdValue, OwnerIdValue, minCreateTime?.UtcDateTime, maxCreateTime?.UtcDateTime, index, count, cancellationToken);
+			IReadOnlyList<IBisectTask> tasks = await _bisectTaskCollection.FindAsync(bisectTaskIdValues, JobIdValue, OwnerIdValue, minCreateTime?.UtcDateTime, maxCreateTime?.UtcDateTime, index, count, cancellationToken);
 
-			if (tasks.Count == 0) 
+			if (tasks.Count == 0)
 			{
 				return responses;
 			}
@@ -445,7 +455,7 @@ namespace Horde.Server.Jobs.Bisect
 
 			for (int i = 0; i < tasks.Count; i++)
 			{
-				IBisectTask task = tasks[i];								
+				IBisectTask task = tasks[i];
 				IJob? job = jobs.FirstOrDefault(x => x.Id == task.InitialJobId);
 				if (job != null)
 				{
diff --git a/Engine/Source/Programs/Horde/Horde.Server/Jobs/Bisect/BisectTaskCollection.cs b/Engine/Source/Programs/Horde/Horde.Server/Jobs/Bisect/BisectTaskCollection.cs
--- a/Engine/Source/Programs/Horde/Horde.Server/Jobs/Bisect/BisectTaskCollection.cs
+++ b/Engine/Source/Programs/Horde/Horde.Server/Jobs/Bisect/BisectTaskCollection.cs
@@ -154,13 +154,19 @@ namespace Horde.Server.Jobs.Bisect
 		}
 
 		/// <inheritdoc/>
-		public async Task<IReadOnlyList<IBisectTask>> FindAsync(JobId? jobId = null, UserId? ownerId = null, DateTime? minCreateTime = null, DateTime? maxCreateTime = null, int? index = null, int? count = null, CancellationToken cancellationToken = default)
+		public async Task<IReadOnlyList<IBisectTask>> FindAsync(BisectTaskId[]? taskIds = null, JobId? jobId = null, UserId? ownerId = null, DateTime? minCreateTime = null, DateTime? maxCreateTime = null, int? index = null, int? count = null, CancellationToken cancellationToken = default)
 		{
 			// Find all the bisection tasks matching the given criteria
 			FilterDefinitionBuilder<BisectTaskDoc> filterBuilder = Builders<BisectTaskDoc>.Filter;
 
 			FilterDefinition<BisectTaskDoc> filter = FilterDefinition<BisectTaskDoc>.Empty;
 
+			if (taskIds != null && taskIds.Length > 0)
+			{
+				filter &= filterBuilder.In(x => x.Id, taskIds);
+			}
+
+
 			if (jobId != null)
 			{
 				filter &= filterBuilder.Eq(x => x.InitialJobId, jobId);
diff --git a/Engine/Source/Programs/Horde/Horde.Server/Jobs/Bisect/IBisectTaskCollection.cs b/Engine/Source/Programs/Horde/Horde.Server/Jobs/Bisect/IBisectTaskCollection.cs
--- a/Engine/Source/Programs/Horde/Horde.Server/Jobs/Bisect/IBisectTaskCollection.cs
+++ b/Engine/Source/Programs/Horde/Horde.Server/Jobs/Bisect/IBisectTaskCollection.cs
@@ -5,7 +5,6 @@ using System.Collections.Generic;
 using System.Threading;
 using System.Threading.Tasks;
 using Horde.Server.Perforce;
-using Horde.Server.Users;
 using HordeCommon;
 using EpicGames.Horde.Api;
 
@@ -40,6 +39,7 @@ namespace Horde.Server.Jobs.Bisect
 		/// <summary>
 		/// Finds all bisect tasks for the provided criteria
 		/// </summary>
+		/// <param name="taskIds">array of bisection task ids</param>
 		/// <param name="jobId">jobId of initial bisection</param>
 		/// <param name="ownerId">ownerId of bisection</param>
 		/// <param name="minCreateTime">min creation time of the bisection task</param>
@@ -48,7 +48,7 @@ namespace Horde.Server.Jobs.Bisect
 		/// <param name="count"></param>
 		/// <param name="cancellationToken">Cancellation token for the operation</param>
 		/// <returns>List of bisect tasks</returns>
-		public Task<IReadOnlyList<IBisectTask>> FindAsync(JobId? jobId = null, UserId? ownerId = null, DateTime? minCreateTime = null, DateTime? maxCreateTime = null, int? index = null, int? count = null, CancellationToken cancellationToken = default);
+		public Task<IReadOnlyList<IBisectTask>> FindAsync(BisectTaskId[]? taskIds = null, JobId? jobId = null, UserId? ownerId = null, DateTime? minCreateTime = null, DateTime? maxCreateTime = null, int? index = null, int? count = null, CancellationToken cancellationToken = default);
 
 		/// <summary>
 		/// Gets a bisect task by id
diff --git a/Engine/Source/Programs/Horde/Horde.Server/Users/IUser.cs b/Engine/Source/Programs/Horde/Horde.Server/Users/IUser.cs
--- a/Engine/Source/Programs/Horde/Horde.Server/Users/IUser.cs
+++ b/Engine/Source/Programs/Horde/Horde.Server/Users/IUser.cs
@@ -3,10 +3,10 @@
 using System;
 using System.Collections.Generic;
 using System.Linq;
-using Amazon.S3.Model;
 using EpicGames.Core;
 using EpicGames.Horde.Api;
 using Horde.Server.Jobs;
+using Horde.Server.Jobs.Bisect;
 using Horde.Server.Streams;
 using MongoDB.Bson;
 
@@ -148,6 +148,11 @@ namespace Horde.Server.Users
 		/// </summary>
 		public IReadOnlyList<JobId> PinnedJobIds { get; }
 
+		/// <summary>
+		/// List of pinned bisection tasks
+		/// </summary>
+		public IReadOnlyList<BisectTaskId> PinnedBisectTaskIds { get; }		
+
 		/// <summary>
 		/// List of job template preferences
 		/// </summary>
@@ -180,6 +185,7 @@ namespace Horde.Server.Users
 
 				response.DashboardSettings = BsonTypeMapper.MapToDotNetValue(settings.DashboardSettings);
 				response.PinnedJobIds = settings.PinnedJobIds.ConvertAll(x => x.ToString());
+				response.PinnedBisectTaskIds = settings.PinnedBisectTaskIds.ConvertAll(x => x.ToString());
 
 				if (settings.JobTemplateSettings != null && settings.JobTemplateSettings.Count > 0)
 				{
diff --git a/Engine/Source/Programs/Horde/Horde.Server/Users/IUserCollection.cs b/Engine/Source/Programs/Horde/Horde.Server/Users/IUserCollection.cs
--- a/Engine/Source/Programs/Horde/Horde.Server/Users/IUserCollection.cs
+++ b/Engine/Source/Programs/Horde/Horde.Server/Users/IUserCollection.cs
@@ -7,6 +7,7 @@ using Horde.Server.Jobs;
 using Horde.Server.Utilities;
 using MongoDB.Bson;
 using EpicGames.Horde.Api;
+using Horde.Server.Jobs.Bisect;
 
 namespace Horde.Server.Users
 {
@@ -93,8 +94,10 @@ namespace Horde.Server.Users
 		/// <param name="addPinnedJobIds"></param>
 		/// <param name="removePinnedJobIds"></param>
 		/// <param name="templateOptions"></param>
+		/// <param name="addBisectTaskIds"></param>
+		/// <param name="removeBisectTaskIds"></param>
 		/// <returns>Updated user object</returns>
-		Task UpdateSettingsAsync(UserId userId, bool? enableExperimentalFeatures = null, BsonValue? dashboardSettings = null, IEnumerable<JobId>? addPinnedJobIds = null, IEnumerable<JobId>? removePinnedJobIds = null, UpdateUserJobTemplateOptions? templateOptions = null);
+		Task UpdateSettingsAsync(UserId userId, bool? enableExperimentalFeatures = null, BsonValue? dashboardSettings = null, IEnumerable<JobId>? addPinnedJobIds = null, IEnumerable<JobId>? removePinnedJobIds = null, UpdateUserJobTemplateOptions? templateOptions = null, IEnumerable<BisectTaskId>? addBisectTaskIds = null, IEnumerable<BisectTaskId>? removeBisectTaskIds = null);
 	}
 
 	/// <summary>
diff --git a/Engine/Source/Programs/Horde/Horde.Server/Users/UserCollection.cs b/Engine/Source/Programs/Horde/Horde.Server/Users/UserCollection.cs
--- a/Engine/Source/Programs/Horde/Horde.Server/Users/UserCollection.cs
+++ b/Engine/Source/Programs/Horde/Horde.Server/Users/UserCollection.cs
@@ -13,6 +13,7 @@ using MongoDB.Bson.Serialization;
 using MongoDB.Bson.Serialization.Attributes;
 using MongoDB.Driver;
 using EpicGames.Horde.Api;
+using Horde.Server.Jobs.Bisect;
 
 namespace Horde.Server.Users
 {
@@ -33,6 +34,7 @@ namespace Horde.Server.Users
 
 			public BsonValue DashboardSettings { get; set; } = BsonNull.Value;
 			public List<JobId> PinnedJobIds { get; set; } = new List<JobId>();
+			public List<BisectTaskId> PinnedBisectTaskIds { get; set; } = new List<BisectTaskId>();
 
 			string IUser.Name => Claims.FirstOrDefault(x => String.Equals(x.Type, "name", StringComparison.Ordinal))?.Value ?? PrimaryClaim.Value;
 			string IUser.Login => Claims.FirstOrDefault(x => String.Equals(x.Type, ClaimTypes.Name, StringComparison.Ordinal))?.Value ?? PrimaryClaim.Value;
@@ -43,6 +45,7 @@ namespace Horde.Server.Users
 
 			UserId IUserSettings.UserId => Id;
 			IReadOnlyList<JobId> IUserSettings.PinnedJobIds => PinnedJobIds;
+			IReadOnlyList<BisectTaskId> IUserSettings.PinnedBisectTaskIds => PinnedBisectTaskIds;
 
 			IReadOnlyList<IUserJobTemplateSettings>? IUserSettings.JobTemplateSettings => null;
 		}
@@ -158,7 +161,7 @@ namespace Horde.Server.Users
 		}
 
 		/// <inheritdoc/>
-		public async Task UpdateSettingsAsync(UserId userId, bool? enableExperimentalFeatures, BsonValue? dashboardSettings = null, IEnumerable<JobId>? addPinnedJobIds = null, IEnumerable<JobId>? removePinnedJobIds = null, UpdateUserJobTemplateOptions? templateOptions = null)
+		public async Task UpdateSettingsAsync(UserId userId, bool? enableExperimentalFeatures, BsonValue? dashboardSettings = null, IEnumerable<JobId>? addPinnedJobIds = null, IEnumerable<JobId>? removePinnedJobIds = null, UpdateUserJobTemplateOptions? templateOptions = null, IEnumerable<BisectTaskId>? addBisectTaskIds = null, IEnumerable<BisectTaskId>? removeBisectTaskIds = null)
 		{
 			if (addPinnedJobIds != null)
 			{
diff --git a/Engine/Source/Programs/Horde/Horde.Server/Users/UserCollectionV2.cs b/Engine/Source/Programs/Horde/Horde.Server/Users/UserCollectionV2.cs
--- a/Engine/Source/Programs/Horde/Horde.Server/Users/UserCollectionV2.cs
+++ b/Engine/Source/Programs/Horde/Horde.Server/Users/UserCollectionV2.cs
@@ -15,6 +15,7 @@ using MongoDB.Bson.Serialization;
 using MongoDB.Bson.Serialization.Attributes;
 using MongoDB.Driver;
 using EpicGames.Horde.Api;
+using Horde.Server.Jobs.Bisect;
 
 namespace Horde.Server.Users
 {
@@ -130,9 +131,11 @@ namespace Horde.Server.Users
 
 			public BsonValue DashboardSettings { get; set; } = BsonNull.Value;
 			public List<JobId> PinnedJobIds { get; set; } = new List<JobId>();
+			public List<BisectTaskId> PinnedBisectTaskIds { get; set; } = new List<BisectTaskId>();
 
 			UserId IUserSettings.UserId => Id;
 			IReadOnlyList<JobId> IUserSettings.PinnedJobIds => PinnedJobIds;
+			IReadOnlyList<BisectTaskId> IUserSettings.PinnedBisectTaskIds => PinnedBisectTaskIds;
 
 			public List<JobTemplateSettingsDocument> JobTemplateSettings { get; set; } = new List<JobTemplateSettingsDocument>();
 			IReadOnlyList<IUserJobTemplateSettings>? IUserSettings.JobTemplateSettings => JobTemplateSettings;
@@ -153,6 +156,7 @@ namespace Horde.Server.Users
 				EnableExperimentalFeatures = other.EnableExperimentalFeatures;
 				DashboardSettings = other.DashboardSettings;
 				PinnedJobIds = new List<JobId>(other.PinnedJobIds);
+				PinnedBisectTaskIds = new List<BisectTaskId>(other.PinnedBisectTaskIds);
 			}
 		}
 
@@ -353,7 +357,7 @@ namespace Horde.Server.Users
 		}
 
 		/// <inheritdoc/>
-		public async Task UpdateSettingsAsync(UserId userId, bool? enableExperimentalFeatures = null, BsonValue? dashboardSettings = null, IEnumerable<JobId>? addPinnedJobIds = null, IEnumerable<JobId>? removePinnedJobIds = null, UpdateUserJobTemplateOptions? templateOptions = null)
+		public async Task UpdateSettingsAsync(UserId userId, bool? enableExperimentalFeatures = null, BsonValue? dashboardSettings = null, IEnumerable<JobId>? addPinnedJobIds = null, IEnumerable<JobId>? removePinnedJobIds = null, UpdateUserJobTemplateOptions? templateOptions = null, IEnumerable<BisectTaskId>? addBisectTaskIds = null, IEnumerable<BisectTaskId>? removeBisectTaskIds = null)
 		{
 			List<UpdateDefinition<UserSettingsDocument>> updates = new List<UpdateDefinition<UserSettingsDocument>>();
 			if (enableExperimentalFeatures != null)
@@ -373,6 +377,16 @@ namespace Horde.Server.Users
 				updates.Add(Builders<UserSettingsDocument>.Update.PullAll(x => x.PinnedJobIds, removePinnedJobIds));
 			}
 
+			if (addBisectTaskIds != null && addBisectTaskIds.Any())
+			{
+				updates.Add(Builders<UserSettingsDocument>.Update.AddToSetEach(x => x.PinnedBisectTaskIds, addBisectTaskIds));
+			}
+			if (removeBisectTaskIds != null && removeBisectTaskIds.Any())
+			{
+				updates.Add(Builders<UserSettingsDocument>.Update.PullAll(x => x.PinnedBisectTaskIds, removeBisectTaskIds));
+			}
+
+
 			if (templateOptions != null)
 			{
 				JobTemplateSettingsDocument doc = new JobTemplateSettingsDocument(templateOptions.StreamId, templateOptions.TemplateId, templateOptions.TemplateHash, templateOptions.Arguments.ToList());
diff --git a/Engine/Source/Programs/Horde/Horde.Server/Users/UserController.cs b/Engine/Source/Programs/Horde/Horde.Server/Users/UserController.cs
--- a/Engine/Source/Programs/Horde/Horde.Server/Users/UserController.cs
+++ b/Engine/Source/Programs/Horde/Horde.Server/Users/UserController.cs
@@ -15,6 +15,7 @@ using Microsoft.AspNetCore.Authorization;
 using Microsoft.AspNetCore.Mvc;
 using Microsoft.Extensions.Options;
 using EpicGames.Horde.Api;
+using Horde.Server.Jobs.Bisect;
 
 namespace Horde.Server.Users
 {
@@ -94,7 +95,7 @@ namespace Horde.Server.Users
 				return BadRequest("Current user does not have a registered profile");
 			}
 
-			await _userCollection.UpdateSettingsAsync(userId.Value, request.EnableExperimentalFeatures, request.DashboardSettings?.ToBsonValue(), request.AddPinnedJobIds?.Select(x => JobId.Parse(x)), request.RemovePinnedJobIds?.Select(x => JobId.Parse(x)));
+			await _userCollection.UpdateSettingsAsync(userId.Value, request.EnableExperimentalFeatures, request.DashboardSettings?.ToBsonValue(), request.AddPinnedJobIds?.Select(x => JobId.Parse(x)), request.RemovePinnedJobIds?.Select(x => JobId.Parse(x)), null, request.AddPinnedBisectTaskIds?.Select(x => BisectTaskId.Parse(x)), request.RemovePinnedBisectTaskIds?.Select(x => BisectTaskId.Parse(x)));
 			return Ok();
 		}
 	}
diff --git a/Engine/Source/Programs/Horde/HordeDashboard/src/backend/Api.ts b/Engine/Source/Programs/Horde/HordeDashboard/src/backend/Api.ts
--- a/Engine/Source/Programs/Horde/HordeDashboard/src/backend/Api.ts
+++ b/Engine/Source/Programs/Horde/HordeDashboard/src/backend/Api.ts
@@ -1075,7 +1075,7 @@ export type GetJobResponse = {
 	/** The user that started this job */
 	abortedByUserInfo?: GetThinUserInfoResponse;
 
-	/** Whether job was created by a bisection task */
+	/** Whether job was created by a bisect task */
 	startedByBisectTaskId?: string;
 
 	/** The roles to impersonate when executing this job */
@@ -3220,6 +3220,9 @@ export type GetUserResponse = {
 	/** List of pinned job ids */
 	pinnedJobIds?: string[];
 
+	/** List of pinned bisect task ids */
+	pinnedBisectTaskIds?: string[];
+
 }
 
 /** Basic information about a user. May be embedded in other responses.*/
@@ -3251,6 +3254,12 @@ export type UpdateUserRequest = {
 	/** Job ids to remove from the pinned list */
 	removePinnedJobIds?: string[];
 
+	/** Bisect ids to add to the pinned list */
+	addPinnedBisectTaskIds?: string[];
+
+	/** Bisect ids to remove from the pinned list */
+	removePinnedBisectTaskIds?: string[];
+
 	/** Whether to enable experimental features for this user */
 	enableExperimentalFeatures?: boolean;
 
@@ -4666,7 +4675,7 @@ export type GetToolSummaryResponse = {
 	version?: string;
 }
 
-/** Job Bisection */
+/** Job Bisect */
 
 /// State of a bisect task
 export enum BisectTaskState {
@@ -4747,10 +4756,10 @@ export type GetBisectTaskResponse = {
 	/// Outcome to search for
 	outcome: JobStepOutcome;
 
-	/// Starting job id for the bisection
+	/// Starting job id for the bisect
 	initialJobId: string;
 
-	/// Starting change for the bisection
+	/// Starting change for the bisect
 	initialChange: number;
 
 	/// First known job id that is broken
@@ -4759,7 +4768,13 @@ export type GetBisectTaskResponse = {
 	/// Changelist number of the first broken job id
 	currentChange: number;
 
-	/// The steps that have been run on bisection
+	/// Next step id of a running bisect
+	nextJobId?: string;
+
+	/// Next step id of a running bisect
+	nextJobChange?: number;
+
+	/// The steps that have been run on bisect
 	steps?: GetJobStepRefResponse[];
 
 }
diff --git a/Engine/Source/Programs/Horde/HordeDashboard/src/backend/Backend.ts b/Engine/Source/Programs/Horde/HordeDashboard/src/backend/Backend.ts
--- a/Engine/Source/Programs/Horde/HordeDashboard/src/backend/Backend.ts
+++ b/Engine/Source/Programs/Horde/HordeDashboard/src/backend/Backend.ts
@@ -1664,6 +1664,17 @@ export class Backend {
         });
     }
 
+    getBisections(query: { id?:string[], ownerId?: string, jobId?: string, minCreateTime?: string, maxCreateTime?: string, index?: number, count?: number }) {
+        
+        return new Promise<GetBisectTaskResponse[]>((resolve, reject) => {
+            this.backend.get(`/api/v1/bisect`, {params: query}).then((value) => {
+                resolve(value.data as GetBisectTaskResponse[]);
+            }).catch(reason => {
+                reject(reason);
+            });
+        });    
+    }
+
     // update a visection task
     updateBisectTask(id: string, request: UpdateBisectTaskRequest): Promise<void> {
         return new Promise<void>((resolve, reject) => {
diff --git a/Engine/Source/Programs/Horde/HordeDashboard/src/backend/Dashboard.ts b/Engine/Source/Programs/Horde/HordeDashboard/src/backend/Dashboard.ts
--- a/Engine/Source/Programs/Horde/HordeDashboard/src/backend/Dashboard.ts
+++ b/Engine/Source/Programs/Horde/HordeDashboard/src/backend/Dashboard.ts
@@ -69,6 +69,36 @@ export class Dashboard {
         this.setUpdated();
     }
 
+    bisectPinned(id: string | undefined) {
+        return !!this.pinnedBisectTaskIds.find(j => j === id);
+    }
+
+    pinBisect(id: string) {
+
+        if (this.data.pinnedBisectTaskIds!.find(j => j === id)) {
+            return;
+        }
+
+        this.data.pinnedBisectTaskIds!.push(id);
+
+        backend.updateUser({ addPinnedBisectTaskIds: [id] });
+
+        this.setUpdated();
+    }
+
+    unpinBisect(id: string) {
+
+        if (!this.data.pinnedBisectTaskIds!.find(j => j === id)) {
+            return;
+        }
+
+        this.data.pinnedBisectTaskIds = this.data.pinnedBisectTaskIds!.filter(j => j !== id);
+
+        backend.updateUser({ removePinnedBisectTaskIds: [id] });
+
+        this.setUpdated();
+    }
+
     clearPinnedJobs() {
 
         if (!this.data.pinnedJobIds?.length) {
@@ -182,6 +212,12 @@ export class Dashboard {
         return this.data.pinnedJobIds ?? [];
     }
 
+    get pinnedBisectTaskIds(): string[] {
+
+        return this.data.pinnedBisectTaskIds ?? [];
+    }
+
+
     getLastJobTemplateSettings(streamId: string, templateIds: string[]): GetJobTemplateSettingsResponse | undefined {
 
         try {
diff --git a/Engine/Source/Programs/Horde/HordeDashboard/src/backend/PollBase.ts b/Engine/Source/Programs/Horde/HordeDashboard/src/backend/PollBase.ts
--- a/Engine/Source/Programs/Horde/HordeDashboard/src/backend/PollBase.ts
+++ b/Engine/Source/Programs/Horde/HordeDashboard/src/backend/PollBase.ts
@@ -88,6 +88,10 @@ export abstract class PollBase {
     @observable
     updated: number = 0;
 
+    subscribe() {
+        if (this.updated) { }
+    }
+
     updating = false;
 
     protected timeoutId: any;
diff --git a/Engine/Source/Programs/Horde/HordeDashboard/src/base/utilities/timeUtils.ts b/Engine/Source/Programs/Horde/HordeDashboard/src/base/utilities/timeUtils.ts
--- a/Engine/Source/Programs/Horde/HordeDashboard/src/base/utilities/timeUtils.ts
+++ b/Engine/Source/Programs/Horde/HordeDashboard/src/base/utilities/timeUtils.ts
@@ -131,8 +131,8 @@ export const getStepETA = (step: StepData, job: JobData): HordeTime => {
 
 export const getStepFinishTime = (step: StepData): HordeTime => {
 
-    let displayFinish = "";    
-    let serverFinish = "";    
+    let displayFinish = "";
+    let serverFinish = "";
 
     const format = dashboard.display24HourClock ? "HH:mm:ss z" : "LT z";
 
@@ -150,8 +150,8 @@ export const getStepFinishTime = (step: StepData): HordeTime => {
 
 export const getStepStartTime = (step: StepData): HordeTime => {
 
-    let displayStart = "";    
-    let serverStart = "";    
+    let displayStart = "";
+    let serverStart = "";
 
     const format = dashboard.display24HourClock ? "HH:mm:ss z" : "LT z";
 
@@ -170,8 +170,8 @@ export const getStepStartTime = (step: StepData): HordeTime => {
 
 export const getLabelFinishTime = (label: JobLabel, job: JobData): HordeTime => {
 
-    let displayFinish = "";    
-    let serverFinish = "";    
+    let displayFinish = "";
+    let serverFinish = "";
 
     const format = dashboard.display24HourClock ? "HH:mm:ss z" : "LT z";
 
@@ -235,7 +235,7 @@ export const getStepPercent = (step: StepData): number | undefined => {
     return percent;
 };
 
-export const getElapsedString = (start: Moment, end: Moment, includeSeconds:boolean = true): string => {
+export const getElapsedString = (start: Moment, end: Moment, includeSeconds: boolean = true): string => {
 
     let duration = "";
     const d = moment.duration(end.diff(start));
@@ -259,7 +259,7 @@ export const getElapsedString = (start: Moment, end: Moment, includeSeconds:bool
         if (d.days()) {
             duration += `${d.days()}d `;
         }
-        
+
         if (d.hours()) {
             duration += `${d.hours()}h `;
         }
@@ -361,7 +361,7 @@ export const getHumanTime = (timeIn: Date | string | undefined): string => {
     const delta = nowDay - timeDay;
 
     if (delta > 2 || now.year() !== time.year()) {
-        return time.format('MMM Do');    
+        return time.format('MMM Do');
     }
 
     if (delta === 1) {
@@ -376,12 +376,12 @@ export const getHumanTime = (timeIn: Date | string | undefined): string => {
         return "Monday";
     }
 
-    return time.format('MMM Do');    
+    return time.format('MMM Do');
 
 };
 
 
-export const getShortNiceTime = (timeIn: Date | string | undefined, relative: boolean = false, includeHour:boolean = false, includeDay: boolean = true): string => {
+export const getShortNiceTime = (timeIn: Date | string | undefined, relative: boolean = false, includeHour: boolean = false, includeDay: boolean = true): string => {
 
     if (!timeIn) {
         return "";
@@ -408,7 +408,7 @@ export const getShortNiceTime = (timeIn: Date | string | undefined, relative: bo
         else
             timeStr = time.format(format);
     }
-    
+
     return timeStr;
 
 };
@@ -439,4 +439,14 @@ export const getNiceTime = (timeIn: Date | string | undefined, relative: boolean
 
     return timeStr;
 
-};
\ No newline at end of file
+};
+
+export const getMongoIdDate = (mongoId: string): Date | undefined => {
+
+    if (!mongoId || mongoId.length < 8) {
+        return undefined;
+    }
+
+    const timestamp = mongoId.substring(0, 8)
+    return new Date(parseInt(timestamp, 16) * 1000);
+}
\ No newline at end of file
diff --git a/Engine/Source/Programs/Horde/HordeDashboard/src/components/ChangeButton.tsx b/Engine/Source/Programs/Horde/HordeDashboard/src/components/ChangeButton.tsx
--- a/Engine/Source/Programs/Horde/HordeDashboard/src/components/ChangeButton.tsx
+++ b/Engine/Source/Programs/Horde/HordeDashboard/src/components/ChangeButton.tsx
@@ -3,7 +3,7 @@
 import { Point } from '@fluentui/react';
 import { ContextualMenu, ContextualMenuItemType, IContextualMenuItem, Stack, Text } from '@fluentui/react';
 import React, { MutableRefObject, useState } from 'react';
-import { GetChangeSummaryResponse, GetJobResponse, GetJobStepRefResponse, JobData, JobStepBatchError, JobStepBatchState, JobStepOutcome, JobStepState } from '../backend/Api';
+import { GetBatchResponse, GetChangeSummaryResponse, GetJobResponse, GetJobStepRefResponse, GetThinUserInfoResponse, JobData, JobStepBatchError, JobStepBatchState, JobStepOutcome, JobStepState } from '../backend/Api';
 import dashboard, { StatusColor } from "../backend/Dashboard";
 import { projectStore } from '../backend/ProjectStore';
 
@@ -14,7 +14,19 @@ export type ChangeContextMenuTarget = {
 
 }
 
-export const ChangeContextMenu: React.FC<{ target: ChangeContextMenuTarget, onDismiss?: () => void, job?: JobData, stepRef?: GetJobStepRefResponse, commit?: GetChangeSummaryResponse, rangeCL?: number }> = ({ target, onDismiss, job, stepRef, commit, rangeCL }) => {
+// partial JobResponse type, so we don't need full job data when presenting change cl's
+export type JobParameters = {
+   id?: string;
+   streamId: string;
+   change?: number;
+   preflightChange?: number;
+   preflightDescription?: string;   
+   abortedByUserInfo?: GetThinUserInfoResponse;
+   startedByUserInfo?: GetThinUserInfoResponse;
+   batches?: GetBatchResponse[];
+}
+
+export const ChangeContextMenu: React.FC<{ target: ChangeContextMenuTarget, onDismiss?: () => void, job?: JobParameters, stepRef?: GetJobStepRefResponse, commit?: GetChangeSummaryResponse, rangeCL?: number }> = ({ target, onDismiss, job, stepRef, commit, rangeCL }) => {
 
    if (!job) {
       return null;
@@ -101,14 +113,6 @@ export const ChangeContextMenu: React.FC<{ target: ChangeContextMenuTarget, onDi
       onClick: () => copyToClipboard(change.replace("PF ", ""))
    });
 
-   if (job) {
-      menuItems.push({
-         key: 'copy_job_to_clipboard',
-         text: 'Copy Job to Clipboard',
-         onClick: () => copyToClipboard(window.location.protocol + "//" + window.location.hostname + `/job/${job.id}`)
-      });
-   }
-
    menuItems.push({ key: 'open_in_swarm', text: 'Open CL in Swarm', onClick: (ev) => { window.open(`${dashboard.swarmUrl}/changes/${change.replace("PF ", "")}`) } })
 
    // range
@@ -129,6 +133,22 @@ export const ChangeContextMenu: React.FC<{ target: ChangeContextMenuTarget, onDi
       const rangeText = `Open CL Range ${lowCL} - ${highCL}`;
       menuItems.push({ key: 'open_in_swarm_range', text: rangeText, onClick: (ev) => { window.open(url) } })
    }
+
+   if (job) {
+
+      menuItems.push({
+         key: 'view_job',
+         text: 'Open Job Details',
+         href: window.location.protocol + "//" + window.location.hostname + `/job/${job.id}`         
+      });
+
+
+      menuItems.push({
+         key: 'copy_job_to_clipboard',
+         text: 'Copy Job to Clipboard',
+         onClick: () => copyToClipboard(window.location.protocol + "//" + window.location.hostname + `/job/${job.id}`)
+      });
+   }
    
    return (<ContextualMenu
       styles={{ root: {paddingBottom: 12, paddingRight: 24, paddingLeft: 8, paddingTop: 12},list: { selectors: { '.ms-ContextualMenu-itemText': { fontSize: "10px", paddingLeft: 8 } } } }}
@@ -141,7 +161,7 @@ export const ChangeContextMenu: React.FC<{ target: ChangeContextMenuTarget, onDi
 
 }
 
-function getJobSummary(job: GetJobResponse): { text: string, color: string } {
+function getJobSummary(job: JobParameters): { text: string, color: string } {
 
    const colors = dashboard.getStatusColors();
    let color = colors.get(StatusColor.Skipped)!;
@@ -215,7 +235,7 @@ function getJobSummary(job: GetJobResponse): { text: string, color: string } {
 }
 
 
-export const ChangeButton: React.FC<{ job?: JobData, stepRef?: GetJobStepRefResponse, commit?: GetChangeSummaryResponse, hideAborted?: boolean, rangeCL?: number, pinned?: boolean, prefix?: string }> = ({ job, stepRef, commit, hideAborted, rangeCL, pinned, prefix }) => {
+export const ChangeButton: React.FC<{  job?: JobParameters, stepRef?: GetJobStepRefResponse, commit?: GetChangeSummaryResponse, hideAborted?: boolean, rangeCL?: number, pinned?: boolean, prefix?: string, buttonColor?: string }> = ({ job, stepRef, commit, hideAborted, rangeCL, pinned, prefix, buttonColor }) => {
 
    const [menuShown, setMenuShown] = useState(false);
 
@@ -244,7 +264,7 @@ export const ChangeButton: React.FC<{ job?: JobData, stepRef?: GetJobStepRefResp
    
    return (<Stack verticalAlign="center" verticalFill={true} horizontalAlign="start"> <div style={{ paddingBottom: "1px" }}>
       <Stack tokens={{ childrenGap: 4 }}>
-         <span ref={spanRef} style={{ padding: "2px 6px 2px 6px", height: "15px", cursor: "pointer" }} className={job.startedByUserInfo ? "cl-callout-button-user" : "cl-callout-button"} onClick={(ev) => { ev.stopPropagation(); ev.preventDefault(); setMenuShown(!menuShown) }} >{change}</span>
+         <span ref={spanRef} style={{ padding: "2px 6px 2px 6px", height: "15px", cursor: "pointer", background: buttonColor ? `${buttonColor}` : undefined }} className={job.startedByUserInfo ? "cl-callout-button-user" : "cl-callout-button"} onClick={(ev) => { ev.stopPropagation(); ev.preventDefault(); setMenuShown(!menuShown) }} >{change}</span>
          {(!!showStatus) && <span ref={spanRef} style={{ padding: "2px 6px 2px 6px", height: "16px", cursor: "pointer", userSelect: "none", fontFamily: "Horde Open Sans SemiBold", fontSize: "10px", backgroundColor: color, color: "rgb(255, 255, 255)" }} onClick={(ev) => { ev.preventDefault(); setMenuShown(!menuShown) }}>{text}</span>}
       </Stack>
       {menuShown && <ChangeContextMenu target={{ ref: spanRef }} job={job} commit={commit} stepRef={stepRef} rangeCL={rangeCL} onDismiss={() => setMenuShown(false)} />}
diff --git a/Engine/Source/Programs/Horde/HordeDashboard/src/components/UserHome.tsx b/Engine/Source/Programs/Horde/HordeDashboard/src/components/UserHome.tsx
--- a/Engine/Source/Programs/Horde/HordeDashboard/src/components/UserHome.tsx
+++ b/Engine/Source/Programs/Horde/HordeDashboard/src/components/UserHome.tsx
@@ -7,7 +7,7 @@ import moment from "moment-timezone";
 import { default as React, useEffect, useState } from 'react';
 import { Link, useLocation, useNavigate } from "react-router-dom";
 import backend, { useBackend } from "../backend";
-import { GetIssueResponse, GetStepResponse, JobData, JobQuery, JobState, JobStepOutcome, LabelData, LabelOutcome, LabelState, ProjectData, StepData, StreamData } from "../backend/Api";
+import { GetBisectTaskResponse, GetIssueResponse, GetStepResponse, JobData, JobQuery, JobState, JobStepOutcome, LabelData, LabelOutcome, LabelState, ProjectData, StepData, StreamData } from "../backend/Api";
 import dashboard, { StatusColor } from "../backend/Dashboard";
 import graphCache, { GraphQuery } from '../backend/GraphCache';
 import { useWindowSize } from '../base/utilities/hooks';
@@ -21,9 +21,10 @@ import { useQuery } from './JobDetailCommon';
 import { JobOperationsContextMenu } from "./JobOperationsContextMenu";
 import { IssueStatusIcon, StepStatusIcon } from "./StatusIcon";
 import { TopNav } from './TopNav';
+import { BisectionList } from './bisection/BisectionList';
+import { PollBase } from '../backend/PollBase';
 import { PreflightConfigModal } from './preflights/PreflightConfigCheck';
 
-
 type JobItem = {
    key: string;
    job: JobData;
@@ -226,7 +227,7 @@ const HealthPanel: React.FC = observer(() => {
 
    const columns: IColumn[] = [
       { key: 'health_column1', name: 'Summary', minWidth: 600, maxWidth: 600, isResizable: false },
-      { key: 'health_column2', name: 'Workflow', minWidth: 100, maxWidth: 100, isResizable: false},
+      { key: 'health_column2', name: 'Workflow', minWidth: 100, maxWidth: 100, isResizable: false },
       { key: 'health_column3', name: 'Status', minWidth: 320, maxWidth: 320, isResizable: false },
       { key: 'health_column4', name: 'Opened', minWidth: 100, maxWidth: 100, isResizable: false },
    ];
@@ -262,7 +263,7 @@ const HealthPanel: React.FC = observer(() => {
             summary = summary.slice(0, 100) + "...";
          }
 
-         return <Stack style={{padding: 8}} horizontal disableShrink={true}>{<IssueStatusIcon issue={issue} />}<Text>{`Issue ${issue.id} - ${summary}`}</Text></Stack>;
+         return <Stack style={{ padding: 8 }} horizontal disableShrink={true}>{<IssueStatusIcon issue={issue} />}<Text>{`Issue ${issue.id} - ${summary}`}</Text></Stack>;
       }
 
       if (column.name === "Status") {
@@ -271,7 +272,7 @@ const HealthPanel: React.FC = observer(() => {
             status += ` (${ack})`;
          }
 
-         return <Stack style={{padding: 8}} horizontalAlign="end" disableShrink={true}><Text>{status}</Text></Stack>;
+         return <Stack style={{ padding: 8 }} horizontalAlign="end" disableShrink={true}><Text>{status}</Text></Stack>;
       }
 
       if (column.name === "Workflow") {
@@ -279,7 +280,7 @@ const HealthPanel: React.FC = observer(() => {
             return null;
          }
 
-         return <a onClick={(e) => e.stopPropagation()} style={{ fontSize: "13px" }} href={issue.workflowThreadUrl} target="_blank" rel="noreferrer"><Stack style={{padding: 8}} horizontalAlign="start" disableShrink={true}>Slack Thread</Stack></a>;
+         return <a onClick={(e) => e.stopPropagation()} style={{ fontSize: "13px" }} href={issue.workflowThreadUrl} target="_blank" rel="noreferrer"><Stack style={{ padding: 8 }} horizontalAlign="start" disableShrink={true}>Slack Thread</Stack></a>;
       }
 
       if (column.name === "Opened") {
@@ -293,7 +294,7 @@ const HealthPanel: React.FC = observer(() => {
    const classes = mergeStyleSets({
       detailsRow: {
          selectors: {
-            '.ms-DetailsRow-cell': {               
+            '.ms-DetailsRow-cell': {
                overflow: "hidden",
                whiteSpace: "nowrap"
             }
@@ -549,7 +550,7 @@ const JobsPanel: React.FC<{ includeOtherPreflights: boolean }> = observer(({ inc
 
       dashboard.startPolling();
 
-      return () => {         
+      return () => {
          dashboard.stopPolling();
          jobHandler.clear();
       };
@@ -1033,7 +1034,7 @@ const JobsPanel: React.FC<{ includeOtherPreflights: boolean }> = observer(({ inc
                   </Stack>
                </Stack>
                <Stack styles={{ root: { paddingLeft: 4, paddingRight: 0, paddingTop: 8, paddingBottom: 4 } }}>
-                  {!!dashboard.pinnedJobsIds.length && !jobItems.length && <Stack><Spinner size={SpinnerSize.large}/></Stack>}
+                  {!!dashboard.pinnedJobsIds.length && !jobItems.length && <Stack><Spinner size={SpinnerSize.large} /></Stack>}
                   {!!jobItems.length && <DetailsList
                      styles={{ root: { paddingLeft: 8, paddingRight: 8, marginBottom: 18 } }}
                      compact={true}
@@ -1059,13 +1060,83 @@ const JobsPanel: React.FC<{ includeOtherPreflights: boolean }> = observer(({ inc
 
 });
 
+// Bisections
+
+class BisectionHandler extends PollBase {
+
+   constructor(pollTime = 5000) {
+
+      super(pollTime);
+
+   }
+
+   clear() {
+      super.stop();
+   }
+
+   async poll(): Promise<void> {
+
+      try {
+
+         if (dashboard.pinnedBisectTaskIds?.length) {
+            this.bisections = await backend.getBisections({ id: dashboard.pinnedBisectTaskIds });
+         }
+
+         this.setUpdated();
+
+      } catch (err) {
+
+      }
+
+   }
+
+   bisections: GetBisectTaskResponse[] = [];
+}
+
+const bisectionHandler = new BisectionHandler();
+
+const BisectionPanel: React.FC = observer(() => {
+
+   useEffect(() => {
+
+      bisectionHandler.start();
+
+      return () => {
+         bisectionHandler.clear();
+      };
+
+   }, []);
+
+   // subscribe
+   if (dashboard.updated) { };
+   if (bisectionHandler.updated) { };
+
+   if (bisectionHandler.bisections.length === 0) {
+      return null;
+   }
+
+   return (<Stack style={{ width: homeWidth, marginLeft: 4 }}>
+      <Stack className={hordeClasses.raised}>
+         <Stack tokens={{ childrenGap: 12 }}>
+            <Text variant="mediumPlus" styles={{ root: { fontFamily: "Horde Open Sans SemiBold" } }}>Bisections</Text>
+            <Stack styles={{ root: { paddingLeft: 4, paddingRight: 0, paddingTop: 8, paddingBottom: 12 } }}>
+               <div style={{ overflowY: 'auto', overflowX: 'hidden', maxHeight: "384px" }} data-is-scrollable={true}>
+                  <BisectionList bisections={bisectionHandler.bisections} />
+               </div>
+            </Stack>
+         </Stack>
+      </Stack>
+   </Stack>);
+
+});
+
 const UserHomeViewInner: React.FC = () => {
-      
+
    const search = new URLSearchParams(window.location.search);
    const change = !search.get("preflightconfig") ? "" : search.get("preflightconfig")!;
 
    const navigate = useNavigate();
-   
+
 
    useEffect(() => {
 
@@ -1078,13 +1149,14 @@ const UserHomeViewInner: React.FC = () => {
    handler.startPolling();
 
    return <Stack tokens={{ childrenGap: 0 }} styles={{ root: { backgroundColor: modeColors.background, margin: 0, paddingTop: 8 } }}>
-      {(!!change || !!search.has("preflightconfig")) && <PreflightConfigModal onClose={() => {navigate("/index", {replace: true})}}/>}
+      {(!!change || !!search.has("preflightconfig")) && <PreflightConfigModal onClose={() => { navigate("/index", { replace: true }) }} />}
       <Stack style={{ padding: 0 }} className={detailClasses.detailsRow}>
          <FocusZone direction={FocusZoneDirection.vertical} style={{ padding: 0 }}>
             <div className={detailClasses.container} style={{ width: "100%", height: 'calc(100vh - 208px)', position: 'relative' }} data-is-scrollable={true}>
                <ScrollablePane scrollbarVisibility={ScrollbarVisibility.auto} onScroll={() => { }}>
                   <Stack tokens={{ childrenGap: 18 }} style={{ padding: 0 }}>
                      <ProjectsPanel />
+                     <BisectionPanel />
                      <HealthPanel />
                      <JobsPanel includeOtherPreflights={false} />
                   </Stack>
diff --git a/Engine/Source/Programs/Horde/HordeDashboard/src/components/bisection/BisectionList.tsx b/Engine/Source/Programs/Horde/HordeDashboard/src/components/bisection/BisectionList.tsx
new file mode 100644
--- /dev/null
+++ b/Engine/Source/Programs/Horde/HordeDashboard/src/components/bisection/BisectionList.tsx
@@ -0,0 +1,645 @@
+// Copyright Epic Games, Inc. All Rights Reserved.
+
+import { Callout, DefaultButton, DetailsList, DetailsListLayoutMode, DirectionalHint, FontIcon, IDetailsListProps, IconButton, MessageBar, MessageBarType, Modal, PrimaryButton, SelectionMode, Stack, Text, mergeStyleSets } from '@fluentui/react';
+import { useConst } from '@fluentui/react-hooks';
+import * as d3 from "d3";
+import { action, makeObservable, observable } from 'mobx';
+import { observer } from 'mobx-react-lite';
+import { default as React, useState } from 'react';
+import { Link } from "react-router-dom";
+import backend from "../../backend";
+import { BisectTaskState, GetBisectTaskResponse, JobStepOutcome } from "../../backend/Api";
+import dashboard, { StatusColor } from "../../backend/Dashboard";
+import { projectStore } from '../../backend/ProjectStore';
+import { getShortNiceTime } from '../../base/utilities/timeUtils';
+import { hordeClasses, modeColors } from '../../styles/Styles';
+import { ChangeButton, JobParameters } from '../ChangeButton';
+
+type SelectionType = d3.Selection<SVGGElement, unknown, null, undefined>;
+
+class Tooltip {
+
+   constructor() {
+      makeObservable(this);
+   }
+
+   subscribe() { if (this.updated) { } }
+
+   @observable
+   updated = 0
+
+   @action
+   update(bisection?: GetBisectTaskResponse, change?: number, x?: number, y?: number) {
+      this.bisection = bisection;
+      this.change = change ?? 0;
+      this.x = x ?? 0;
+      this.y = y ?? 0;
+      this.updated++;
+   }
+
+   @action
+   freeze(frozen: boolean) {
+      this.frozen = frozen;
+      this.updated++;
+   }
+
+   get step() {
+      return this.bisection?.steps?.find((s => s.change === this.change));
+   }
+
+   bisection?: GetBisectTaskResponse;
+   change: number = 0;
+   x: number = 0;
+   y: number = 0;
+   frozen = false;
+}
+
+
+const GraphTooltip: React.FC<{ renderer: BisectionRenderer }> = observer(({ renderer }) => {
+
+   const tooltip = renderer.tooltip;
+   tooltip.subscribe();
+
+   const bisection = tooltip.bisection!;
+   if (!bisection) {
+      return null;
+   }
+
+   let tipX = tooltip.x;
+   let offsetX = 32;
+   let translateX = "0%";
+
+   if (tipX > 1000) {
+      offsetX = -32;
+      translateX = "-100%";
+   }
+
+   const translateY = "-50%";
+
+   const titleWidth = 48;
+
+   const dataElement = (title: string, value: string, linkText?: string, linkTarget?: string, valueElementIn?: JSX.Element) => {
+
+      const textSize = 12;
+
+      let valueElement = valueElementIn;
+
+      if (!valueElement) {
+         valueElement = <Stack>
+            <Stack style={{ fontSize: textSize }}>
+               {value}
+            </Stack>
+         </Stack>
+      }
+
+      if (linkText) {
+         valueElement = <Link className="horde-link" to={linkText} target={linkTarget}>
+            {valueElement}
+         </Link>
+      }
+
+
+      let component = <Stack style={{ cursor: (linkText) ? "pointer" : undefined }} onClick={() => {
+         /*
+         if (agent) {
+            setViewAgent(agent);
+         }
+         */
+      }}>
+         <Stack horizontal>
+            <Stack style={{ width: titleWidth }}>
+               {!!title && <Text style={{ fontFamily: "Horde Open Sans SemiBold", fontSize: textSize }} >{title}</Text>}
+            </Stack>
+            {valueElement}
+         </Stack>
+      </Stack>
+
+
+      return component;
+   }
+
+   const elements: JSX.Element[] = [];
+   const step = tooltip.step;
+   if (step) {
+
+      const jobParams: JobParameters = {
+         id: step.jobId,
+         change: tooltip.change,
+         streamId: bisection.streamId
+      }
+
+
+      elements.push(dataElement("Step:", bisection.nodeName, `/job/${step.jobId}?step=${step.stepId}`));
+      elements.push(dataElement("", "", undefined, undefined, <ChangeButton job={jobParams} hideAborted={true} />));
+      if (step.logId) {
+         elements.push(dataElement("", "View Log", `/log/${step.logId}`, "_blank"));
+      }
+
+   } else {
+
+      if (bisection.nextJobChange === tooltip.change) {
+
+         const jobParams: JobParameters = {
+            id: bisection.nextJobId,
+            change: bisection.nextJobChange,
+            streamId: bisection.streamId
+         }
+
+         elements.push(dataElement("Setup:", "View Job", `/job/${bisection.nextJobId}`, "_blank"));
+         elements.push(dataElement("", "", undefined, undefined, <ChangeButton job={jobParams} hideAborted={true} />));
+      } else {
+         elements.push(dataElement("Task:", bisection.state));
+      }
+
+   }
+
+   const target = `bisection_callout_target_${bisection.id}_${tooltip.change}`;
+
+
+   return <div id={target} style={{
+      position: "absolute",
+      display: "block",
+      top: `${tooltip.y}px`,
+      left: `${tooltip.x + offsetX}px`,
+      backgroundColor: modeColors.background,
+      width: "max-content",
+      borderColor: dashboard.darktheme ? "#413F3D" : "#2D3F5F",
+      pointerEvents: tooltip.frozen ? undefined : "none",
+   }}>
+      <Callout target={`#${target}`} isBeakVisible={false} gapSpace={24} directionalHint={DirectionalHint.bottomCenter}>
+         <Stack horizontal>
+            <Stack tokens={{ childrenGap: 6, padding: 12 }}>
+               {elements}
+            </Stack>
+            {!!tooltip.frozen && <Stack style={{ paddingLeft: 32, paddingRight: 12, paddingTop: 12, cursor: "pointer" }} onClick={() => { tooltip.freeze(false); tooltip.update() }}>
+               <FontIcon iconName="Cancel" />
+            </Stack>}
+            {!tooltip.frozen && <Stack style={{ paddingLeft: 32, paddingRight: 12 + 16, paddingTop: 12 }} />}
+         </Stack>
+      </Callout>
+   </div>
+
+})
+
+type BisectionResult = {
+   change: number
+}
+
+
+let id_counter = 0;
+class BisectionRenderer {
+
+   constructor(bisection: GetBisectTaskResponse) {
+
+      this.bisection = bisection;
+      this.margin = { top: 0, right: 32, bottom: 0, left: 32 };
+
+   }
+
+   render(container: HTMLDivElement) {
+
+      const width = 760;
+      const height = 64;
+      const margin = this.margin;
+      const scolors = dashboard.getStatusColors();
+      const bisection = this.bisection;
+
+      const maxCL = bisection.initialChange;
+      let minCL = Math.min(bisection.currentChange, bisection.nextJobChange ?? Number.MAX_SAFE_INTEGER);
+      bisection.steps?.forEach(s => minCL = Math.min(minCL, s.change));
+
+      // bisection result
+      let result: BisectionResult | undefined;
+      if (bisection.state === BisectTaskState.Succeeded) {
+         result = {
+            change: bisection.currentChange
+         }
+      }
+
+      const svgId = `suite_svg_${id_counter++}`;
+
+      const xScale = d3.scaleLinear()
+         .domain([minCL - 1, maxCL + 1])
+         .range([margin.left, width - margin.right])
+
+      const changeSet = new Set<number>();
+      const changeColors = new Map<number, string>();
+      changeColors.set(maxCL, bisection.outcome === JobStepOutcome.Failure ? scolors.get(StatusColor.Failure)! : scolors.get(StatusColor.Warnings)!);
+
+      bisection.steps?.forEach(s => {
+
+         changeColors.set(s.change, scolors.get(StatusColor.Running)!);
+         if (s.outcome === JobStepOutcome.Failure) {
+            changeColors.set(s.change, scolors.get(StatusColor.Failure)!);
+         } else if (s.outcome === JobStepOutcome.Warnings) {
+            changeColors.set(s.change, scolors.get(StatusColor.Warnings)!);
+         } else if (s.outcome === JobStepOutcome.Success) {
+            changeColors.set(s.change, scolors.get(StatusColor.Success)!);
+         }
+      });
+
+      if (bisection.state === BisectTaskState.Running) {
+         if (bisection.nextJobChange) {
+            changeSet.add(bisection.nextJobChange);
+            changeColors.set(bisection.nextJobChange, scolors.get(StatusColor.Running)!);
+         }
+      }
+
+      changeSet.add(minCL);
+      changeSet.add(maxCL);
+      bisection.steps?.forEach(s => {
+         changeSet.add(s.change);
+      });
+
+      const X = d3.map(Array.from(changeSet).sort((a, b) => a - b), (t) => t);
+      const I = d3.range(X.length);
+
+      let svg = this.svg;
+
+      if (!svg) {
+         svg = d3.select(container)
+            .append("svg") as any as SelectionType;
+
+         svg.attr("id", svgId)
+            .attr("width", width + 48)
+            .attr("height", height)
+            .attr("viewBox", [0, 0, width, height] as any)
+
+         this.svg = svg;
+      } else {
+         // remove tooltip
+         d3.select(container).selectAll('div').remove();
+         svg.selectAll("*").remove();
+      }
+
+      const g = svg.append("g");
+
+      // Bisection Line
+
+      const yoffset = 36;
+
+      g.append("line")
+         .attr("stroke", () => {
+            return dashboard.darktheme ? "#6D6C6B" : "#4D4C4B";
+         })
+         .attr("stroke-width", 1)
+         .attr("stroke-linecap", 4)
+         .attr("stroke-opacity", dashboard.darktheme ? 0.35 : 0.25)
+         .attr("x1", () => this.margin.left)
+         .attr("x2", () => width - this.margin.right)
+         .attr("y1", () => yoffset)
+         .attr("y2", () => yoffset)
+
+      const radius = 3.5
+      g.selectAll("circle")
+         .data(I)
+         .join("circle")
+         .attr("id", i => `id_${bisection.id}_${X[i]}`)
+         .attr("cx", i => {
+            return xScale(X[i]);
+         })
+         .attr("cy", i => yoffset)
+         .attr("fill", i => changeColors.get(X[i])!)
+         .attr("r", i => X[i] === result?.change ? radius * 1.5 : radius);
+
+      if (result) {
+         svg.append("line")
+            .attr("class", "cline")
+            .attr("x1", () => xScale(result!.change))
+            .attr("x2", () => xScale(result!.change))
+            .attr("y1", () => 4)
+            .attr("y2", () => 64)
+            .attr("stroke-width", () => 2)
+            .style("stroke-dasharray", "3, 3")
+            //.attr("marker-end", "url(#cmarker)")
+            .attr("stroke", bisection.outcome === JobStepOutcome.Failure ? scolors.get(StatusColor.Failure)! : scolors.get(StatusColor.Warnings)!)
+      }
+
+      const closestCL = (x: number, y: number): number | undefined => {
+
+         if (!changeSet.size) {
+            return undefined;
+         }
+
+         const changes = Array.from(changeSet);
+
+         return Array.from(changeSet).reduce<number>((bestCL: number, change: number) => {
+
+            const bs = Math.abs(xScale(bestCL)! - x);
+            const cs = Math.abs(xScale(change)! - x);
+
+            if (cs > bs) {
+               return bestCL;
+            }
+
+            return change;
+
+         }, changes[0]);
+      }
+
+      const handleMouseMove = (event: any) => {
+
+         if (this.tooltip.frozen) {
+            return;
+         }
+
+         let mouseX = d3.pointer(event)[0];
+         let mouseY = d3.pointer(event)[1];
+         const change = closestCL(mouseX, mouseY);
+         this.tooltip.update(this.bisection, change, Math.floor(d3.pointer(event, container)[0]), Math.floor(d3.pointer(event, container)[1]));
+      }
+
+      const handleMouseLeave = (event: any) => {
+         if (!this.tooltip.frozen) {
+            this.tooltip.update();
+         }
+      }
+
+      const handleMouseClick = (event: any) => {
+
+         if (this.tooltip.frozen) {
+            this.tooltip.freeze(false);
+            return;
+         }
+
+         if (this.tooltip.change) {
+            this.tooltip.freeze(true);
+         }
+
+      }
+
+      svg.on("click", (event) => handleMouseClick(event))
+      svg.on("mousemove", (event) => handleMouseMove(event));
+      svg.on("mouseleave", (event) => { handleMouseLeave(event); })
+
+   }
+
+   bisection: GetBisectTaskResponse;
+   margin: { top: number, right: number, bottom: number, left: number }
+   svg?: SelectionType;
+   tooltip = new Tooltip();
+}
+
+const BisectionGraph: React.FC<{ bisection: GetBisectTaskResponse }> = ({ bisection }) => {
+
+   const graph_container_id = `timeline_graph_container`;
+
+   const [container, setContainer] = useState<HTMLDivElement | null>(null);
+   const graph = useConst(new BisectionRenderer(bisection));
+
+   if (container) {
+      try {
+         graph.render(container);
+
+      } catch (err) {
+         console.error(err);
+      }
+   }
+
+   return <Stack className={hordeClasses.horde}>
+      <Stack style={{ paddingLeft: 8 }}>
+         <div style={{ position: "relative" }}>
+            <GraphTooltip renderer={graph} />
+         </div>
+         <div id={graph_container_id} className="horde-no-darktheme" style={{ shapeRendering: "crispEdges", userSelect: "none", position: "relative" }} ref={(ref: HTMLDivElement) => setContainer(ref)} onMouseEnter={() => { }} onMouseLeave={() => { }} />
+      </Stack>
+   </Stack>;
+}
+
+export const BisectionList: React.FC<{ bisections?: GetBisectTaskResponse[] }> = ({ bisections }) => {
+
+   const [showCancelModal, setShowCancelModal] = useState("");
+
+   if (!bisections?.length) {
+      return null;
+   }
+
+   type BisectionItem = {
+      bisection: GetBisectTaskResponse;
+   }
+
+   const items: BisectionItem[] = [];
+
+   bisections.forEach((b) => {
+
+      items.push({
+         bisection: b
+      })
+
+   })
+
+   const renderRow: IDetailsListProps['onRenderRow'] = (props) => {
+
+      if (!props) {
+         return null;
+      }
+
+      const item = items[props.itemIndex];
+      const bisection = item?.bisection;
+
+      if (!bisection) {
+         return null;
+      }
+
+      const streamId = bisection.streamId;
+      const templateId = bisection.templateId;
+
+      const stream = projectStore.streamById(streamId);
+      const template = stream?.templates?.find(t => t.id === templateId);
+      const nodeName = bisection.nodeName;
+
+      const colors = dashboard.getStatusColors();
+      let statusColor = "#FF0000";
+      if (bisection.state === BisectTaskState.Running) {
+         statusColor = colors.get(StatusColor.Running)!;
+      } else if (bisection.state === BisectTaskState.Succeeded) {
+         statusColor = colors.get(StatusColor.Success)!;
+      } else if (bisection.state === BisectTaskState.Cancelled) {
+         statusColor = colors.get(StatusColor.Skipped)!;
+      } else {
+         statusColor = colors.get(StatusColor.Failure)!;
+      }
+
+      const timestamp = bisection.id.substring(0, 8)
+      const time = getShortNiceTime(new Date(parseInt(timestamp, 16) * 1000), true);
+      const initialChange: JobParameters = {
+         id: bisection.initialJobId,
+         change: bisection.initialChange,
+         streamId: bisection.streamId
+      }
+
+      const running = bisection.state === BisectTaskState.Running;
+      const succeeded = bisection.state === BisectTaskState.Succeeded;
+      const statusWidth = 64;
+
+      return <Stack /*style={{ cursor: "pointer" }} onClick={() => { setShowBisectionModal(bisection.id) }}*/>
+         <Stack style={{ backgroundColor: props.itemIndex % 2 ? undefined : "#E9E8E7", height: 96, padding: "12px 18px" }}>
+            <Stack verticalAlign='center' verticalFill>
+               <Stack horizontal style={{ paddingTop: 8, paddingBottom: 8 }}>
+                  <Stack style={{ width: 280 }} horizontal verticalAlign='center' verticalFill tokens={{ childrenGap: 12 }}>
+                     <Stack className="horde-no-darktheme">
+                        <FontIcon style={{ color: statusColor }} iconName={"Square"} />
+                     </Stack>
+                     <Stack>
+                        <Stack tokens={{ childrenGap: 8 }}>
+                           <Stack>
+                              <Text style={{ fontWeight: 600 }}>{nodeName}</Text>
+                           </Stack>
+                           <Stack>
+                              <Text variant='small'>{stream?.fullname} - {template?.name}</Text>
+                           </Stack>
+                           <Stack>
+                              {running && <Stack horizontal tokens={{ childrenGap: 8 }}>
+                                 <Text style={{ fontWeight: 600, width: statusWidth }}>Running:</Text>
+                                 <DefaultButton style={{ height: "20px", fontSize: "10px", width: "80px" }} text="Cancel" onClick={() => setShowCancelModal(bisection.id)} />
+                              </Stack>}
+                              {succeeded && <Stack horizontal tokens={{ childrenGap: 8 }}>
+                                 <Stack>
+                                    <Text style={{ fontWeight: 600, width: statusWidth }}>Suspect:</Text>
+                                 </Stack>
+                                 <Stack>
+                                    <ChangeButton job={{ streamId: bisection.streamId, change: bisection.currentChange }} />
+                                 </Stack>
+                              </Stack>}
+                              {!succeeded && !running && <Text style={{ fontWeight: 600, width: statusWidth }}>{bisection.state}</Text>}
+                           </Stack>
+                        </Stack>
+
+                     </Stack>
+                  </Stack>
+                  <Stack>
+                     <BisectionGraph bisection={bisection} />
+                  </Stack>
+                  <Stack grow />
+                  <Stack horizontal tokens={{ childrenGap: 12 }}>
+                     <Stack verticalAlign='center' verticalFill>
+                        <Text variant='small'>{time}</Text>
+                     </Stack>
+                     {dashboard.bisectPinned(bisection.id) && <Stack verticalAlign="center" verticalFill={true} horizontalAlign={"end"} onClick={(ev) => {
+                        ev.preventDefault();
+                        ev.stopPropagation();
+                        dashboard.unpinBisect(bisection.id);
+                     }}>
+                        <IconButton iconProps={{ iconName: 'Pin' }} />
+                     </Stack>}
+                     {!!dashboard.bisectPinned(bisection.id) && <Stack>
+                        <DefaultButton text="Pin" onClick={() => {
+                           dashboard.pinBisect(bisection.id);
+                        }} />
+                     </Stack>}
+
+                  </Stack>
+
+               </Stack>
+            </Stack>
+         </Stack>
+      </Stack >
+   }
+
+   return <Stack>
+      {!!showCancelModal && <CancelBisectionModal bisectTaskId={showCancelModal} onClose={(canceled) => { if (canceled) { /*window.location.reload();*/ setShowCancelModal("") } }} />}
+      <DetailsList
+         styles={{ root: { overflow: "hidden" } }}
+         className={customStyles.details}
+         compact={true}
+         isHeaderVisible={false}
+         indentWidth={0}
+         items={items}
+         setKey="set"
+         selectionMode={SelectionMode.none}
+         layoutMode={DetailsListLayoutMode.justified}
+         onRenderRow={renderRow}
+      />
+   </Stack>;
+};
+
+
+const customStyles = mergeStyleSets({
+   details: {
+      selectors: {
+         '.ms-GroupHeader,.ms-GroupHeader:hover': {
+            background: "#DFDEDD",
+         },
+         '.ms-GroupHeader-title': {
+            cursor: "default"
+         },
+         '.ms-GroupHeader-expand,.ms-GroupHeader-expand:hover': {
+            background: "#DFDEDD"
+         },
+         '.ms-DetailsRow': {
+            animation: "none",
+            background: "unset"
+         },
+         '.ms-DetailsRow:hover': {
+            background: "#F3F2F1"
+         }
+      },
+   }
+});
+
+
+export const CancelBisectionModal: React.FC<{ bisectTaskId: string, onClose: (canceled: boolean) => void }> = ({ bisectTaskId, onClose }) => {
+
+   const [state, setState] = useState<{ submitting?: boolean, error?: string }>({});
+
+   const onCancelTask = async () => {
+
+      setState({ ...state, submitting: true });
+
+      try {
+
+         await backend.updateBisectTask(bisectTaskId, {
+            cancel: true
+         });
+
+         setState({ ...state, submitting: false, error: "" });
+         onClose(true);
+
+      } catch (reason) {
+
+         setState({ ...state, submitting: false, error: reason as string });
+      }
+
+   };
+
+   const height = state.error ? 200 : 180;
+
+   return <Modal className={hordeClasses.modal} isOpen={true} topOffsetFixed={true} styles={{ main: { padding: 8, width: 540, height: height, minHeight: height, hasBeenOpened: false, top: "24px", position: "absolute" } }} onDismiss={() => { onClose(false) }}>
+      <Stack horizontal styles={{ root: { padding: 8 } }}>
+         <Stack.Item grow={2}>
+            <Text variant="mediumPlus">Cancel Bisection</Text>
+         </Stack.Item>
+         <Stack.Item grow={0}>
+            <IconButton
+               iconProps={{ iconName: 'Cancel' }}
+               ariaLabel="Close popup modal"
+               onClick={() => { onClose(false); }}
+            />
+         </Stack.Item>
+      </Stack>
+
+      <Stack styles={{ root: { paddingLeft: 8, width: 540, paddingTop: 12 } }}>
+         <Stack tokens={{ childrenGap: 12 }}>
+            {!!state.error && <MessageBar
+               messageBarType={MessageBarType.error}
+               isMultiline={false}> {state.error} </MessageBar>}
+
+         </Stack>
+      </Stack>
+
+      <Stack tokens={{ childrenGap: 16 }} styles={{ root: { paddingTop: 32, paddingLeft: 8, paddingBottom: 8 } }}>
+         <Stack horizontal>
+            <Stack>
+            </Stack>
+            <Stack grow />
+            <Stack>
+               <Stack horizontal tokens={{ childrenGap: 12 }} style={{ paddingRight: 24 }}>
+                  <PrimaryButton text="Cancel Task" disabled={state.submitting ?? false} onClick={() => { onCancelTask() }} />
+                  <DefaultButton text="Go Back" disabled={state.submitting} onClick={() => { onClose(false); }} />
+               </Stack>
+            </Stack>
+         </Stack>
+      </Stack>
+   </Modal>;
+
+};
\ No newline at end of file
diff --git a/Engine/Source/Programs/Horde/HordeDashboard/src/components/bisection/BisectionModal.tsx b/Engine/Source/Programs/Horde/HordeDashboard/src/components/bisection/BisectionModal.tsx
new file mode 100644
--- /dev/null
+++ b/Engine/Source/Programs/Horde/HordeDashboard/src/components/bisection/BisectionModal.tsx
@@ -0,0 +1,40 @@
+import { IconButton, Modal, Stack, Text } from "@fluentui/react";
+import { observer } from "mobx-react-lite";
+import React from "react";
+import { hordeClasses, modeColors } from "../../styles/Styles";
+
+export const BisectionModal: React.FC<{ taskId: string, onCloseExternal?: () => void }> = observer(({ onCloseExternal }) => {
+
+   return <Modal isOpen={true} isBlocking={true} topOffsetFixed={true} styles={{ main: { padding: 8, width: 1420, backgroundColor: modeColors.background, hasBeenOpened: false, top: "24px", position: "absolute", height: "95vh" } }} className={hordeClasses.modal} onDismiss={() => { if (onCloseExternal) { onCloseExternal() } else { /*onClose()*/ } }}>
+      <Stack style={{ height: "93vh" }}>
+         <Stack style={{ height: "100%" }}>
+            <Stack style={{ flexBasis: "70px", flexShrink: 0 }}>
+               <Stack horizontal styles={{ root: { padding: 8 } }} style={{ padding: 20, paddingBottom: 8 }}>
+                  <Stack horizontal style={{ width: 1024 }} tokens={{ childrenGap: 24 }} verticalAlign="center" verticalFill={true}>
+                     <Stack >
+                        <Text styles={{ root: { fontWeight: "unset", maxWidth: 720, wordBreak: "break-word", fontFamily: "Horde Open Sans SemiBold", fontSize: "14px", color: "#087BC4", } }}>{"This is a bisection modal"}</Text>
+                     </Stack>
+                     <Stack grow />
+                  </Stack>
+                  <Stack>
+                     <Text>Bisection Command Bar</Text>
+                  </Stack>
+                  <Stack grow />
+                  <Stack horizontalAlign="end">
+                     <IconButton
+                        iconProps={{ iconName: 'Cancel' }}
+                        ariaLabel="Close popup modal"
+                        onClick={() => { if (onCloseExternal) { onCloseExternal() } else { /*onClose()*/ } }}
+                     />
+                  </Stack>
+
+               </Stack>
+            </Stack>
+            <Stack style={{ flexGrow: 1, flexShrink: 0 }}>
+               <Text> Add More Here</Text>
+            </Stack>
+         </Stack>
+      </Stack>
+   </Modal>
+
+});
\ No newline at end of file
diff --git a/Engine/Source/Programs/Horde/HordeDashboard/src/components/jobDetailsV2/JobDetailBisection.tsx b/Engine/Source/Programs/Horde/HordeDashboard/src/components/jobDetailsV2/JobDetailBisection.tsx
--- a/Engine/Source/Programs/Horde/HordeDashboard/src/components/jobDetailsV2/JobDetailBisection.tsx
+++ b/Engine/Source/Programs/Horde/HordeDashboard/src/components/jobDetailsV2/JobDetailBisection.tsx
@@ -1,348 +1,178 @@
 // Copyright Epic Games, Inc. All Rights Reserved.
 
-import { DefaultButton, DetailsList, DetailsListLayoutMode, IColumn, IGroup, IGroupHeaderProps, IconButton, MessageBar, MessageBarType, Modal, PrimaryButton, SelectionMode, Stack, Text } from '@fluentui/react';
-import { mergeStyleSets } from '@fluentui/react/lib/Styling';
+import { Stack, Text, ThemeSettingName } from '@fluentui/react';
 import { observer } from 'mobx-react-lite';
-import moment from 'moment';
-import React, { useEffect, useState } from 'react';
-import { Link } from 'react-router-dom';
+import React, { useEffect } from 'react';
 import backend from '../../backend';
-import { BisectTaskState, GetBisectTaskResponse, GetJobStepRefResponse } from '../../backend/Api';
-import dashboard, { StatusColor } from '../../backend/Dashboard';
-import { Markdown } from '../../base/components/Markdown';
+import { GetBisectTaskResponse } from '../../backend/Api';
 import { ISideRailLink } from '../../base/components/SideRail';
-import { displayTimeZone, getElapsedString } from '../../base/utilities/timeUtils';
 import { hordeClasses } from '../../styles/Styles';
-import { ChangeButton } from '../ChangeButton';
-import { StepRefStatusIcon } from '../StatusIcon';
+import { BisectionList } from '../bisection/BisectionList';
 import { JobDataView, JobDetailsV2 } from './JobDetailsViewCommon';
-
+import { PollBase } from '../../backend/PollBase';
+import dashboard from '../../backend/Dashboard';
 
 const sideRail: ISideRailLink = { text: "Bisection", url: "rail_detail_bisection" };
 
+class BisectionHandler extends PollBase {
 
-class StepBisectionView extends JobDataView {
-
-   filterUpdated() {
-      // this.updateReady();
+   constructor(view: StepBisectionView, pollTime = 10000) {
+      super(pollTime);
+      this.view = view;
    }
 
-   async set(stepId?: string) {
-
-      const details = this.details;
-      if (!details) {
-         return;
-      }
-
-      if ((stepId && this.stepId === stepId) || !details.jobId || !details.jobData?.id) {
-         return;
-      }
-
-      this.stepId = stepId;
-
-      if (!this.bisections) {
+   clear() {
+      super.stop();
+   }
 
-         try {
+   async poll(): Promise<void> {
 
-            let jobId = details.jobData?.startedByBisectTaskId ? details.jobData?.startedByBisectTaskId : details.jobId;
+      let bisections: GetBisectTaskResponse[] = [];
 
-            if (jobId !== details.jobId) {
-               const task = await backend.getBisectTask(details.jobData.startedByBisectTaskId!);
-               jobId = task.initialJobId;
+      const view = this.view;
 
-            }
+      try {
 
-            this.bisections = await backend.getJobBisectTasks(jobId);
-            if (stepId) {
-               this.bisections = this.bisections.filter(r => r.nodeName === details.getStepName(stepId));
-            }
-            
+         const jobId = await view.getBisectionJobId();
 
-         } catch (error) {
-            console.error(error);
+         if (!jobId) {
+            return;
          }
-         finally {
-            this.initialize(this.bisections?.length ? [sideRail] : undefined);
+
+         bisections = await backend.getBisections({ jobId: jobId });
+         if (view.stepName) {
+            bisections = bisections.filter(r => r.nodeName === view.stepName);
          }
-      }
-   }
 
-   clear() {
-      this.stepId = undefined;
-      this.bisections = undefined;
-      super.clear();
-   }
 
-   detailsUpdated() {
+      } catch (err) {
+         console.error(err);
+      } finally {
 
-      if (!this.details?.jobData) {
-         return;
+         view.handlerUpdated(bisections);         
       }
-
-      this.updateReady();
-
    }
 
-   stepId?: string;
-
-   bisections?: GetBisectTaskResponse[];
-
-   order = 9;
+   view: StepBisectionView;
 
 }
 
-JobDetailsV2.registerDataView("StepBisectionView", (details: JobDetailsV2) => new StepBisectionView(details));
+class StepBisectionView extends JobDataView {
 
-const customStyles = mergeStyleSets({
-   details: {
-      selectors: {
-         '.ms-GroupHeader,.ms-GroupHeader:hover': {
-            background: "#DFDEDD",
-         },
-         '.ms-GroupHeader-title': {
-            cursor: "default"
-         },
-         '.ms-GroupHeader-expand,.ms-GroupHeader-expand:hover': {
-            background: "#DFDEDD"
-         },
-         '.ms-DetailsRow': {
-            animation: "none",
-            background: "unset"
-         },
-         '.ms-DetailsRow:hover': {
-            background: "#F3F2F1"
-         }
-      },
+   filterUpdated() {
+      // this.updateReady();
    }
 
-});
-
-
-export const CancelBisectionModal: React.FC<{ jobDetails: JobDetailsV2, dataView: StepBisectionView, bisectTaskId: string, onClose: (canceled:boolean) => void }> = ({ jobDetails, dataView, bisectTaskId, onClose }) => {
-
-   const [state, setState] = useState<{ submitting?: boolean, error?: string }>({});
+   bisectionUpdated() {
+      this.handler?.stop();
+      this.handler?.start();
+   }
 
-   const onCancelTask = async () => {
+   handlerUpdated(bisections: GetBisectTaskResponse[]) {
+      
+      const prev = this.bisections;
+      this.bisections = bisections;
+
+      if (!this.initialized) {
+         this.initialize(bisections?.length ? [sideRail] : undefined);
+      } else {
+         if (!prev?.length && bisections.length) {
+            // didn't have bisections and now we do
+            this.addRail(sideRail);
+         } 
+      }
 
-      setState({ ...state, submitting: true });
+      this.updateReady();
 
-      try {
+   }
 
-         await backend.updateBisectTask(bisectTaskId, {
-            cancel: true
-         });
-         
-         setState({ ...state, submitting: false, error: "" });
-         onClose(true);
-         
-      } catch (reason) {
+   // 
+   async getBisectionJobId() {
 
-         setState({ ...state, submitting: false, error: reason as string });
+      if (this.bisectionJobId) {
+         return this.bisectionJobId;
       }
 
-   };
-
-   const height = state.error ? 200 : 180;
-
-   return <Modal className={hordeClasses.modal} isOpen={true} topOffsetFixed={true} styles={{ main: { padding: 8, width: 540, height: height, minHeight: height, hasBeenOpened: false, top: "24px", position: "absolute" } }} onDismiss={() => { onClose(false) }}>
-      <Stack horizontal styles={{ root: { padding: 8 } }}>
-         <Stack.Item grow={2}>
-            <Text variant="mediumPlus">Cancel Bisection</Text>
-         </Stack.Item>
-         <Stack.Item grow={0}>
-            <IconButton
-               iconProps={{ iconName: 'Cancel' }}
-               ariaLabel="Close popup modal"
-               onClick={() => { onClose(false); }}
-            />
-         </Stack.Item>
-      </Stack>
-
-      <Stack styles={{ root: { paddingLeft: 8, width: 540, paddingTop: 12 } }}>
-         <Stack tokens={{ childrenGap: 12 }}>
-            {!!state.error && <MessageBar
-               messageBarType={MessageBarType.error}
-               isMultiline={false}> {state.error} </MessageBar>}
-
-         </Stack>
-      </Stack>
+      const details = this.details;
 
-      <Stack tokens={{ childrenGap: 16 }} styles={{ root: { paddingTop: 32, paddingLeft: 8, paddingBottom: 8 } }}>
-         <Stack horizontal>
-            <Stack>
-            </Stack>
-            <Stack grow />
-            <Stack>
-               <Stack horizontal tokens={{ childrenGap: 12 }} style={{ paddingRight: 24 }}>
-                  <PrimaryButton text="Cancel Task" disabled={state.submitting ?? false} onClick={() => { onCancelTask() }} />
-                  <DefaultButton text="Go Back" disabled={state.submitting} onClick={() => { onClose(false); }} />
-               </Stack>
-            </Stack>
-         </Stack>
-      </Stack>
-   </Modal>;
+      if (!details) {
+         return;
+      }
 
-};
+      let jobId = details.jobData?.startedByBisectTaskId ? details.jobData?.startedByBisectTaskId : details.jobId!;
 
+      if (jobId !== details.jobId) {
+         const task = await backend.getBisectTask(jobId);
+         jobId = task.initialJobId;
+      }
 
-const StepBisectionList: React.FC<{ jobDetails: JobDetailsV2, stepId?: string, dataView: StepBisectionView }> = ({ jobDetails, stepId, dataView }) => {
+      this.bisectionJobId = jobId;
 
-   const [showCancelModal, setShowCancelModal] = useState("");
+      return jobId;
 
-   type BisectionItem = {
-      bisection: GetBisectTaskResponse;
-      step?: GetJobStepRefResponse;
    }
 
-   const bisections = dataView.bisections;
-   if (!bisections?.length) {
-      return null;
+   get stepName(): string | undefined {
+      return this.details?.getStepName(this.stepId);
    }
 
-   const items: BisectionItem[] = [];
-
-   const groups: IGroup[] = [];
-
-   bisections.forEach((b, cindex) => {
-
-      const key = `group_${b.id}`;
-      groups.push({
-         key: key,
-         name: `Bisection ${b.id} - ${b.state}`,
-         data: b,
-         startIndex: items.length,
-         count: (b.steps?.length ?? 0),
-         level: 0,
-         isCollapsed: false //bisections.length > 0 && cindex !== 0
-      });
-
-      b.steps?.sort((a, b) => a.change - b.change).forEach(s => {
-         items.push({
-            bisection: b,
-            step: s
-         })
-      })
-   })
-
-   const columns = [
-      { key: 'column1', name: 'Name', minWidth: 620, maxWidth: 620, isResizable: false },
-      { key: 'column2', name: 'Change', minWidth: 80, maxWidth: 80, isResizable: false },
-      { key: 'column3', name: 'Started', minWidth: 180, maxWidth: 180, isResizable: false },
-      { key: 'column4', name: 'Duration', minWidth: 110, maxWidth: 110, isResizable: false },
-   ];
-
-   const renderItem = (item: BisectionItem, index?: number, column?: IColumn) => {
-
-      const step = item?.step;
-
-      if (!column || !step) {
-         return null;
-      }
-
-      if (column.name === "Name") {
-         return <Link to={`/job/${step.jobId}?step=${step.stepId}`}><Stack horizontal>{<StepRefStatusIcon stepRef={step} />}<Text>{item.bisection.nodeName}</Text></Stack></Link>;
-      }
-
-      if (column.name === "Change") {
-
-
-         return <ChangeButton job={jobDetails.jobData!} stepRef={step} hideAborted={true} />;
+   async set(stepId?: string) {
 
+      if (this.bisectionJobId) {
+         return;
       }
 
+      try {
 
-      if (column.name === "Started") {
-
-         if (step.startTime) {
+         const bisectionJobId = await this.getBisectionJobId()!;
 
-            const displayTime = moment(step.startTime).tz(displayTimeZone());
-            const format = dashboard.display24HourClock ? "HH:mm:ss z" : "LT z";
+         if (!bisectionJobId) {
+            return;
+         }
 
-            let displayTimeStr = displayTime.format('MMM Do') + ` at ${displayTime.format(format)}`;
+         this.stepId = stepId;
 
+         this.handler = new BisectionHandler(this);
+         this.handler.start();
 
-            return <Stack horizontal horizontalAlign={"start"}>{displayTimeStr}</Stack>;
+      } catch (error) {
+         console.error(error);
+      } finally {
 
-         } else {
-            return "???";
-         }
       }
+   }
 
-      if (column.name === "Duration") {
+   clear() {
+      this.stepId = undefined;
+      this.bisections = undefined;
+      this.handler?.clear();
+      super.clear();
+   }
 
-         const start = moment(step.startTime);
-         let end = moment(Date.now());
+   detailsUpdated() {
 
-         if (step.finishTime) {
-            end = moment(step.finishTime);
-         }
-         if (step.startTime) {
-            const time = getElapsedString(start, end);
-            return <Stack horizontal horizontalAlign={"end"} style={{ paddingRight: 8 }}><Text>{time}</Text></Stack>;
-         } else {
-            return "???";
-         }
+      if (!this.details?.jobData) {
+         return;
       }
 
-
-      return null;
+      this.updateReady();
 
    }
 
-   const onRenderGroupHeader = (props?: IGroupHeaderProps): JSX.Element | null => {
-
-      if (!props) {
-         return null;
-      }
-
-      const bisection = props.group?.data as GetBisectTaskResponse;
+   stepId?: string;
 
-      let markdown = `**Bisection ${bisection.id}** -  **State:** ${bisection.state} / **Initial Job:** [${bisection.initialJobId}](/job/${bisection.initialJobId}) - ${bisection.initialChange}`;
-      if (bisection.currentJobId) {
-         markdown += ` / **Current Job:** [${bisection.currentJobId}](/job/${bisection.currentJobId}) - ${bisection.currentChange}`
-      }
+   bisections?: GetBisectTaskResponse[];
 
-      markdown += ` / **Outcome:** ${bisection.outcome}`
+   handler?: BisectionHandler;
 
-      /*
-      if (bisection.owner) {
-         markdown += ` / **Owner:** ${bisection.owner.name}`
-      }
-      */
+   bisectionJobId?: string;
 
-      return <Stack verticalAlign="center" style={{ height: 48, padding: 18, backgroundColor: "rgb(233, 232, 231)" }}>
-         <Stack horizontal verticalAlign="center">
-            <Stack>
-               <Markdown>{markdown}</Markdown>
-            </Stack>
-            <Stack grow />
-            {bisection.state === BisectTaskState.Running && <Stack>
-               <DefaultButton style={{ color: "#FFFFFF", backgroundColor: dashboard.getStatusColors().get(StatusColor.Failure) }} text="Cancel" onClick={() => setShowCancelModal(bisection.id)} />
-            </Stack>
-            }
-         </Stack>
-      </Stack>
-   }
+   order = 9;
 
-   return <Stack>
-      {!!showCancelModal && <CancelBisectionModal jobDetails={jobDetails} dataView={dataView} bisectTaskId={showCancelModal} onClose={(canceled) => { if (canceled) { window.location.reload(); setShowCancelModal("") } }} />}
-      <DetailsList
-         className={customStyles.details}
-         compact={true}
-         isHeaderVisible={false}
-         indentWidth={0}
-         items={items}
-         groups={groups}
-         groupProps={{ onRenderHeader: onRenderGroupHeader, showEmptyGroups: true }}
-         columns={columns}
-         setKey="set"
-         selectionMode={SelectionMode.none}
-         layoutMode={DetailsListLayoutMode.justified}
-         onRenderItemColumn={renderItem}
-      />
-   </Stack>;
 }
 
+JobDetailsV2.registerDataView("StepBisectionView", (details: JobDetailsV2) => new StepBisectionView(details));
+
 export const BisectionPanel: React.FC<{ jobDetails: JobDetailsV2, stepId?: string }> = observer(({ jobDetails, stepId }) => {
 
    const dataView = jobDetails.getDataView<StepBisectionView>("StepBisectionView");
@@ -367,7 +197,7 @@ export const BisectionPanel: React.FC<{ jobDetails: JobDetailsV2, stepId?: strin
             <Text variant="mediumPlus" styles={{ root: { fontFamily: "Horde Open Sans SemiBold" } }}>Bisection</Text>
             <Stack styles={{ root: { paddingLeft: 4, paddingRight: 0, paddingTop: 8, paddingBottom: 4 } }}>
                <div style={{ overflowY: 'auto', overflowX: 'hidden', minHeight: "400px", maxHeight: "400px" }} data-is-scrollable={true}>
-                  <StepBisectionList jobDetails={jobDetails} dataView={dataView} />
+                  <BisectionList bisections={dataView.bisections} />
                </div>
             </Stack>
          </Stack>
diff --git a/Engine/Source/Programs/Horde/HordeDashboard/src/components/jobDetailsV2/JobDetailTimeline.tsx b/Engine/Source/Programs/Horde/HordeDashboard/src/components/jobDetailsV2/JobDetailTimeline.tsx
--- a/Engine/Source/Programs/Horde/HordeDashboard/src/components/jobDetailsV2/JobDetailTimeline.tsx
+++ b/Engine/Source/Programs/Horde/HordeDashboard/src/components/jobDetailsV2/JobDetailTimeline.tsx
@@ -225,7 +225,7 @@ class TimelineDataView extends JobDataView {
       });
 
       if (!batches.length) {
-         console.log("No batches to render in timeline");
+         //console.log("No batches to render in timeline");
          return;
       }
 
diff --git a/Engine/Source/Programs/Horde/HordeDashboard/src/components/jobDetailsV2/JobDetailsViewCommon.tsx b/Engine/Source/Programs/Horde/HordeDashboard/src/components/jobDetailsV2/JobDetailsViewCommon.tsx
--- a/Engine/Source/Programs/Horde/HordeDashboard/src/components/jobDetailsV2/JobDetailsViewCommon.tsx
+++ b/Engine/Source/Programs/Horde/HordeDashboard/src/components/jobDetailsV2/JobDetailsViewCommon.tsx
@@ -60,7 +60,21 @@ export abstract class JobDataView {
    get initialized(): boolean {
       return this._initialized;
    }
-   // may be called more than once with different parameters, though allways initializes
+
+   addRail(link: ISideRailLink) {
+
+      const existing = this._railLinks.find(r => r.text === link.text);
+      if (existing) {
+         return;
+      }
+
+      this._railLinks.push(link);
+      this.details?.externalUpdate();
+      this.details?.setRootUpdated();
+   }
+
+
+   // may be called more than once with different parameters, though always initializes
    initialize(railLinks?: ISideRailLink[]) {
 
       if (this._initialized) {
@@ -523,6 +537,14 @@ export class JobDetailsV2 extends PollBase {
       this.views.forEach(v => v.filterUpdated());
    }
 
+   bisectionUpdated() {
+      const v = this.views.find(v => v.name === "StepBisectionView") as any;
+      if (v) {
+         v.bisectionUpdated();
+      }
+   }
+
+
    getReportData(placement: ReportPlacement, stepId?: string): string | undefined {
 
       if (stepId) {
diff --git a/Engine/Source/Programs/Horde/HordeDashboard/src/components/jobDetailsV2/JobOperationsBar.tsx b/Engine/Source/Programs/Horde/HordeDashboard/src/components/jobDetailsV2/JobOperationsBar.tsx
--- a/Engine/Source/Programs/Horde/HordeDashboard/src/components/jobDetailsV2/JobOperationsBar.tsx
+++ b/Engine/Source/Programs/Horde/HordeDashboard/src/components/jobDetailsV2/JobOperationsBar.tsx
@@ -430,8 +430,9 @@ const StepOperations: React.FC<{ jobDetails: JobDetailsV2, stepId: string }> = o
       {!!shown.abortShown && <AbortJobModal stepId={stepId} jobDetails={jobDetails} show={true} onClose={() => { setShown({}); }} />}
       {!!shown.bisectShown && <BisectionCreateModal jobId={jobId} nodeName={node?.name ?? "Unknown Node"} onClose={(response) => {
          if (response) {
-            window.location.reload();
+            jobDetails.bisectionUpdated();            
          }
+
          setShown({});
       }} />}
       <Stack horizontal styles={{ root: { paddingLeft: 0 } }}>
diff --git a/Engine/Source/Programs/Horde/HordeDashboard/src/styles/Styles.tsx b/Engine/Source/Programs/Horde/HordeDashboard/src/styles/Styles.tsx
--- a/Engine/Source/Programs/Horde/HordeDashboard/src/styles/Styles.tsx
+++ b/Engine/Source/Programs/Horde/HordeDashboard/src/styles/Styles.tsx
@@ -134,6 +134,8 @@ registerIcons({
       "ArrowUp": IconContents.ARROW_UP,
       "Repeat": IconContents.REPEAT,
       "FlowReview": IconContents.FLOW_REVIEW,
+      "FastForward": IconContents.FAST_FORWARD,
+      "Play": IconContents.PLAY,
    }
 });
 
diff --git a/Engine/Source/Programs/Shared/EpicGames.Horde/Api/UserMessages.cs b/Engine/Source/Programs/Shared/EpicGames.Horde/Api/UserMessages.cs
--- a/Engine/Source/Programs/Shared/EpicGames.Horde/Api/UserMessages.cs
+++ b/Engine/Source/Programs/Shared/EpicGames.Horde/Api/UserMessages.cs
@@ -79,6 +79,11 @@ namespace EpicGames.Horde.Api
 		/// </summary>
 		public List<string>? PinnedJobIds { get; set; }
 
+		/// <summary>
+		/// List of pinned bisection task ids
+		/// </summary>
+		public List<string>? PinnedBisectTaskIds { get; set; }
+
 		/// <summary>
 		/// Constructor
 		/// </summary>
@@ -269,5 +274,16 @@ namespace EpicGames.Horde.Api
 		/// Jobs ids to remove from the pinned list
 		/// </summary>
 		public List<string>? RemovePinnedJobIds { get; set; }
+
+		/// <summary>
+		/// Bisection task ids to add to the pinned list
+		/// </summary>
+		public List<string>? AddPinnedBisectTaskIds { get; set; }
+
+		/// <summary>
+		/// Bisection task ids to remove from the pinned list
+		/// </summary>
+		public List<string>? RemovePinnedBisectTaskIds { get; set; }
+
 	}
 }
