Horde: Additional work on bisection pinning/unpinning
#skipci
#rnx

diff --git a/Engine/Source/Programs/Horde/HordeDashboard/src/backend/Dashboard.ts b/Engine/Source/Programs/Horde/HordeDashboard/src/backend/Dashboard.ts
--- a/Engine/Source/Programs/Horde/HordeDashboard/src/backend/Dashboard.ts
+++ b/Engine/Source/Programs/Horde/HordeDashboard/src/backend/Dashboard.ts
@@ -221,7 +221,7 @@ export class Dashboard {
     getLastJobTemplateSettings(streamId: string, templateIds: string[]): GetJobTemplateSettingsResponse | undefined {
 
         try {
-            
+
             const streamTemplates = this.data?.jobTemplateSettings?.filter(t => t.streamId === streamId && templateIds.indexOf(t.templateId) !== -1)
 
             if (!streamTemplates?.length) {
@@ -577,6 +577,10 @@ export class Dashboard {
         this.setUpdated();
     }
 
+    subscribe() {
+        if (this.updated) { }
+    }
+
     @observable
     updated: number = 0;
 
diff --git a/Engine/Source/Programs/Horde/HordeDashboard/src/components/UserHome.tsx b/Engine/Source/Programs/Horde/HordeDashboard/src/components/UserHome.tsx
--- a/Engine/Source/Programs/Horde/HordeDashboard/src/components/UserHome.tsx
+++ b/Engine/Source/Programs/Horde/HordeDashboard/src/components/UserHome.tsx
@@ -1080,6 +1080,8 @@ class BisectionHandler extends PollBase {
 
          if (dashboard.pinnedBisectTaskIds?.length) {
             this.bisections = await backend.getBisections({ id: dashboard.pinnedBisectTaskIds });
+         } else {
+            this.bisections = [];
          }
 
          this.setUpdated();
@@ -1108,7 +1110,8 @@ const BisectionPanel: React.FC = observer(() => {
    }, []);
 
    // subscribe
-   if (dashboard.updated) { };
+   dashboard.subscribe();   
+
    if (bisectionHandler.updated) { };
 
    if (bisectionHandler.bisections.length === 0) {
diff --git a/Engine/Source/Programs/Horde/HordeDashboard/src/components/bisection/BisectionList.tsx b/Engine/Source/Programs/Horde/HordeDashboard/src/components/bisection/BisectionList.tsx
--- a/Engine/Source/Programs/Horde/HordeDashboard/src/components/bisection/BisectionList.tsx
+++ b/Engine/Source/Programs/Horde/HordeDashboard/src/components/bisection/BisectionList.tsx
@@ -406,10 +406,12 @@ const BisectionGraph: React.FC<{ bisection: GetBisectTaskResponse }> = ({ bisect
    </Stack>;
 }
 
-export const BisectionList: React.FC<{ bisections?: GetBisectTaskResponse[] }> = ({ bisections }) => {
+export const BisectionList: React.FC<{ bisections?: GetBisectTaskResponse[] }> = observer(({ bisections }) => {
 
    const [showCancelModal, setShowCancelModal] = useState("");
 
+   dashboard.subscribe();
+
    if (!bisections?.length) {
       return null;
    }
@@ -520,12 +522,14 @@ export const BisectionList: React.FC<{ bisections?: GetBisectTaskResponse[] }> =
                         ev.stopPropagation();
                         dashboard.unpinBisect(bisection.id);
                      }}>
-                        <IconButton iconProps={{ iconName: 'Pin' }} />
+                        <IconButton iconProps={{ iconName: 'Unpin' }} />
                      </Stack>}
-                     {!!dashboard.bisectPinned(bisection.id) && <Stack>
-                        <DefaultButton text="Pin" onClick={() => {
-                           dashboard.pinBisect(bisection.id);
-                        }} />
+                     {!dashboard.bisectPinned(bisection.id) && <Stack verticalAlign="center" verticalFill={true} horizontalAlign={"end"} onClick={(ev) => {
+                        ev.preventDefault();
+                        ev.stopPropagation();
+                        dashboard.pinBisect(bisection.id);
+                     }}>
+                        <IconButton iconProps={{ iconName: 'Pin' }} style={{color: dashboard.getStatusColors().get(StatusColor.Skipped!)}} />
                      </Stack>}
 
                   </Stack>
@@ -551,7 +555,7 @@ export const BisectionList: React.FC<{ bisections?: GetBisectTaskResponse[] }> =
          onRenderRow={renderRow}
       />
    </Stack>;
-};
+});
 
 
 const customStyles = mergeStyleSets({
diff --git a/Engine/Source/Programs/Horde/HordeDashboard/src/styles/Styles.tsx b/Engine/Source/Programs/Horde/HordeDashboard/src/styles/Styles.tsx
--- a/Engine/Source/Programs/Horde/HordeDashboard/src/styles/Styles.tsx
+++ b/Engine/Source/Programs/Horde/HordeDashboard/src/styles/Styles.tsx
@@ -107,6 +107,7 @@ registerIcons({
       "Issue": IconContents.ISSUE,
       "Dot": IconContents.DOT,
       "Pin": IconContents.PIN,
+      "Unpin": IconContents.UNPIN,
       "SearchTemplate": IconContents.SEARCH_TEMPLATE,
       "History": IconContents.HISTORY,
       "Import": IconContents.IMPORT,
