Horde: Prevent P4 cache updates during scheduled downtime.

diff --git a/Engine/Source/Programs/Horde/Horde.Server/Perforce/PerforceServiceCache.cs b/Engine/Source/Programs/Horde/Horde.Server/Perforce/PerforceServiceCache.cs
--- a/Engine/Source/Programs/Horde/Horde.Server/Perforce/PerforceServiceCache.cs
+++ b/Engine/Source/Programs/Horde/Horde.Server/Perforce/PerforceServiceCache.cs
@@ -248,12 +248,6 @@ namespace Horde.Server.Perforce
 		/// <returns></returns>
 		async ValueTask UpdateCommitsAsync(CancellationToken cancellationToken)
 		{
-			// Don't do any updates during downtime; we might just create a bunch of P4 errors.
-			if (_downtimeService.IsDowntimeActive)
-			{
-				return;
-			}
-
 			CacheState state = await _mongoService.GetSingletonAsync<CacheState>();
 
 			// Get the current list of streams and their views
@@ -269,7 +263,7 @@ namespace Horde.Server.Perforce
 			for (; ; )
 			{
 				// Update the background task for refreshing the list of clusters
-				if(clusterTask != null && clusterTask.IsCompleted)
+				if (clusterTask != null && clusterTask.IsCompleted)
 				{
 					try
 					{
@@ -281,72 +275,78 @@ namespace Horde.Server.Perforce
 					}
 					clusterTask = null;
 				}
-				if (clusterTask == null && clusterTimer.Elapsed > TimeSpan.FromSeconds(30.0))
-				{
-					clusterTask = Task.Run(() => CreateStreamInfoAsync(cancellationToken), cancellationToken);
-					clusterTimer.Restart();
-				}
 
-				// Remove any state for clusters that are no longer valid
-				bool updateState = false;
-				foreach (string clusterName in state.Clusters.Keys)
+				// Don't do any updates during downtime; we might just create a bunch of P4 errors.
+				if (!_downtimeService.IsDowntimeActive)
 				{
-					if (!clusters.ContainsKey(clusterName))
+					// Check if it's time to start a new cluster update
+					if (clusterTask == null && clusterTimer.Elapsed > TimeSpan.FromSeconds(30.0))
 					{
-						state.Clusters.Remove(clusterName);
-						updateState = true;
+						clusterTask = Task.Run(() => CreateStreamInfoAsync(cancellationToken), cancellationToken);
+						clusterTimer.Restart();
 					}
-				}
 
-				// Make sure there's a ticker for every cluster
-				foreach (string clusterName in clusters.Keys)
-				{
-					if (!tickers.Any(x => x.ClusterName.Equals(clusterName, StringComparison.OrdinalIgnoreCase)))
+					// Remove any state for clusters that are no longer valid
+					bool updateState = false;
+					foreach (string clusterName in state.Clusters.Keys)
 					{
-						ClusterTicker ticker = new ClusterTicker(clusterName);
-						tickers.Add(ticker);
-					}
-				}
-
-				// Check if it's time to update any tickers
-				for (int idx = 0; idx < tickers.Count; idx++)
-				{
-					ClusterTicker ticker = tickers[idx];
-					if (ticker.Task != null && ticker.Task.IsCompleted)
-					{
-						ClusterState? clusterState = await ticker.Task;
-						if (clusterState != null)
+						if (!clusters.ContainsKey(clusterName))
 						{
-							state.Clusters[ticker.ClusterName] = clusterState;
+							state.Clusters.Remove(clusterName);
 							updateState = true;
 						}
-						ticker.Task = null;
 					}
-					if (ticker.Task == null)
+
+					// Make sure there's a ticker for every cluster
+					foreach (string clusterName in clusters.Keys)
 					{
-						List<StreamInfo>? streams;
-						if (!clusters.TryGetValue(ticker.ClusterName, out streams))
+						if (!tickers.Any(x => x.ClusterName.Equals(clusterName, StringComparison.OrdinalIgnoreCase)))
 						{
-							tickers.RemoveAt(idx--);
-							continue;
+							ClusterTicker ticker = new ClusterTicker(clusterName);
+							tickers.Add(ticker);
 						}
+					}
 
-						ClusterState? clusterState;
-						if (!state.Clusters.TryGetValue(ticker.ClusterName, out clusterState))
+					// Check if it's time to update any tickers
+					for (int idx = 0; idx < tickers.Count; idx++)
+					{
+						ClusterTicker ticker = tickers[idx];
+						if (ticker.Task != null && ticker.Task.IsCompleted)
 						{
-							clusterState = new ClusterState();
+							ClusterState? clusterState = await ticker.Task;
+							if (clusterState != null)
+							{
+								state.Clusters[ticker.ClusterName] = clusterState;
+								updateState = true;
+							}
+							ticker.Task = null;
 						}
+						if (ticker.Task == null)
+						{
+							List<StreamInfo>? streams;
+							if (!clusters.TryGetValue(ticker.ClusterName, out streams))
+							{
+								tickers.RemoveAt(idx--);
+								continue;
+							}
+
+							ClusterState? clusterState;
+							if (!state.Clusters.TryGetValue(ticker.ClusterName, out clusterState))
+							{
+								clusterState = new ClusterState();
+							}
 
-						ticker.Task = Task.Run(() => UpdateClusterGuardedAsync(ticker.ClusterName, streams, clusterState, cancellationToken));
+							ticker.Task = Task.Run(() => UpdateClusterGuardedAsync(ticker.ClusterName, streams, clusterState, cancellationToken));
+						}
 					}
-				}
 
-				// Apply any updates to the global state
-				if (updateState)
-				{
-					if (!await _mongoService.TryUpdateSingletonAsync(state))
+					// Apply any updates to the global state
+					if (updateState)
 					{
-						state = await _mongoService.GetSingletonAsync<CacheState>();
+						if (!await _mongoService.TryUpdateSingletonAsync(state))
+						{
+							state = await _mongoService.GetSingletonAsync<CacheState>();
+						}
 					}
 				}
 
