Horde: Skip schedule updates during maintainance windows.

diff --git a/Engine/Source/Programs/Horde/Horde.Server/Jobs/Schedules/ScheduleService.cs b/Engine/Source/Programs/Horde/Horde.Server/Jobs/Schedules/ScheduleService.cs
--- a/Engine/Source/Programs/Horde/Horde.Server/Jobs/Schedules/ScheduleService.cs
+++ b/Engine/Source/Programs/Horde/Horde.Server/Jobs/Schedules/ScheduleService.cs
@@ -57,6 +57,7 @@ namespace Horde.Server.Jobs.Schedules
 		readonly IGraphCollection _graphs;
 		readonly ICommitService _commitService;
 		readonly IJobCollection _jobCollection;
+		readonly IDowntimeService _downtimeService;
 		readonly JobService _jobService;
 		readonly IStreamCollection _streamCollection;
 		readonly ITemplateCollection _templateCollection;
@@ -73,12 +74,13 @@ namespace Horde.Server.Jobs.Schedules
 		/// <summary>
 		/// Constructor
 		/// </summary>
-		public ScheduleService(RedisService redis, IGraphCollection graphs, ICommitService commitService, IJobCollection jobCollection, JobService jobService, IStreamCollection streamCollection, ITemplateCollection templateCollection, MongoService mongoService, IClock clock, IOptionsMonitor<GlobalConfig> globalConfig, Tracer tracer, ILogger<ScheduleService> logger)
+		public ScheduleService(RedisService redis, IGraphCollection graphs, ICommitService commitService, IJobCollection jobCollection, JobService jobService, IDowntimeService downtimeService, IStreamCollection streamCollection, ITemplateCollection templateCollection, MongoService mongoService, IClock clock, IOptionsMonitor<GlobalConfig> globalConfig, Tracer tracer, ILogger<ScheduleService> logger)
 		{
 			_graphs = graphs;
 			_commitService = commitService;
 			_jobCollection = jobCollection;
 			_jobService = jobService;
+			_downtimeService = downtimeService;
 			_streamCollection = streamCollection;
 			_templateCollection = templateCollection;
 			_clock = clock;
@@ -119,6 +121,12 @@ namespace Horde.Server.Jobs.Schedules
 		{
 			DateTime utcNow = _clock.UtcNow;
 
+			// Don't start any new jobs during scheduled downtime
+			if (_downtimeService.IsDowntimeActive)
+			{
+				return;
+			}
+
 			// Update the current queue
 			await using (RedisLock sharedLock = new (_redis.GetDatabase(), s_tickLockKey))
 			{
