Horde: Set correct build root for BuildGraph in WorkspaceExecutor

Previously it used the stream path from workspace when it should be using it from current batch.
This mostly works, but in steps like Setup Build in Epic's setup, where a different stream is used this got mixed up (Dev-Main-Utility vs Main)

diff --git a/Engine/Source/Programs/Horde/Horde.Agent.Tests/Execution/WorkspaceExecutorTest.cs b/Engine/Source/Programs/Horde/Horde.Agent.Tests/Execution/WorkspaceExecutorTest.cs
--- a/Engine/Source/Programs/Horde/Horde.Agent.Tests/Execution/WorkspaceExecutorTest.cs
+++ b/Engine/Source/Programs/Horde/Horde.Agent.Tests/Execution/WorkspaceExecutorTest.cs
@@ -86,7 +86,7 @@ public sealed class WorkspaceExecutorTest : IDisposable
 	[TestMethod]
 	public async Task EnvVarsAsync()
 	{
-		BeginBatchResponse batch = new BeginBatchResponse { Change = 1 };
+		BeginBatchResponse batch = new BeginBatchResponse { Change = 1, StreamName = "//UE5/Main" };
 		JobExecutorOptions executorOptions = new JobExecutorOptions(_session, null!, JobId, "batch1", batch, default, "", null!, new JobOptions());
 		WorkspaceExecutor executor = new (executorOptions, _workspace, _autoSdkWorkspace, NullLogger.Instance);
 		await executor.InitializeAsync(_logger, CancellationToken.None);
@@ -96,8 +96,8 @@ public sealed class WorkspaceExecutorTest : IDisposable
 		IReadOnlyDictionary<string, string> envVars = executor.GetEnvVars();
 		Assert.AreEqual("1", envVars["IsBuildMachine"]);
 		Assert.AreEqual(settings.DirectoryPath.FullName, envVars["uebp_LOCAL_ROOT"]);
-		Assert.AreEqual(settings.StreamRoot, envVars["uebp_BuildRoot_P4"]);
-		Assert.AreEqual(settings.StreamRoot, envVars["uebp_BuildRoot_Escaped"]);
+		Assert.AreEqual(batch.StreamName, envVars["uebp_BuildRoot_P4"]);
+		Assert.AreEqual("++UE5+Main", envVars["uebp_BuildRoot_Escaped"]);
 		Assert.AreEqual("1", envVars["uebp_CL"]);
 		Assert.AreEqual("0", envVars["uebp_CodeCL"]);
 		Assert.AreEqual(autoSdkSettings.DirectoryPath.FullName, envVars["UE_SDKS_ROOT"]);
diff --git a/Engine/Source/Programs/Horde/Horde.Agent/Execution/WorkspaceExecutor.cs b/Engine/Source/Programs/Horde/Horde.Agent/Execution/WorkspaceExecutor.cs
--- a/Engine/Source/Programs/Horde/Horde.Agent/Execution/WorkspaceExecutor.cs
+++ b/Engine/Source/Programs/Horde/Horde.Agent/Execution/WorkspaceExecutor.cs
@@ -79,8 +79,8 @@ namespace Horde.Agent.Execution
 			// Set any non-materializer specific environment variables for jobs
 			_envVars["IsBuildMachine"] = "1";
 			_envVars["uebp_LOCAL_ROOT"] = workspaceSettings.DirectoryPath.FullName;
-			_envVars["uebp_BuildRoot_P4"] = workspaceSettings.StreamRoot;
-			_envVars["uebp_BuildRoot_Escaped"] = workspaceSettings.StreamRoot.Replace('/', '+');
+			_envVars["uebp_BuildRoot_P4"] = _batch.StreamName;
+			_envVars["uebp_BuildRoot_Escaped"] = _batch.StreamName.Replace('/', '+');
 			_envVars["uebp_CL"] = _batch.Change.ToString();
 			_envVars["uebp_CodeCL"] = _batch.CodeChange.ToString();
 
