Horde: Tweak dashed bisection line
#rnx
#skipci

diff --git a/Engine/Source/Programs/Horde/HordeDashboard/src/components/bisection/BisectionList.tsx b/Engine/Source/Programs/Horde/HordeDashboard/src/components/bisection/BisectionList.tsx
--- a/Engine/Source/Programs/Horde/HordeDashboard/src/components/bisection/BisectionList.tsx
+++ b/Engine/Source/Programs/Horde/HordeDashboard/src/components/bisection/BisectionList.tsx
@@ -307,10 +307,10 @@ class BisectionRenderer {
             .attr("class", "cline")
             .attr("x1", () => xScale(result!.change))
             .attr("x2", () => xScale(result!.change))
-            .attr("y1", () => 4)
-            .attr("y2", () => 64)
+            .attr("y1", () => 8)
+            .attr("y2", () => 68)
             .attr("stroke-width", () => 2)
-            .style("stroke-dasharray", "3, 3")
+            .style("stroke-dasharray", "2, 2")
             //.attr("marker-end", "url(#cmarker)")
             .attr("stroke", bisection.outcome === JobStepOutcome.Failure ? scolors.get(StatusColor.Failure)! : scolors.get(StatusColor.Warnings)!)
       }
