Horde: Speculative fix for flaky tests in JobQueueStrategyTest

No repro locally but has happened during CI runs.

diff --git a/Engine/Source/Programs/Horde/Horde.Server.Tests/Fleet/JobQueueStrategyTest.cs b/Engine/Source/Programs/Horde/Horde.Server.Tests/Fleet/JobQueueStrategyTest.cs
--- a/Engine/Source/Programs/Horde/Horde.Server.Tests/Fleet/JobQueueStrategyTest.cs
+++ b/Engine/Source/Programs/Horde/Horde.Server.Tests/Fleet/JobQueueStrategyTest.cs
@@ -11,7 +11,6 @@ using Horde.Server.Streams;
 using Horde.Server.Jobs;
 using Horde.Server.Jobs.Graphs;
 using Horde.Server.Server;
-using Horde.Server.Utilities;
 using Microsoft.VisualStudio.TestTools.UnitTesting;
 using Horde.Server.Agents;
 using EpicGames.Horde.Api;
@@ -101,7 +100,9 @@ namespace Horde.Server.Tests.Fleet
 			PoolSizeResult result = await strategy.CalculatePoolSizeAsync(pool, agents);
 			Assert.AreEqual(agents.Count + expectedAgentDelta, result.DesiredAgentCount);
 		}
-	
+
+		private static int s_uniqueId = 0;
+		
 		/// <summary>
 		/// Set up a fixture for job queue tests, ensuring a certain number of job batches are in running or waiting state
 		/// </summary>
@@ -109,7 +110,7 @@ namespace Horde.Server.Tests.Fleet
 		/// <param name="numBatchesReady">Num of job batches that should be in state waiting</param>
 		private async Task<(JobQueueStrategy, PoolSizeResult, IPool, List<IAgent> agents)> SetUpJobsAsync(int numBatchesRunning, int numBatchesReady, int numAgents = 8, bool isDowntimeActive = false)
 		{
-			IPool pool = await PoolService.CreatePoolAsync("bogusPool1", new AddPoolOptions { EnableAutoscaling = true, MinAgents = 0, NumReserveAgents = 0 });
+			IPool pool = await PoolService.CreatePoolAsync("bogusPool" + ++s_uniqueId, new AddPoolOptions { EnableAutoscaling = true, MinAgents = 0, NumReserveAgents = 0 });
 			List<IAgent> agents = new();
 			for (int i = 0; i < numAgents; i++)
 			{
@@ -118,7 +119,7 @@ namespace Horde.Server.Tests.Fleet
 			
 			PoolSizeResult poolSize = new (agents.Count, agents.Count, null);
 			
-			string agentTypeName1 = "bogusAgentType1";
+			string agentTypeName1 = "bogusAgentType" + ++s_uniqueId;
 			Dictionary<string, AgentConfig> agentTypes = new() { {agentTypeName1, new() { Pool = new PoolId(pool.Name) } }, };
 
 			StreamConfig streamConfig = new StreamConfig { Id = new StreamId("ue5"), Name = "//UE5/Main", AgentTypes = agentTypes };
@@ -132,7 +133,7 @@ namespace Horde.Server.Tests.Fleet
 
 			SetConfig(globalConfig);
 
-			string nodeForAgentType1 = "bogusNodeOnAgentType1";
+			string nodeForAgentType1 = "bogusNodeOnAgentType" + ++s_uniqueId;
 			IGraph graph = await GraphCollection.AppendAsync(null, new()
 			{
 				new NewGroup(agentTypeName1, new List<NewNode>
