Horde: Add more logging around null reference exceptions on server.

diff --git a/Engine/Source/Programs/Horde/Horde.Server/Storage/StorageBackendProvider.cs b/Engine/Source/Programs/Horde/Horde.Server/Storage/StorageBackendProvider.cs
--- a/Engine/Source/Programs/Horde/Horde.Server/Storage/StorageBackendProvider.cs
+++ b/Engine/Source/Programs/Horde/Horde.Server/Storage/StorageBackendProvider.cs
@@ -77,7 +77,14 @@ namespace Horde.Server.Storage
 			public IAsyncEnumerable<string> EnumerateAsync(CancellationToken cancellationToken = default) => _backend.EnumerateAsync(cancellationToken);
 
 			/// <inheritdoc/>
-			public Task<Stream> ReadAsync(string path, int offset, int? length, CancellationToken cancellationToken = default) => _backend.ReadAsync(path, offset, length, cancellationToken);
+			public Task<Stream> ReadAsync(string path, int offset, int? length, CancellationToken cancellationToken = default)
+			{
+				if (_backend == null)
+				{
+					throw new InvalidOperationException("Backend has already been disposed");
+				}
+				return _backend.ReadAsync(path, offset, length, cancellationToken);
+			}
 
 			/// <inheritdoc/>
 			public Task WriteAsync(string path, Stream stream, CancellationToken cancellationToken = default) => _backend.WriteAsync(path, stream, cancellationToken);
