Horde: Include a list of agents stuck in an upgrade loop in the agent report.

diff --git a/Engine/Source/Programs/Horde/Horde.Server/Agents/AgentCollection.cs b/Engine/Source/Programs/Horde/Horde.Server/Agents/AgentCollection.cs
--- a/Engine/Source/Programs/Horde/Horde.Server/Agents/AgentCollection.cs
+++ b/Engine/Source/Programs/Horde/Horde.Server/Agents/AgentCollection.cs
@@ -85,6 +85,9 @@ namespace Horde.Server.Agents
 			[BsonIgnoreIfNull]
 			public DateTime? LastUpgradeTime { get; set; }
 
+			[BsonIgnoreIfNull]
+			public int? UpgradeAttemptCount { get; set; }
+
 			public List<PoolId> DynamicPools { get; set; } = new List<PoolId>();
 			public List<PoolId> Pools { get; set; } = new List<PoolId>();
 
@@ -395,7 +398,9 @@ namespace Horde.Server.Agents
 						}
 						else if (payload.TryUnpack(out UpgradeTask upgradeTask))
 						{
+							int newUpgradeAttemptCount = (agent.UpgradeAttemptCount ?? 0) + 1;
 							updates.Add(updateBuilder.Set(x => x.LastUpgradeVersion, upgradeTask.SoftwareId));
+							updates.Add(updateBuilder.Set(x => x.UpgradeAttemptCount, newUpgradeAttemptCount));
 							updates.Add(updateBuilder.Set(x => x.LastUpgradeTime, DateTime.UtcNow));
 						}
 					}
@@ -480,6 +485,11 @@ namespace Horde.Server.Agents
 			updates.Add(updateBuilder.Unset(x => x.RequestShutdown));
 			updates.Add(updateBuilder.Set(x => x.LastShutdownReason, "Unexpected"));
 
+			if (String.Equals(agent.Version, agent.LastUpgradeVersion, StringComparison.Ordinal))
+			{
+				updates.Add(updateBuilder.Unset(x => x.UpgradeAttemptCount));
+			}
+
 			if (agent.Status != status)
 			{
 				updates.Add(updateBuilder.Set(x => x.LastStatusChange, lastStatusChange));
diff --git a/Engine/Source/Programs/Horde/Horde.Server/Agents/AgentReportService.cs b/Engine/Source/Programs/Horde/Horde.Server/Agents/AgentReportService.cs
--- a/Engine/Source/Programs/Horde/Horde.Server/Agents/AgentReportService.cs
+++ b/Engine/Source/Programs/Horde/Horde.Server/Agents/AgentReportService.cs
@@ -25,6 +25,11 @@ namespace Horde.Server.Agents
 		/// List of agents stuck in a conform loop, plus the number of times they've attempted to conform
 		/// </summary>
 		public List<(AgentId, int)> ConformLoop { get; } = new List<(AgentId, int)>();
+
+		/// <summary>
+		/// List of agents stuck in a upgrade loop, plus the number of times they've attempted to upgrade
+		/// </summary>
+		public List<(AgentId, int)> UpgradeLoop { get; } = new List<(AgentId, int)>();
 	}
 
 	[SingletonDocument("agent-report-state", "6268871c211d05611b3e4fd8")]
@@ -85,7 +90,7 @@ namespace Horde.Server.Agents
 
 			TimeSpan agentReportTime = TimeSpan.FromHours(9.0); // 9am
 
-			DateTime nextUpdateTime = _clock.TimeZone.GetStartOfDayUtc(state.LastUpdateUtc + TimeSpan.FromHours(24.0)) + agentReportTime;
+			DateTime nextUpdateTime = state.LastUpdateUtc + TimeSpan.FromHours(1.0);//_clock.TimeZone.GetStartOfDayUtc(state.LastUpdateUtc + TimeSpan.FromHours(24.0)) + agentReportTime;
 			if (utcNow > nextUpdateTime)
 			{
 				AgentReport report = new AgentReport();
@@ -97,6 +102,10 @@ namespace Horde.Server.Agents
 					{
 						report.ConformLoop.Add((agent.Id, agent.ConformAttemptCount.Value));
 					}
+					if (agent.UpgradeAttemptCount.HasValue && agent.UpgradeAttemptCount.Value > 3)
+					{
+						report.UpgradeLoop.Add((agent.Id, agent.UpgradeAttemptCount.Value));
+					}
 				}
 
 				await _notificationService.SendAgentReportAsync(report);
diff --git a/Engine/Source/Programs/Horde/Horde.Server/Agents/IAgent.cs b/Engine/Source/Programs/Horde/Horde.Server/Agents/IAgent.cs
--- a/Engine/Source/Programs/Horde/Horde.Server/Agents/IAgent.cs
+++ b/Engine/Source/Programs/Horde/Horde.Server/Agents/IAgent.cs
@@ -523,6 +523,11 @@ namespace Horde.Server.Agents
 		/// </summary>
 		public DateTime? LastUpgradeTime { get; }
 
+		/// <summary>
+		/// Number of times an upgrade job has failed
+		/// </summary>
+		public int? UpgradeAttemptCount { get; }
+
 		/// <summary>
 		/// Dynamically applied pools
 		/// </summary>
diff --git a/Engine/Source/Programs/Horde/Horde.Server/Notifications/Sinks/SlackNotificationSink.cs b/Engine/Source/Programs/Horde/Horde.Server/Notifications/Sinks/SlackNotificationSink.cs
--- a/Engine/Source/Programs/Horde/Horde.Server/Notifications/Sinks/SlackNotificationSink.cs
+++ b/Engine/Source/Programs/Horde/Horde.Server/Notifications/Sinks/SlackNotificationSink.cs
@@ -3037,6 +3037,27 @@ namespace Horde.Server.Notifications.Sinks
 						}
 					}
 				}
+
+				{
+					StringBuilder upgradeMessage = new StringBuilder("*Upgrade issues:*\n");
+					if (report.UpgradeLoop.Count == 0)
+					{
+						upgradeMessage.Append("None.\n");
+						await _slackClient.PostMessageAsync(_settings.AgentNotificationChannel, upgradeMessage.ToString());
+					}
+					else
+					{
+						foreach (IReadOnlyList<(AgentId, int)> upgradeBatch in report.UpgradeLoop.Batch(10))
+						{
+							foreach ((AgentId agentId, int upgradeCount) in upgradeBatch)
+							{
+								Uri agentUrl = new Uri(_settings.DashboardUrl, $"agents?agentId={agentId}");
+								upgradeMessage.Append($"* *[{agentId}]({agentUrl})* has attempted to upgrade {upgradeCount} times\n");
+							}
+							await _slackClient.PostMessageAsync(_settings.AgentNotificationChannel, upgradeMessage.ToString());
+						}
+					}
+				}
 			}
 		}
 	}
