Horde: Validate length argument passed to AWS storage backend.

diff --git a/Engine/Source/Programs/Horde/Horde.Server/Storage/Backends/AwsStorageBackend.cs b/Engine/Source/Programs/Horde/Horde.Server/Storage/Backends/AwsStorageBackend.cs
--- a/Engine/Source/Programs/Horde/Horde.Server/Storage/Backends/AwsStorageBackend.cs
+++ b/Engine/Source/Programs/Horde/Horde.Server/Storage/Backends/AwsStorageBackend.cs
@@ -277,6 +277,10 @@ namespace Horde.Server.Storage.Backends
 			{
 				range = $"bytes={offset}-";
 			}
+			else if (length.Value == 0)
+			{
+				throw new ArgumentException("Cannot read empty stream from AWS backend", nameof(length));
+			}
 			else
 			{
 				range = $"bytes={offset}-{offset + length.Value - 1}";
