Horde: Fix streams not being disposed when reading from bundles.

diff --git a/Engine/Source/Programs/Horde/Horde.Server/Storage/StorageService.cs b/Engine/Source/Programs/Horde/Horde.Server/Storage/StorageService.cs
--- a/Engine/Source/Programs/Horde/Horde.Server/Storage/StorageService.cs
+++ b/Engine/Source/Programs/Horde/Horde.Server/Storage/StorageService.cs
@@ -1053,13 +1053,13 @@ namespace Horde.Server.Storage
 			await _gcState.UpdateAsync(state => state.FindOrAddNamespace(namespaceId).LastTime = utcNow);
 		}
 
-		async Task<BundleHeader?> ReadHeaderAsync(BundleStorageClient store, BundleLocator locator, CancellationToken cancellationToken)
+		static async Task<BundleHeader?> ReadHeaderAsync(BundleStorageClient store, BundleLocator locator, CancellationToken cancellationToken)
 		{
 			int fetchSize = 64 * 1024;
 			for (; ; )
 			{
 				// Read the start of the blob
-				Stream stream = await store.OpenAsync(locator, 0, fetchSize, cancellationToken);
+				using Stream stream = await store.OpenAsync(locator, 0, fetchSize, cancellationToken);
 
 				byte[] prelude = new byte[BundleHeader.PreludeLength];
 				await stream.ReadFixedLengthBytesAsync(prelude, cancellationToken);
