Horde: Fix stat tracking for agents stuck in a conform/upgrade loop.

diff --git a/Engine/Source/Programs/Horde/Horde.Server/Agents/AgentCollection.cs b/Engine/Source/Programs/Horde/Horde.Server/Agents/AgentCollection.cs
--- a/Engine/Source/Programs/Horde/Horde.Server/Agents/AgentCollection.cs
+++ b/Engine/Source/Programs/Horde/Horde.Server/Agents/AgentCollection.cs
@@ -389,20 +389,7 @@ namespace Horde.Server.Agents
 				{
 					if (lease.Payload != null && (agent.Leases == null || !agent.Leases.Any(x => x.Id == lease.Id)))
 					{
-						Any payload = Any.Parser.ParseFrom(lease.Payload.ToArray());
-						if (payload.TryUnpack(out ConformTask conformTask))
-						{
-							int newConformAttemptCount = (agent.ConformAttemptCount ?? 0) + 1;
-							updates.Add(updateBuilder.Set(x => x.ConformAttemptCount, newConformAttemptCount));
-							updates.Add(updateBuilder.Set(x => x.LastConformTime, DateTime.UtcNow));
-						}
-						else if (payload.TryUnpack(out UpgradeTask upgradeTask))
-						{
-							int newUpgradeAttemptCount = (agent.UpgradeAttemptCount ?? 0) + 1;
-							updates.Add(updateBuilder.Set(x => x.LastUpgradeVersion, upgradeTask.SoftwareId));
-							updates.Add(updateBuilder.Set(x => x.UpgradeAttemptCount, newUpgradeAttemptCount));
-							updates.Add(updateBuilder.Set(x => x.LastUpgradeTime, DateTime.UtcNow));
-						}
+						GetNewLeaseUpdates(agent, lease, updates); 
 					}
 				}
 
@@ -531,10 +518,35 @@ namespace Horde.Server.Agents
 			}
 			leases.Add(newLease);
 
-			UpdateDefinition<AgentDocument> update = Builders<AgentDocument>.Update.Set(x => x.Leases, leases);
+			List<UpdateDefinition<AgentDocument>> updates = new List<UpdateDefinition<AgentDocument>>();
+			updates.Add(Builders<AgentDocument>.Update.Set(x => x.Leases, leases));
+			GetNewLeaseUpdates(agent, newLease, updates);
+
+			UpdateDefinition<AgentDocument> update = Builders<AgentDocument>.Update.Combine(updates);
 			return await TryUpdateAsync(agent, update);
 		}
 
+		static void GetNewLeaseUpdates(IAgent agent, AgentLease lease, List<UpdateDefinition<AgentDocument>> updates)
+		{
+			if (lease.Payload != null)
+			{
+				Any payload = Any.Parser.ParseFrom(lease.Payload.ToArray());
+				if (payload.TryUnpack(out ConformTask conformTask))
+				{
+					int newConformAttemptCount = (agent.ConformAttemptCount ?? 0) + 1;
+					updates.Add(Builders<AgentDocument>.Update.Set(x => x.ConformAttemptCount, newConformAttemptCount));
+					updates.Add(Builders<AgentDocument>.Update.Set(x => x.LastConformTime, DateTime.UtcNow));
+				}
+				else if (payload.TryUnpack(out UpgradeTask upgradeTask))
+				{
+					int newUpgradeAttemptCount = (agent.UpgradeAttemptCount ?? 0) + 1;
+					updates.Add(Builders<AgentDocument>.Update.Set(x => x.LastUpgradeVersion, upgradeTask.SoftwareId));
+					updates.Add(Builders<AgentDocument>.Update.Set(x => x.UpgradeAttemptCount, newUpgradeAttemptCount));
+					updates.Add(Builders<AgentDocument>.Update.Set(x => x.LastUpgradeTime, DateTime.UtcNow));
+				}
+			}
+		}
+
 		/// <inheritdoc/>
 		public async Task<IAgent?> TryCancelLeaseAsync(IAgent agentInterface, int leaseIdx)
 		{
