Horde: Fix formatting of Slack agent reports.

diff --git a/Engine/Source/Programs/Horde/Horde.Server/Notifications/Sinks/SlackNotificationSink.cs b/Engine/Source/Programs/Horde/Horde.Server/Notifications/Sinks/SlackNotificationSink.cs
--- a/Engine/Source/Programs/Horde/Horde.Server/Notifications/Sinks/SlackNotificationSink.cs
+++ b/Engine/Source/Programs/Horde/Horde.Server/Notifications/Sinks/SlackNotificationSink.cs
@@ -3031,9 +3031,10 @@ namespace Horde.Server.Notifications.Sinks
 							foreach ((AgentId agentId, int conformCount) in conformBatch)
 							{
 								Uri agentUrl = new Uri(_settings.DashboardUrl, $"agents?agentId={agentId}");
-								conformMessage.Append($"* *[{agentId}]({agentUrl})* has run conform {conformCount} times\n");
+								conformMessage.Append($"\u2022 *<{agentUrl}|{agentId}>* has run conform {conformCount} times\n");
 							}
 							await _slackClient.PostMessageAsync(_settings.AgentNotificationChannel, conformMessage.ToString());
+							conformMessage.Clear();
 						}
 					}
 				}
@@ -3052,9 +3053,10 @@ namespace Horde.Server.Notifications.Sinks
 							foreach ((AgentId agentId, int upgradeCount) in upgradeBatch)
 							{
 								Uri agentUrl = new Uri(_settings.DashboardUrl, $"agents?agentId={agentId}");
-								upgradeMessage.Append($"* *[{agentId}]({agentUrl})* has attempted to upgrade {upgradeCount} times\n");
+								upgradeMessage.Append($"\u2022 *<{agentUrl}|{agentId}>* has attempted to upgrade {upgradeCount} times\n");
 							}
 							await _slackClient.PostMessageAsync(_settings.AgentNotificationChannel, upgradeMessage.ToString());
+							upgradeMessage.Clear();
 						}
 					}
 				}
