Horde: Reduce the length of upgrade messages in Slack notifications.

diff --git a/Engine/Source/Programs/Horde/Horde.Server/Agents/AgentCollection.cs b/Engine/Source/Programs/Horde/Horde.Server/Agents/AgentCollection.cs
--- a/Engine/Source/Programs/Horde/Horde.Server/Agents/AgentCollection.cs
+++ b/Engine/Source/Programs/Horde/Horde.Server/Agents/AgentCollection.cs
@@ -539,8 +539,16 @@ namespace Horde.Server.Agents
 				}
 				else if (payload.TryUnpack(out UpgradeTask upgradeTask))
 				{
+					string newVersion = upgradeTask.SoftwareId;
+
+					int versionIdx = newVersion.IndexOf(':');
+					if (versionIdx != -1)
+					{
+						newVersion = newVersion.Substring(versionIdx + 1);
+					}
+
 					int newUpgradeAttemptCount = (agent.UpgradeAttemptCount ?? 0) + 1;
-					updates.Add(Builders<AgentDocument>.Update.Set(x => x.LastUpgradeVersion, upgradeTask.SoftwareId));
+					updates.Add(Builders<AgentDocument>.Update.Set(x => x.LastUpgradeVersion, newVersion));
 					updates.Add(Builders<AgentDocument>.Update.Set(x => x.UpgradeAttemptCount, newUpgradeAttemptCount));
 					updates.Add(Builders<AgentDocument>.Update.Set(x => x.LastUpgradeTime, DateTime.UtcNow));
 				}
diff --git a/Engine/Source/Programs/Horde/Horde.Server/Notifications/Sinks/SlackNotificationSink.cs b/Engine/Source/Programs/Horde/Horde.Server/Notifications/Sinks/SlackNotificationSink.cs
--- a/Engine/Source/Programs/Horde/Horde.Server/Notifications/Sinks/SlackNotificationSink.cs
+++ b/Engine/Source/Programs/Horde/Horde.Server/Notifications/Sinks/SlackNotificationSink.cs
@@ -3022,21 +3022,21 @@ namespace Horde.Server.Notifications.Sinks
 					if (report.ConformLoop.Count == 0)
 					{
 						conformMessage.Append("None.\n");
-						await _slackClient.PostMessageAsync(_settings.AgentNotificationChannel, conformMessage.ToString());
 					}
 					else
 					{
-						foreach (IReadOnlyList<(AgentId, int)> conformBatch in report.ConformLoop.Batch(10))
+						const int NumItems = 15;
+						foreach ((AgentId agentId, int conformCount) in report.ConformLoop.OrderBy(x => x.Item1.ToString()).Take(NumItems))
 						{
-							foreach ((AgentId agentId, int conformCount) in conformBatch)
-							{
-								Uri agentUrl = new Uri(_settings.DashboardUrl, $"agents?agentId={agentId}");
-								conformMessage.Append($"\u2022 *<{agentUrl}|{agentId}>* has run conform {conformCount} times\n");
-							}
-							await _slackClient.PostMessageAsync(_settings.AgentNotificationChannel, conformMessage.ToString());
-							conformMessage.Clear();
+							Uri agentUrl = new Uri(_settings.DashboardUrl, $"agents?agentId={agentId}");
+							conformMessage.Append($"\u2022 *<{agentUrl}|{agentId}>* has run conform {conformCount} times\n");
+						}
+						if (report.ConformLoop.Count > NumItems)
+						{
+							conformMessage.Append($"+ {report.ConformLoop.Count - NumItems} other(s)\n");
 						}
 					}
+					await _slackClient.PostMessageAsync(_settings.AgentNotificationChannel, conformMessage.ToString());
 				}
 
 				{
@@ -3044,21 +3044,21 @@ namespace Horde.Server.Notifications.Sinks
 					if (report.UpgradeLoop.Count == 0)
 					{
 						upgradeMessage.Append("None.\n");
-						await _slackClient.PostMessageAsync(_settings.AgentNotificationChannel, upgradeMessage.ToString());
 					}
 					else
 					{
-						foreach (IReadOnlyList<(AgentId, int)> upgradeBatch in report.UpgradeLoop.Batch(10))
+						const int NumItems = 15;
+						foreach ((AgentId agentId, int upgradeCount) in report.UpgradeLoop.OrderBy(x => x.Item1.ToString()).Take(NumItems))
 						{
-							foreach ((AgentId agentId, int upgradeCount) in upgradeBatch)
-							{
-								Uri agentUrl = new Uri(_settings.DashboardUrl, $"agents?agentId={agentId}");
-								upgradeMessage.Append($"\u2022 *<{agentUrl}|{agentId}>* has attempted to upgrade {upgradeCount} times\n");
-							}
-							await _slackClient.PostMessageAsync(_settings.AgentNotificationChannel, upgradeMessage.ToString());
-							upgradeMessage.Clear();
+							Uri agentUrl = new Uri(_settings.DashboardUrl, $"agents?agentId={agentId}");
+							upgradeMessage.Append($"\u2022 *<{agentUrl}|{agentId}>* has attempted to upgrade {upgradeCount} times\n");
+						}
+						if (report.ConformLoop.Count > NumItems)
+						{
+							upgradeMessage.Append($"+ {report.ConformLoop.Count - NumItems} other(s)\n");
 						}
 					}
+					await _slackClient.PostMessageAsync(_settings.AgentNotificationChannel, upgradeMessage.ToString());
 				}
 			}
 		}
