Horde: Reduce frequency of agent reports.

diff --git a/Engine/Source/Programs/Horde/Horde.Server/Agents/AgentReportService.cs b/Engine/Source/Programs/Horde/Horde.Server/Agents/AgentReportService.cs
--- a/Engine/Source/Programs/Horde/Horde.Server/Agents/AgentReportService.cs
+++ b/Engine/Source/Programs/Horde/Horde.Server/Agents/AgentReportService.cs
@@ -48,21 +48,17 @@ namespace Horde.Server.Agents
 		readonly INotificationService _notificationService;
 		readonly IClock _clock;
 		readonly ITicker _ticker;
-		readonly IOptionsMonitor<GlobalConfig> _globalConfig;
-		readonly ILogger _logger;
 
 		/// <summary>
 		/// Constructor
 		/// </summary>
-		public AgentReportService(MongoService mongoService, IAgentCollection agentCollection, INotificationService notificationService, IClock clock, IOptionsMonitor<GlobalConfig> globalConfig, ILogger<AgentReportService> logger)
+		public AgentReportService(MongoService mongoService, IAgentCollection agentCollection, INotificationService notificationService, IClock clock, ILogger<AgentReportService> logger)
 		{
 			_state = new SingletonDocument<AgentReportState>(mongoService);
 			_agentCollection = agentCollection;
 			_notificationService = notificationService;
 			_clock = clock;
 			_ticker = clock.AddSharedTicker<AgentReportService>(TimeSpan.FromMinutes(5.0), TickAsync, logger);
-			_globalConfig = globalConfig;
-			_logger = logger;
 		}
 
 		/// <inheritdoc/>
@@ -90,7 +86,7 @@ namespace Horde.Server.Agents
 
 			TimeSpan agentReportTime = TimeSpan.FromHours(9.0); // 9am
 
-			DateTime nextUpdateTime = state.LastUpdateUtc + TimeSpan.FromHours(1.0);//_clock.TimeZone.GetStartOfDayUtc(state.LastUpdateUtc + TimeSpan.FromHours(24.0)) + agentReportTime;
+			DateTime nextUpdateTime = _clock.TimeZone.GetStartOfDayUtc(state.LastUpdateUtc + TimeSpan.FromHours(24.0)) + agentReportTime;
 			if (utcNow > nextUpdateTime)
 			{
 				AgentReport report = new AgentReport();
diff --git a/Engine/Source/Programs/Horde/Horde.Server/Notifications/Sinks/SlackNotificationSink.cs b/Engine/Source/Programs/Horde/Horde.Server/Notifications/Sinks/SlackNotificationSink.cs
--- a/Engine/Source/Programs/Horde/Horde.Server/Notifications/Sinks/SlackNotificationSink.cs
+++ b/Engine/Source/Programs/Horde/Horde.Server/Notifications/Sinks/SlackNotificationSink.cs
@@ -3025,8 +3025,8 @@ namespace Horde.Server.Notifications.Sinks
 					}
 					else
 					{
-						const int NumItems = 15;
-						foreach ((AgentId agentId, int conformCount) in report.ConformLoop.OrderBy(x => x.Item1.ToString()).Take(NumItems))
+						const int NumItems = 10;
+						foreach ((AgentId agentId, int conformCount) in report.ConformLoop.OrderByDescending(x => x.Item2).ThenBy(x => x.Item1.ToString()).Take(NumItems))
 						{
 							Uri agentUrl = new Uri(_settings.DashboardUrl, $"agents?agentId={agentId}");
 							conformMessage.Append($"\u2022 *<{agentUrl}|{agentId}>* has run conform {conformCount} times\n");
@@ -3047,15 +3047,15 @@ namespace Horde.Server.Notifications.Sinks
 					}
 					else
 					{
-						const int NumItems = 15;
-						foreach ((AgentId agentId, int upgradeCount) in report.UpgradeLoop.OrderBy(x => x.Item1.ToString()).Take(NumItems))
+						const int NumItems = 10;
+						foreach ((AgentId agentId, int upgradeCount) in report.UpgradeLoop.OrderByDescending(x => x.Item2).ThenBy(x => x.Item1.ToString()).Take(NumItems))
 						{
 							Uri agentUrl = new Uri(_settings.DashboardUrl, $"agents?agentId={agentId}");
 							upgradeMessage.Append($"\u2022 *<{agentUrl}|{agentId}>* has attempted to upgrade {upgradeCount} times\n");
 						}
-						if (report.ConformLoop.Count > NumItems)
+						if (report.UpgradeLoop.Count > NumItems)
 						{
-							upgradeMessage.Append($"+ {report.ConformLoop.Count - NumItems} other(s)\n");
+							upgradeMessage.Append($"+ {report.UpgradeLoop.Count - NumItems} other(s)\n");
 						}
 					}
 					await _slackClient.PostMessageAsync(_settings.AgentNotificationChannel, upgradeMessage.ToString());
diff --git a/Engine/Source/Programs/Horde/Horde.Server/Tasks/UpgradeTaskSource.cs b/Engine/Source/Programs/Horde/Horde.Server/Tasks/UpgradeTaskSource.cs
--- a/Engine/Source/Programs/Horde/Horde.Server/Tasks/UpgradeTaskSource.cs
+++ b/Engine/Source/Programs/Horde/Horde.Server/Tasks/UpgradeTaskSource.cs
@@ -51,7 +51,7 @@ namespace Horde.Server.Tasks
 				return Skip(cancellationToken);
 			}
 
-			(ITool, IToolDeployment)? required = await GetRequiredSoftwareVersion(agent);
+			(ITool, IToolDeployment)? required = await GetRequiredSoftwareVersionAsync(agent);
 			if (required == null)
 			{
 				return Skip(cancellationToken);
@@ -85,7 +85,7 @@ namespace Horde.Server.Tasks
 		/// </summary>
 		/// <param name="agent">The agent instance</param>
 		/// <returns>Unique id of the client version this agent should be running</returns>
-		public async Task<(ITool, IToolDeployment)?> GetRequiredSoftwareVersion(IAgent agent)
+		public async Task<(ITool, IToolDeployment)?> GetRequiredSoftwareVersionAsync(IAgent agent)
 		{
 			GlobalConfig globalConfig = _globalConfig.CurrentValue;
 
