Horde: Update preflight tag on successful preflight without requiring autosubmit
#rnx

diff --git a/Engine/Source/Programs/Horde/Horde.Server/Jobs/JobService.cs b/Engine/Source/Programs/Horde/Horde.Server/Jobs/JobService.cs
--- a/Engine/Source/Programs/Horde/Horde.Server/Jobs/JobService.cs
+++ b/Engine/Source/Programs/Horde/Horde.Server/Jobs/JobService.cs
@@ -1036,6 +1036,10 @@ namespace Horde.Server.Jobs
 							{
 								job = await AutoSubmitChangeAsync(streamConfig, job, graph);
 							}
+							else if (job.ClonedPreflightChange == 0 && outcome == JobStepOutcome.Success && job.AbortedByUserId == null)
+							{
+								await _perforceService.UpdateChangelistDescriptionAsync(streamConfig.ClusterName, job.PreflightChange, x => x.TrimEnd() + $"\n#preflight {job.Id}");
+							}
 							else if (job.ClonedPreflightChange != 0)
 							{
 								await DeleteShelvedChangeAsync(streamConfig.ClusterName, job.ClonedPreflightChange);
