Horde: Fix new artifacts modal query modificatons on log view
#skipci
#rnx

diff --git a/Engine/Source/Programs/Horde/HordeDashboard/src/components/LogView.tsx b/Engine/Source/Programs/Horde/HordeDashboard/src/components/LogView.tsx
--- a/Engine/Source/Programs/Horde/HordeDashboard/src/components/LogView.tsx
+++ b/Engine/Source/Programs/Horde/HordeDashboard/src/components/LogView.tsx
@@ -891,14 +891,21 @@ export const LogList: React.FC<{ logId: string }> = observer(({ logId }) => {
 
          const opsList: IContextualMenuItem[] = [];
 
-         const baseUrl = window.location.pathname + window.location.search;
+         const navigateToArtifacts = (context: string) => {
+            const search = new URLSearchParams(window.location.search);
+            search.set("artifactContext", encodeURIComponent(context));
+            const url = `${window.location.pathname}?` + search.toString();      
+            navigate(url, { replace: true })
+         }
 
          opsList.push({
             key: 'stepops_artifacts_step',
             text: "Logs",
             iconProps: { iconName: "Folder" },
             disabled: !atypes.get("step-saved"),
-            onClick: () => { navigate(`${baseUrl}?artifactContext=step-saved`, { replace: true }) }
+            onClick: () => {
+               navigateToArtifacts("step-saved");
+            }
          });
 
          opsList.push({
@@ -906,7 +913,9 @@ export const LogList: React.FC<{ logId: string }> = observer(({ logId }) => {
             text: "Temp Storage",
             iconProps: { iconName: "MenuOpen" },
             disabled: !atypes.get("step-output"),
-            onClick: () => { navigate(`${baseUrl}?artifactContext=step-output`, { replace: true }) }
+            onClick: () => {
+               navigateToArtifacts("step-output");
+            }
          });
 
          opsList.push({
@@ -914,7 +923,9 @@ export const LogList: React.FC<{ logId: string }> = observer(({ logId }) => {
             text: "Traces",
             iconProps: { iconName: "SearchTemplate" },
             disabled: !atypes.get("step-trace"),
-            onClick: () => { navigate(`${baseUrl}?artifactContext=step-trace`, { replace: true }) }
+            onClick: () => {
+               navigateToArtifacts("step-trace");
+            }
          });
 
          menuProps.items.push({
