Horde: Support tagging user preflights without needing autosubmit set
#rnx

diff --git a/Engine/Source/Programs/Horde/Horde.Server/Jobs/JobService.cs b/Engine/Source/Programs/Horde/Horde.Server/Jobs/JobService.cs
--- a/Engine/Source/Programs/Horde/Horde.Server/Jobs/JobService.cs
+++ b/Engine/Source/Programs/Horde/Horde.Server/Jobs/JobService.cs
@@ -1036,6 +1036,14 @@ namespace Horde.Server.Jobs
 							{
 								job = await AutoSubmitChangeAsync(streamConfig, job, graph);
 							}
+							else if (job.ClonedPreflightChange == 0 && job.StartedByUserId.HasValue && outcome == JobStepOutcome.Success && job.AbortedByUserId == null)
+							{
+								IUserSettings settings = await _userCollection.GetSettingsAsync(job.StartedByUserId.Value);
+								if (settings.AlwaysTagPreflightCL)
+								{
+									await _perforceService.UpdateChangelistDescriptionAsync(streamConfig.ClusterName, job.PreflightChange, x => x.TrimEnd() + $"\n#preflight {job.Id}");
+								}
+							}
 							else if (job.ClonedPreflightChange != 0)
 							{
 								await DeleteShelvedChangeAsync(streamConfig.ClusterName, job.ClonedPreflightChange);
diff --git a/Engine/Source/Programs/Horde/Horde.Server/Users/IUser.cs b/Engine/Source/Programs/Horde/Horde.Server/Users/IUser.cs
--- a/Engine/Source/Programs/Horde/Horde.Server/Users/IUser.cs
+++ b/Engine/Source/Programs/Horde/Horde.Server/Users/IUser.cs
@@ -138,6 +138,11 @@ namespace Horde.Server.Users
 		/// </summary>
 		public bool EnableExperimentalFeatures { get; }
 
+		/// <summary>
+		/// Whether to always tag CL descriptions for preflights
+		/// </summary>
+		public bool AlwaysTagPreflightCL { get; }
+
 		/// <summary>
 		/// Opaque settings dictionary for the dashboard
 		/// </summary>
@@ -182,6 +187,7 @@ namespace Horde.Server.Users
 			if (settings != null)
 			{
 				response.EnableExperimentalFeatures = settings.EnableExperimentalFeatures;
+				response.AlwaysTagPreflightCL = settings.AlwaysTagPreflightCL;
 
 				response.DashboardSettings = BsonTypeMapper.MapToDotNetValue(settings.DashboardSettings);
 				response.PinnedJobIds = settings.PinnedJobIds.ConvertAll(x => x.ToString());
diff --git a/Engine/Source/Programs/Horde/Horde.Server/Users/IUserCollection.cs b/Engine/Source/Programs/Horde/Horde.Server/Users/IUserCollection.cs
--- a/Engine/Source/Programs/Horde/Horde.Server/Users/IUserCollection.cs
+++ b/Engine/Source/Programs/Horde/Horde.Server/Users/IUserCollection.cs
@@ -90,6 +90,7 @@ namespace Horde.Server.Users
 		/// </summary>
 		/// <param name="userId">The user to update</param>
 		/// <param name="enableExperimentalFeatures"></param>
+		/// <param name="alwaysTagPreflightCL"></param>
 		/// <param name="dashboardSettings">Opaque settings object for the dashboard</param>
 		/// <param name="addPinnedJobIds"></param>
 		/// <param name="removePinnedJobIds"></param>
@@ -97,7 +98,7 @@ namespace Horde.Server.Users
 		/// <param name="addBisectTaskIds"></param>
 		/// <param name="removeBisectTaskIds"></param>
 		/// <returns>Updated user object</returns>
-		Task UpdateSettingsAsync(UserId userId, bool? enableExperimentalFeatures = null, BsonValue? dashboardSettings = null, IEnumerable<JobId>? addPinnedJobIds = null, IEnumerable<JobId>? removePinnedJobIds = null, UpdateUserJobTemplateOptions? templateOptions = null, IEnumerable<BisectTaskId>? addBisectTaskIds = null, IEnumerable<BisectTaskId>? removeBisectTaskIds = null);
+		Task UpdateSettingsAsync(UserId userId, bool? enableExperimentalFeatures = null, bool? alwaysTagPreflightCL = null, BsonValue? dashboardSettings = null, IEnumerable<JobId>? addPinnedJobIds = null, IEnumerable<JobId>? removePinnedJobIds = null, UpdateUserJobTemplateOptions? templateOptions = null, IEnumerable<BisectTaskId>? addBisectTaskIds = null, IEnumerable<BisectTaskId>? removeBisectTaskIds = null);
 	}
 
 	/// <summary>
diff --git a/Engine/Source/Programs/Horde/Horde.Server/Users/UserCollection.cs b/Engine/Source/Programs/Horde/Horde.Server/Users/UserCollection.cs
--- a/Engine/Source/Programs/Horde/Horde.Server/Users/UserCollection.cs
+++ b/Engine/Source/Programs/Horde/Horde.Server/Users/UserCollection.cs
@@ -32,6 +32,9 @@ namespace Horde.Server.Users
 			[BsonDefaultValue(false), BsonIgnoreIfDefault]
 			public bool EnableExperimentalFeatures { get; set; }
 
+			[BsonDefaultValue(false), BsonIgnoreIfDefault]
+			public bool AlwaysTagPreflightCL { get; set; }
+
 			public BsonValue DashboardSettings { get; set; } = BsonNull.Value;
 			public List<JobId> PinnedJobIds { get; set; } = new List<JobId>();
 			public List<BisectTaskId> PinnedBisectTaskIds { get; set; } = new List<BisectTaskId>();
@@ -161,7 +164,7 @@ namespace Horde.Server.Users
 		}
 
 		/// <inheritdoc/>
-		public async Task UpdateSettingsAsync(UserId userId, bool? enableExperimentalFeatures, BsonValue? dashboardSettings = null, IEnumerable<JobId>? addPinnedJobIds = null, IEnumerable<JobId>? removePinnedJobIds = null, UpdateUserJobTemplateOptions? templateOptions = null, IEnumerable<BisectTaskId>? addBisectTaskIds = null, IEnumerable<BisectTaskId>? removeBisectTaskIds = null)
+		public async Task UpdateSettingsAsync(UserId userId, bool? enableExperimentalFeatures, bool? alwaysTagPreflightCL = null, BsonValue ? dashboardSettings = null, IEnumerable<JobId>? addPinnedJobIds = null, IEnumerable<JobId>? removePinnedJobIds = null, UpdateUserJobTemplateOptions? templateOptions = null, IEnumerable<BisectTaskId>? addBisectTaskIds = null, IEnumerable<BisectTaskId>? removeBisectTaskIds = null)
 		{
 			if (addPinnedJobIds != null)
 			{
@@ -178,6 +181,10 @@ namespace Horde.Server.Users
 			{
 				updates.Add(Builders<UserDocument>.Update.SetOrUnsetNull(x => x.EnableExperimentalFeatures, enableExperimentalFeatures));
 			}
+			if (alwaysTagPreflightCL != null)
+			{
+				updates.Add(Builders<UserDocument>.Update.SetOrUnsetNull(x => x.AlwaysTagPreflightCL, alwaysTagPreflightCL));
+			}
 			if (dashboardSettings != null)
 			{
 				updates.Add(Builders<UserDocument>.Update.Set(x => x.DashboardSettings, dashboardSettings));
diff --git a/Engine/Source/Programs/Horde/Horde.Server/Users/UserCollectionV2.cs b/Engine/Source/Programs/Horde/Horde.Server/Users/UserCollectionV2.cs
--- a/Engine/Source/Programs/Horde/Horde.Server/Users/UserCollectionV2.cs
+++ b/Engine/Source/Programs/Horde/Horde.Server/Users/UserCollectionV2.cs
@@ -129,6 +129,9 @@ namespace Horde.Server.Users
 			[BsonDefaultValue(false), BsonIgnoreIfDefault]
 			public bool EnableExperimentalFeatures { get; set; }
 
+			[BsonDefaultValue(false), BsonIgnoreIfDefault]
+			public bool AlwaysTagPreflightCL { get; set; }
+
 			public BsonValue DashboardSettings { get; set; } = BsonNull.Value;
 			public List<JobId> PinnedJobIds { get; set; } = new List<JobId>();
 			public List<BisectTaskId> PinnedBisectTaskIds { get; set; } = new List<BisectTaskId>();
@@ -357,13 +360,17 @@ namespace Horde.Server.Users
 		}
 
 		/// <inheritdoc/>
-		public async Task UpdateSettingsAsync(UserId userId, bool? enableExperimentalFeatures = null, BsonValue? dashboardSettings = null, IEnumerable<JobId>? addPinnedJobIds = null, IEnumerable<JobId>? removePinnedJobIds = null, UpdateUserJobTemplateOptions? templateOptions = null, IEnumerable<BisectTaskId>? addBisectTaskIds = null, IEnumerable<BisectTaskId>? removeBisectTaskIds = null)
+		public async Task UpdateSettingsAsync(UserId userId, bool? enableExperimentalFeatures = null, bool? alwaysTagPreflightCL = null, BsonValue ? dashboardSettings = null, IEnumerable<JobId>? addPinnedJobIds = null, IEnumerable<JobId>? removePinnedJobIds = null, UpdateUserJobTemplateOptions? templateOptions = null, IEnumerable<BisectTaskId>? addBisectTaskIds = null, IEnumerable<BisectTaskId>? removeBisectTaskIds = null)
 		{
 			List<UpdateDefinition<UserSettingsDocument>> updates = new List<UpdateDefinition<UserSettingsDocument>>();
 			if (enableExperimentalFeatures != null)
 			{
 				updates.Add(Builders<UserSettingsDocument>.Update.SetOrUnsetNull(x => x.EnableExperimentalFeatures, enableExperimentalFeatures));
 			}
+			if (alwaysTagPreflightCL != null)
+			{
+				updates.Add(Builders<UserSettingsDocument>.Update.SetOrUnsetNull(x => x.AlwaysTagPreflightCL, alwaysTagPreflightCL));
+			}
 			if (dashboardSettings != null)
 			{
 				updates.Add(Builders<UserSettingsDocument>.Update.Set(x => x.DashboardSettings, dashboardSettings));
@@ -386,7 +393,6 @@ namespace Horde.Server.Users
 				updates.Add(Builders<UserSettingsDocument>.Update.PullAll(x => x.PinnedBisectTaskIds, removeBisectTaskIds));
 			}
 
-
 			if (templateOptions != null)
 			{
 				JobTemplateSettingsDocument doc = new JobTemplateSettingsDocument(templateOptions.StreamId, templateOptions.TemplateId, templateOptions.TemplateHash, templateOptions.Arguments.ToList());
diff --git a/Engine/Source/Programs/Horde/Horde.Server/Users/UserController.cs b/Engine/Source/Programs/Horde/Horde.Server/Users/UserController.cs
--- a/Engine/Source/Programs/Horde/Horde.Server/Users/UserController.cs
+++ b/Engine/Source/Programs/Horde/Horde.Server/Users/UserController.cs
@@ -1,6 +1,5 @@
 // Copyright Epic Games, Inc. All Rights Reserved.
 
-using System.Collections.Generic;
 using System.Linq;
 using System.Security.Claims;
 using System.Threading.Tasks;
@@ -95,7 +94,7 @@ namespace Horde.Server.Users
 				return BadRequest("Current user does not have a registered profile");
 			}
 
-			await _userCollection.UpdateSettingsAsync(userId.Value, request.EnableExperimentalFeatures, request.DashboardSettings?.ToBsonValue(), request.AddPinnedJobIds?.Select(x => JobId.Parse(x)), request.RemovePinnedJobIds?.Select(x => JobId.Parse(x)), null, request.AddPinnedBisectTaskIds?.Select(x => BisectTaskId.Parse(x)), request.RemovePinnedBisectTaskIds?.Select(x => BisectTaskId.Parse(x)));
+			await _userCollection.UpdateSettingsAsync(userId.Value, request.EnableExperimentalFeatures, request.AlwaysTagPreflightCL, request.DashboardSettings?.ToBsonValue(), request.AddPinnedJobIds?.Select(x => JobId.Parse(x)), request.RemovePinnedJobIds?.Select(x => JobId.Parse(x)), null, request.AddPinnedBisectTaskIds?.Select(x => BisectTaskId.Parse(x)), request.RemovePinnedBisectTaskIds?.Select(x => BisectTaskId.Parse(x)));
 			return Ok();
 		}
 	}
diff --git a/Engine/Source/Programs/Shared/EpicGames.Horde/Api/UserMessages.cs b/Engine/Source/Programs/Shared/EpicGames.Horde/Api/UserMessages.cs
--- a/Engine/Source/Programs/Shared/EpicGames.Horde/Api/UserMessages.cs
+++ b/Engine/Source/Programs/Shared/EpicGames.Horde/Api/UserMessages.cs
@@ -2,7 +2,6 @@
 
 using System;
 using System.Collections.Generic;
-using System.Linq;
 using System.Text.Json;
 
 #pragma warning disable CA2227
@@ -59,6 +58,11 @@ namespace EpicGames.Horde.Api
 		/// </summary>
 		public bool? EnableExperimentalFeatures { get; set; }
 
+		/// <summary>
+		/// Whether to always tag preflight changelists
+		/// </summary>
+		public bool? AlwaysTagPreflightCL { get; set; }
+
 		/// <summary>
 		/// Settings for the dashboard
 		/// </summary>
@@ -260,6 +264,11 @@ namespace EpicGames.Horde.Api
 		/// </summary>
 		public bool? EnableExperimentalFeatures { get; set; }
 
+		/// <summary>
+		/// Whether to always tag preflight CL
+		/// </summary>
+		public bool? AlwaysTagPreflightCL { get; set; }
+
 		/// <summary>
 		/// New dashboard settings
 		/// </summary>
@@ -284,6 +293,5 @@ namespace EpicGames.Horde.Api
 		/// Bisection task ids to remove from the pinned list
 		/// </summary>
 		public List<string>? RemovePinnedBisectTaskIds { get; set; }
-
 	}
 }
