Horde: Do not include machines in the agent report which are not online.

diff --git a/Engine/Source/Programs/Horde/Horde.Server/Agents/AgentReportService.cs b/Engine/Source/Programs/Horde/Horde.Server/Agents/AgentReportService.cs
--- a/Engine/Source/Programs/Horde/Horde.Server/Agents/AgentReportService.cs
+++ b/Engine/Source/Programs/Horde/Horde.Server/Agents/AgentReportService.cs
@@ -94,13 +94,16 @@ namespace Horde.Server.Agents
 				List<IAgent> agents = await _agentCollection.FindAsync();
 				foreach (IAgent agent in agents)
 				{
-					if (agent.ConformAttemptCount.HasValue && agent.ConformAttemptCount.Value > 3)
+					if (agent.IsSessionValid(utcNow))
 					{
-						report.ConformLoop.Add((agent.Id, agent.ConformAttemptCount.Value));
-					}
-					if (agent.UpgradeAttemptCount.HasValue && agent.UpgradeAttemptCount.Value > 3)
-					{
-						report.UpgradeLoop.Add((agent.Id, agent.UpgradeAttemptCount.Value));
+						if (agent.ConformAttemptCount.HasValue && agent.ConformAttemptCount.Value > 3)
+						{
+							report.ConformLoop.Add((agent.Id, agent.ConformAttemptCount.Value));
+						}
+						if (agent.UpgradeAttemptCount.HasValue && agent.UpgradeAttemptCount.Value > 3)
+						{
+							report.UpgradeLoop.Add((agent.Id, agent.UpgradeAttemptCount.Value));
+						}
 					}
 				}
 
