Horde: Add some logging for preflight desc updating
#rnx

diff --git a/Engine/Source/Programs/Horde/Horde.Server/Jobs/JobService.cs b/Engine/Source/Programs/Horde/Horde.Server/Jobs/JobService.cs
--- a/Engine/Source/Programs/Horde/Horde.Server/Jobs/JobService.cs
+++ b/Engine/Source/Programs/Horde/Horde.Server/Jobs/JobService.cs
@@ -1041,6 +1041,7 @@ namespace Horde.Server.Jobs
 								IUserSettings settings = await _userCollection.GetSettingsAsync(job.StartedByUserId.Value);
 								if (settings.AlwaysTagPreflightCL)
 								{
+									_logger.LogInformation("Updating description for {PreflightChange} for {UserId} user settings", job.PreflightChange, job.StartedByUserId.Value);
 									await _perforceService.UpdateChangelistDescriptionAsync(streamConfig.ClusterName, job.PreflightChange, x => x.TrimEnd() + $"\n#preflight {job.Id}");
 								}
 							}
