Horde: Add support for tagging preflights without autosubmit selected (dashboard)
#rnx
#skipci

diff --git a/Engine/Source/Programs/Horde/HordeDashboard/src/backend/Api.ts b/Engine/Source/Programs/Horde/HordeDashboard/src/backend/Api.ts
--- a/Engine/Source/Programs/Horde/HordeDashboard/src/backend/Api.ts
+++ b/Engine/Source/Programs/Horde/HordeDashboard/src/backend/Api.ts
@@ -3206,6 +3206,9 @@ export type GetUserResponse = {
 	/** Whether to enable experimental features for this user */
 	enableExperimentalFeatures?: boolean;
 
+	/** Whether to always tag preflight changelists */
+	alwaysTagPreflightCL?: boolean;
+
 	/**  Settings for the dashboard */
 	dashboardSettings?: DashboardSettings;
 
@@ -3263,6 +3266,9 @@ export type UpdateUserRequest = {
 	/** Whether to enable experimental features for this user */
 	enableExperimentalFeatures?: boolean;
 
+	/** Whether to always tag preflight changelists */
+	alwaysTagPreflightCL?: boolean;
+
 }
 
 export type GetPerforceServerStatusResponse = {
diff --git a/Engine/Source/Programs/Horde/HordeDashboard/src/backend/Dashboard.ts b/Engine/Source/Programs/Horde/HordeDashboard/src/backend/Dashboard.ts
--- a/Engine/Source/Programs/Horde/HordeDashboard/src/backend/Dashboard.ts
+++ b/Engine/Source/Programs/Horde/HordeDashboard/src/backend/Dashboard.ts
@@ -259,7 +259,21 @@ export class Dashboard {
             this.setUpdated();
 
         });
+    }
+
+    get alwaysTagPreflightCL(): boolean {
+
+        return this.data.alwaysTagPreflightCL!;
+    }
+
+    set alwaysTagPreflightCL(value: boolean) {
 
+        backend.updateUser({ alwaysTagPreflightCL: value }).then(() => {
+
+            this.data.alwaysTagPreflightCL = value;
+            this.setUpdated();
+
+        });
     }
 
     get roles(): UserClaim[] {
@@ -631,7 +645,7 @@ export class Dashboard {
 
     serverSettingsChanged: boolean = false;
 
-    private data: GetUserResponse = { id: "", name: "", enableExperimentalFeatures: false, claims: [], pinnedJobIds: [], dashboardSettings: { preferences: new Map() } };
+    private data: GetUserResponse = { id: "", name: "", enableExperimentalFeatures: false, alwaysTagPreflightCL: false, claims: [], pinnedJobIds: [], dashboardSettings: { preferences: new Map() } };
 
     private updateTimeoutId: any = undefined;
 
diff --git a/Engine/Source/Programs/Horde/HordeDashboard/src/components/DashboardView.tsx b/Engine/Source/Programs/Horde/HordeDashboard/src/components/DashboardView.tsx
--- a/Engine/Source/Programs/Horde/HordeDashboard/src/components/DashboardView.tsx
+++ b/Engine/Source/Programs/Horde/HordeDashboard/src/components/DashboardView.tsx
@@ -364,6 +364,9 @@ const GeneralPanel: React.FC = observer(() => {
                         <Toggle label="Show All Preflights" inlineLabel={true} defaultChecked={dashboard.showPreflights} onChange={(ev, checked) => {
                            dashboard.setShowPreflights(checked ? true : false);
                         }} />
+                        <Toggle label="Always Tag Preflight Changelists" inlineLabel={true} defaultChecked={dashboard.alwaysTagPreflightCL} onChange={(ev, checked) => {
+                           dashboard.alwaysTagPreflightCL = checked ? true : false;
+                        }} />
                         <Toggle label="Prefer Compact Views" inlineLabel={true} defaultChecked={dashboard.compactViews} onChange={(ev, checked) => {
                            dashboard.setCompactViews(checked ? true : false);
                         }} />
