Horde: Support multiple SQS queues for receiving AWS lifecycle events

diff --git a/Engine/Source/Programs/Horde/Horde.Server.Tests/Fleet/AwsAutoScalingLifecycleTest.cs b/Engine/Source/Programs/Horde/Horde.Server.Tests/Fleet/AwsAutoScalingLifecycleTest.cs
--- a/Engine/Source/Programs/Horde/Horde.Server.Tests/Fleet/AwsAutoScalingLifecycleTest.cs
+++ b/Engine/Source/Programs/Horde/Horde.Server.Tests/Fleet/AwsAutoScalingLifecycleTest.cs
@@ -31,7 +31,8 @@ public class AwsAutoScalingLifecycleServiceTest : TestSetup
 	private Mock<IAmazonAutoScaling> _asgMock = default!;
 	private readonly List<CompleteLifecycleActionRequest> _lifecycleUpdates = new();
 	private readonly FakeAmazonSqs _fakeSqs = new();
-	private readonly string _queueUrl = "https://sqs.us-east-1.amazonaws.com/123456789/MyQueue";
+	private readonly string _queueUrl1 = "https://sqs.us-east-1.amazonaws.com/123456789/MyQueue-1";
+	private readonly string _queueUrl2 = "https://sqs.us-east-1.amazonaws.com/123456789/MyQueue-2";
 
 	protected override void Dispose(bool disposing)
 	{
@@ -60,7 +61,7 @@ public class AwsAutoScalingLifecycleServiceTest : TestSetup
 			.Setup(x => x.CompleteLifecycleActionAsync(It.IsAny<CompleteLifecycleActionRequest>(), It.IsAny<CancellationToken>()))
 			.Returns(OnCompleteLifecycleActionAsync);
 		
-		ServerSettings ss = new () { AwsAutoScalingQueueUrl = _queueUrl };
+		ServerSettings ss = new () { AwsAutoScalingQueueUrls = new [] { _queueUrl1, _queueUrl2 } };
 		_asgLifecycleService = new AwsAutoScalingLifecycleService(AgentService, GetRedisServiceSingleton(), AgentCollection, Clock, new TestOptionsMonitor<ServerSettings>(ss), ServiceProvider, Tracer, logger);
 		_asgLifecycleService.SetAmazonClientsTesting(_asgMock.Object, _fakeSqs);
 		await _asgLifecycleService.StartAsync(CancellationToken.None);
@@ -75,8 +76,8 @@ public class AwsAutoScalingLifecycleServiceTest : TestSetup
 		await AgentService.DeleteAgentAsync(agent);
 		
 		// Act
-		await _fakeSqs.SendMessageAsync(_queueUrl, JsonSerializer.Serialize(lae));
-		await _asgLifecycleService.ReceiveLifecycleEventsAsync(CancellationToken.None);
+		await _fakeSqs.SendMessageAsync(_queueUrl1, JsonSerializer.Serialize(lae));
+		await _asgLifecycleService.ReceiveLifecycleEventsAsync(_queueUrl1, CancellationToken.None);
 
 		// Assert
 		agent = (await AgentService.GetAgentAsync(agent.Id))!;
diff --git a/Engine/Source/Programs/Horde/Horde.Server/Agents/Fleet/AwsAutoScalingLifecycle.cs b/Engine/Source/Programs/Horde/Horde.Server/Agents/Fleet/AwsAutoScalingLifecycle.cs
--- a/Engine/Source/Programs/Horde/Horde.Server/Agents/Fleet/AwsAutoScalingLifecycle.cs
+++ b/Engine/Source/Programs/Horde/Horde.Server/Agents/Fleet/AwsAutoScalingLifecycle.cs
@@ -71,8 +71,8 @@ public sealed class AwsAutoScalingLifecycleService : IHostedService, IDisposable
 	private readonly Tracer _tracer;
 	private readonly ILogger<AwsAutoScalingLifecycleService> _logger;
 	private readonly ITicker _updateLifecyclesTicker;
-	private readonly BackgroundTask _lifecycleEventListenerTask;
-	private readonly string? _sqsQueueUrl;
+	private readonly List<BackgroundTask> _lifecycleEventListenerTasks = new ();
+	private readonly string[] _sqsQueueUrls;
 
 #pragma warning disable CA2213 // Disposable fields should be disposed
 	private IAmazonAutoScaling? _awsAutoScaling;
@@ -103,8 +103,12 @@ public sealed class AwsAutoScalingLifecycleService : IHostedService, IDisposable
 
 		string tickerName = $"{nameof(AwsAutoScalingLifecycleService)}.{nameof(UpdateLifecyclesAsync)}";
 		_updateLifecyclesTicker = clock.AddSharedTicker(tickerName, LifecycleUpdaterInterval, UpdateLifecyclesAsync, logger);
-		_lifecycleEventListenerTask = new BackgroundTask(ListenForLifecycleEventsAsync);
-		_sqsQueueUrl = settings.CurrentValue.AwsAutoScalingQueueUrl;
+		_sqsQueueUrls = settings.CurrentValue.AwsAutoScalingQueueUrls;
+		
+		foreach (string sqsQueueUrl in _sqsQueueUrls)
+		{
+			_lifecycleEventListenerTasks.Add(new BackgroundTask(ct => ListenForLifecycleEventsAsync(sqsQueueUrl, ct)));
+		}
 	}
 
 	internal void SetAmazonClientsTesting(IAmazonAutoScaling awsAutoScaling, IAmazonSQS awsSqs)
@@ -116,17 +120,22 @@ public sealed class AwsAutoScalingLifecycleService : IHostedService, IDisposable
 	/// <inheritdoc/>
 	public async Task StartAsync(CancellationToken cancellationToken)
 	{
-		if (_sqsQueueUrl != null)
+		if (_sqsQueueUrls.Length > 0)
 		{
 			await _updateLifecyclesTicker.StartAsync();
-			_lifecycleEventListenerTask.Start();	
 		}
+
+		_lifecycleEventListenerTasks.ForEach(x => x.Start());
 	}
 
 	/// <inheritdoc/>
 	public async Task StopAsync(CancellationToken cancellationToken)
 	{
-		await _lifecycleEventListenerTask.StopAsync();
+		foreach (BackgroundTask bgTask in _lifecycleEventListenerTasks)
+		{
+			await bgTask.StopAsync();
+		}
+
 		await _updateLifecyclesTicker.StopAsync();
 	}
 
@@ -134,21 +143,22 @@ public sealed class AwsAutoScalingLifecycleService : IHostedService, IDisposable
 	public void Dispose()
 	{
 		_updateLifecyclesTicker.Dispose();
-		_lifecycleEventListenerTask.Dispose();
+		_lifecycleEventListenerTasks.ForEach(x => x.Dispose());
 	}
 
 	/// <summary>
 	/// Continuously request and receive messages from SQS
 	/// </summary>
+	/// <param name="sqsQueueUrl">SQS queue to fetch messages from</param>
 	/// <param name="cancellationToken">Cancellation token</param>
-	private async Task ListenForLifecycleEventsAsync(CancellationToken cancellationToken)
+	private async Task ListenForLifecycleEventsAsync(string sqsQueueUrl, CancellationToken cancellationToken)
 	{
-		_logger.LogInformation("Listening for lifecycle events...");
+		_logger.LogInformation("Listening for lifecycle events on {SqsQueueUrl}...", sqsQueueUrl);
 		while (!cancellationToken.IsCancellationRequested)
 		{
 			try
 			{
-				await ReceiveLifecycleEventsAsync(cancellationToken);
+				await ReceiveLifecycleEventsAsync(sqsQueueUrl, cancellationToken);
 			}
 			catch (OperationCanceledException) when (cancellationToken.IsCancellationRequested)
 			{
@@ -168,13 +178,13 @@ public sealed class AwsAutoScalingLifecycleService : IHostedService, IDisposable
 	/// Receive any queued messages from SQS.
 	/// This is a one-time operation, if more messages are expected, call this method again.
 	/// </summary>
-	/// <param name="cancellationToken"></param>
-	/// <exception cref="JsonException"></exception>
-	internal async Task ReceiveLifecycleEventsAsync(CancellationToken cancellationToken)
+	/// <param name="sqsQueueUrl">SQS queue to fetch messages from</param>
+	/// <param name="cancellationToken">Cancellation token</param>
+	internal async Task ReceiveLifecycleEventsAsync(string sqsQueueUrl, CancellationToken cancellationToken)
 	{
 		ReceiveMessageRequest request = new ()
 		{
-			QueueUrl = _sqsQueueUrl,
+			QueueUrl = sqsQueueUrl,
 			MaxNumberOfMessages = 10,
 			VisibilityTimeout = 10,
 			WaitTimeSeconds = 10,
@@ -194,21 +204,21 @@ public sealed class AwsAutoScalingLifecycleService : IHostedService, IDisposable
 			}
 			finally
 			{
-				await DeleteMessageAsync(message, cancellationToken);
+				await DeleteMessageAsync(sqsQueueUrl, message, cancellationToken);
 			}
 		}
 	}
 
-	private async Task DeleteMessageAsync(Message message, CancellationToken cancellationToken)
+	private async Task DeleteMessageAsync(string sqsQueueUrl, Message message, CancellationToken cancellationToken)
 	{
 		try
 		{
-			DeleteMessageRequest deleteRequest = new () { QueueUrl = _sqsQueueUrl, ReceiptHandle = message.ReceiptHandle };
+			DeleteMessageRequest deleteRequest = new () { QueueUrl = sqsQueueUrl, ReceiptHandle = message.ReceiptHandle };
 			await GetSqs().DeleteMessageAsync(deleteRequest, cancellationToken);
 		}
 		catch (AmazonServiceException e)
 		{
-			_logger.LogError(e, "Failed to delete message {MessageId} from SQS queue after processing it", message.MessageId);
+			_logger.LogError(e, "Failed to delete message {MessageId} from SQS queue {SqsQueueUrl} after processing it", message.MessageId, sqsQueueUrl);
 		}
 	}
 
@@ -228,6 +238,12 @@ public sealed class AwsAutoScalingLifecycleService : IHostedService, IDisposable
 	/// <returns>True if handled</returns>
 	public async Task<bool> InitiateTerminationAsync(LifecycleActionEvent e, CancellationToken cancellationToken)
 	{
+		using TelemetrySpan span = _tracer.StartActiveSpan($"{nameof(AwsAutoScalingLifecycleService)}.{nameof(InitiateTerminationAsync)}");
+		span.SetAttribute("asgName", e.AutoScalingGroupName);
+		span.SetAttribute("origin", e.Origin);
+		span.SetAttribute("instanceId", e.Ec2InstanceId);
+		span.SetAttribute("lifecycleTransition", e.LifecycleTransition);
+		
 		if (e.Origin == OriginAsg)
 		{
 			string instanceIdProp = KnownPropertyNames.AwsInstanceId + "=" + e.Ec2InstanceId;
@@ -334,7 +350,9 @@ public sealed class AwsAutoScalingLifecycleService : IHostedService, IDisposable
 	}
 	
 	/// <summary>
-	/// Update AWS auto-scaling on the lifecycle status for each agent being tracked for shutdown
+	/// Update AWS auto-scaling on the lifecycle status for each agent being tracked for shutdown.
+	/// This ensures the EC2 instance is not released for termination until it gracefully finished any outstanding work.
+	/// That is, the agent status is set to stopped in Horde server.
 	/// </summary>
 	/// <param name="cancellationToken">Cancellation token for the async task</param>
 	/// <returns>Async task</returns>
@@ -391,6 +409,9 @@ public sealed class AwsAutoScalingLifecycleService : IHostedService, IDisposable
 		await Task.WhenAll(tasks);
 	}
 
+	/// <summary>
+	/// Class for agent lifecycle state. Serialized and stored in Redis.
+	/// </summary>
 	private class AgentLifecycleInfo
 	{
 		public string AgentId { get; private set; }
@@ -470,7 +491,7 @@ public class LifecycleActionEvent
 	[JsonPropertyName("LifecycleActionToken")] public string LifecycleActionToken { get; set; } = "";
 	
 	/// <summary>
-	/// Lifecycle transition
+	/// Lifecycle transition (e.g EC2_INSTANCE_LAUNCHING, EC2_INSTANCE_TERMINATING)
 	/// </summary>
 	[JsonPropertyName("LifecycleTransition")] public string LifecycleTransition { get; set; } = "";
 	
diff --git a/Engine/Source/Programs/Horde/Horde.Server/ServerSettings.cs b/Engine/Source/Programs/Horde/Horde.Server/ServerSettings.cs
--- a/Engine/Source/Programs/Horde/Horde.Server/ServerSettings.cs
+++ b/Engine/Source/Programs/Horde/Horde.Server/ServerSettings.cs
@@ -634,10 +634,10 @@ namespace Horde.Server
 		public string? FleetManagerV2Config { get; set; }
 		
 		/// <summary>
-		/// AWS SQS queue URL where lifecycle events from EC2 auto-scaling are received
+		/// AWS SQS queue URLs where lifecycle events from EC2 auto-scaling are received
 		/// <see cref="AwsAutoScalingLifecycleService" />
 		/// </summary>
-		public string? AwsAutoScalingQueueUrl { get; set; }
+		public string[] AwsAutoScalingQueueUrls { get; set; } = Array.Empty<string>();
 
 		/// <summary>
 		/// Whether to run scheduled jobs.
