Horde: Add markdown description to new build dialog
#skipci
#rnx

diff --git a/Engine/Source/Programs/Horde/HordeDashboard/src/backend/Api.ts b/Engine/Source/Programs/Horde/HordeDashboard/src/backend/Api.ts
--- a/Engine/Source/Programs/Horde/HordeDashboard/src/backend/Api.ts
+++ b/Engine/Source/Programs/Horde/HordeDashboard/src/backend/Api.ts
@@ -2013,6 +2013,9 @@ export type GetTemplateResponseBase = {
 
 	/// List of parameters for this template
 	parameters: ParameterData[];
+
+	/// Description for the template, supports markdown
+	description?: string;
 }
 
 /// Query selecting the base changelist to use
diff --git a/Engine/Source/Programs/Horde/HordeDashboard/src/components/DashboardView.tsx b/Engine/Source/Programs/Horde/HordeDashboard/src/components/DashboardView.tsx
--- a/Engine/Source/Programs/Horde/HordeDashboard/src/components/DashboardView.tsx
+++ b/Engine/Source/Programs/Horde/HordeDashboard/src/components/DashboardView.tsx
@@ -9,6 +9,7 @@ import { useWindowSize } from '../base/utilities/hooks';
 import { hordeClasses, modeColors } from '../styles/Styles';
 import { Breadcrumbs } from './Breadcrumbs';
 import { TopNav } from './TopNav';
+import { getSiteConfig } from '../backend/Config';
 
 
 
@@ -234,7 +235,7 @@ const ColorPreferenceDialog: React.FC<{ shown: boolean, statusIn: StatusColor, o
 const GeneralPanel: React.FC = observer(() => {
 
    let [colorState, setColorState] = useState<{ status?: StatusColor }>({});
-
+   const config = getSiteConfig();
    const defaultStatusColors = dashboard.getDefaultStatusColors();
 
    type GeneralItem = {
@@ -364,9 +365,9 @@ const GeneralPanel: React.FC = observer(() => {
                         <Toggle label="Show All Preflights" inlineLabel={true} defaultChecked={dashboard.showPreflights} onChange={(ev, checked) => {
                            dashboard.setShowPreflights(checked ? true : false);
                         }} />
-                        <Toggle label="Always Tag Preflight Changelists" inlineLabel={true} defaultChecked={dashboard.alwaysTagPreflightCL} onChange={(ev, checked) => {
+                        {config.environment !== "production" && <Toggle label="Always Tag Preflight Changelists" inlineLabel={true} defaultChecked={dashboard.alwaysTagPreflightCL} onChange={(ev, checked) => {
                            dashboard.alwaysTagPreflightCL = checked ? true : false;
-                        }} />
+                        }} />}
                         <Toggle label="Prefer Compact Views" inlineLabel={true} defaultChecked={dashboard.compactViews} onChange={(ev, checked) => {
                            dashboard.setCompactViews(checked ? true : false);
                         }} />
diff --git a/Engine/Source/Programs/Horde/HordeDashboard/src/components/NewBuild.tsx b/Engine/Source/Programs/Horde/HordeDashboard/src/components/NewBuild.tsx
--- a/Engine/Source/Programs/Horde/HordeDashboard/src/components/NewBuild.tsx
+++ b/Engine/Source/Programs/Horde/HordeDashboard/src/components/NewBuild.tsx
@@ -13,6 +13,7 @@ import { hordeClasses, modeColors } from '../styles/Styles';
 import { useQuery } from './JobDetailCommon';
 import { JobDetailsV2 } from './jobDetailsV2/JobDetailsViewCommon';
 import dashboard from '../backend/Dashboard';
+import { Markdown } from '../base/components/Markdown';
 
 let toolTipId = 0;
 
@@ -802,7 +803,7 @@ export const NewBuild: React.FC<{ streamId: string; show: boolean; onClose: (new
                   reason: `${errorReason}`,
                   title: `Preflight Template Error`,
                   message: `There was an issue with the specified preflight template.\n\nReason: ${errorReason}\n\nTime: ${moment.utc().format("MMM Do, HH:mm z")}`
-               }, true);   
+               }, true);
             }
          }
 
@@ -1652,6 +1653,14 @@ export const NewBuild: React.FC<{ streamId: string; show: boolean; onClose: (new
 
                      </Stack>
                   </Stack>
+                  {!!template?.description && <Stack>
+                     <Label style={{ padding: 0, paddingBottom: 4 }}>{`Description:`}</Label>
+                     <Stack style={{ border: "1px solid #605e5c", width: 767 }}>
+                        <Stack style={{padding: "8px 12px"}}>
+                           <Markdown styles={{ root: { maxHeight: 240, overflow: "auto", th: { fontSize: 12 } } }}>{template.description}</Markdown>
+                        </Stack>
+                     </Stack>
+                  </Stack>}
 
                   <Stack>
 
