Horde: Remove inconsistent color from description
#skipci
#rnx

diff --git a/Engine/Source/Programs/Horde/HordeDashboard/src/components/NewBuild.tsx b/Engine/Source/Programs/Horde/HordeDashboard/src/components/NewBuild.tsx
--- a/Engine/Source/Programs/Horde/HordeDashboard/src/components/NewBuild.tsx
+++ b/Engine/Source/Programs/Horde/HordeDashboard/src/components/NewBuild.tsx
@@ -1654,7 +1654,7 @@ export const NewBuild: React.FC<{ streamId: string; show: boolean; onClose: (new
                      </Stack>
                   </Stack>
                   {!!template?.description && <Stack>
-                     <Label style={{ padding: 0, paddingBottom: 4 }}>{`Description:`}</Label>
+                     <Label style={{ padding: 0, paddingBottom: 4 }}>{`Description`}</Label>
                      <Stack style={{ border: "1px solid #605e5c", width: 767 }}>
                         <Stack style={{padding: "8px 12px"}}>
                            <Markdown styles={{ root: { maxHeight: 240, overflow: "auto", th: { fontSize: 12 } } }}>{template.description}</Markdown>
