Horde: Check for AWS EC2 interruptions and terminate the agent appropriately

In case of a spot interruption, there's a hard deadline of 2 minutes that must be adhered.

diff --git a/Engine/Source/Programs/Horde/Horde.Agent.Tests/Services/AwsInstanceLifecycleServiceTests.cs b/Engine/Source/Programs/Horde/Horde.Agent.Tests/Services/AwsInstanceLifecycleServiceTests.cs
new file mode 100644
--- /dev/null
+++ b/Engine/Source/Programs/Horde/Horde.Agent.Tests/Services/AwsInstanceLifecycleServiceTests.cs
@@ -0,0 +1,124 @@
+// Copyright Epic Games, Inc. All Rights Reserved.
+
+using System.Net;
+using System.Net.Http;
+using System.Threading;
+using System.Threading.Tasks;
+using Horde.Agent.Services;
+using Microsoft.Extensions.Logging;
+using Microsoft.VisualStudio.TestTools.UnitTesting;
+using Moq;
+using Moq.Protected;
+
+namespace Horde.Agent.Tests.Services;
+
+/// <summary>
+/// Fake implementation of the AWS EC2 metadata server (IMDS)
+/// </summary>
+public class FakeAwsImds
+{
+	public const string BaseUrl = "http://169.254.169.254/latest/meta-data";
+	public const string SpotInstanceData = "{\"action\": \"terminate\", \"time\": \"2017-09-18T08:22:00Z\"}";
+	public const string OnDemand = "on-demand";
+	public const string Spot = "spot";
+	
+	private readonly Mock<HttpMessageHandler> _mock;
+	public string InstanceLifeCycle { get; set; } = OnDemand;
+	public string? TargetLifecycleState { get; set; }
+	public string? SpotInstanceAction { get; set; }
+	
+	public FakeAwsImds()
+	{
+		_mock = new (MockBehavior.Strict);
+		_mock
+			.Protected()
+			.Setup<Task<HttpResponseMessage>>(
+				"SendAsync",
+				ItExpr.IsAny<HttpRequestMessage>(),
+				ItExpr.IsAny<CancellationToken>()
+			)
+			.Returns(SendAsync);
+		
+		// Ignore Dispose() calls
+		_mock.Protected().Setup("Dispose", ItExpr.IsAny<bool>());
+	}
+
+	public HttpClient GetHttpClient()
+	{
+		return new (_mock.Object);
+	}
+
+	private Task<HttpResponseMessage> SendAsync(HttpRequestMessage request, CancellationToken cancellationToken)
+	{
+		switch (request.RequestUri!.ToString())
+		{
+			case $"{BaseUrl}/autoscaling/target-lifecycle-state":
+				return Task.FromResult(CreateResponse(HttpStatusCode.OK, TargetLifecycleState));
+
+			case $"{BaseUrl}/spot/instance-action":
+				return Task.FromResult(CreateResponse(HttpStatusCode.OK, SpotInstanceAction));
+
+			case $"{BaseUrl}/instance-life-cycle":
+				return Task.FromResult(CreateResponse(HttpStatusCode.OK, InstanceLifeCycle));
+
+			default:
+				return Task.FromResult(CreateResponse(HttpStatusCode.InternalServerError, "Invalid state"));
+		}
+	}
+	
+	private static HttpResponseMessage CreateResponse(HttpStatusCode status, string? data)
+	{
+		return data != null 
+			? new HttpResponseMessage { StatusCode = status, Content = new StringContent(data) }
+			: new HttpResponseMessage { StatusCode = HttpStatusCode.NotFound, Content = new StringContent("Not found") };
+	}
+}
+
+[TestClass]
+public sealed class AwsInstanceLifecycleServiceTests : System.IDisposable
+{
+	private readonly LoggerFactory _loggerFactory = new ();
+	private readonly HttpClient _httpClient;
+	private readonly FakeAwsImds _fakeImds = new ();
+	private readonly AwsInstanceLifecycleService _service;
+	private Ec2InstanceState? _terminationState;
+	private bool? _terminationIsSpot;
+	
+	public AwsInstanceLifecycleServiceTests()
+	{
+		_httpClient = _fakeImds.GetHttpClient();
+		_service = new AwsInstanceLifecycleService(_httpClient, null!, _loggerFactory.CreateLogger<AwsInstanceLifecycleService>());
+		_service._terminationCallback = (state, isSpot, _) =>
+		{
+			_terminationState = state;
+			_terminationIsSpot = isSpot;
+			return Task.CompletedTask;
+		};
+	}
+
+	[TestMethod]
+	public async Task Terminate_Asg_CallbackHasCorrectParameters()
+	{
+		_fakeImds.TargetLifecycleState = "Terminated";
+		await _service.MonitorInstanceLifecycle(CancellationToken.None);
+		Assert.AreEqual(Ec2InstanceState.TerminatingAsg, _terminationState);
+		Assert.IsFalse(_terminationIsSpot);
+	}
+	
+	[TestMethod]
+	public async Task Terminate_Spot_CallbackHasCorrectParameters()
+	{
+		_fakeImds.SpotInstanceAction = FakeAwsImds.SpotInstanceData;
+		_fakeImds.InstanceLifeCycle = FakeAwsImds.Spot;
+		await _service.MonitorInstanceLifecycle(CancellationToken.None);
+		Assert.AreEqual(Ec2InstanceState.TerminatingSpot, _terminationState);
+		Assert.IsTrue(_terminationIsSpot);
+	}
+
+	public void Dispose()
+	{
+		_loggerFactory.Dispose();
+		_httpClient.Dispose();
+		_service.Dispose();
+	}
+}
\ No newline at end of file
diff --git a/Engine/Source/Programs/Horde/Horde.Agent/Program.cs b/Engine/Source/Programs/Horde/Horde.Agent/Program.cs
--- a/Engine/Source/Programs/Horde/Horde.Agent/Program.cs
+++ b/Engine/Source/Programs/Horde/Horde.Agent/Program.cs
@@ -185,6 +185,12 @@ namespace Horde.Agent
 				{
 					return builder.WaitAndRetryAsync(new[] { TimeSpan.FromSeconds(1), TimeSpan.FromSeconds(5), TimeSpan.FromSeconds(10) });
 				});
+			services.AddHttpClient(AwsInstanceLifecycleService.HttpClientName);
+			services.AddSingleton<AwsInstanceLifecycleService>();
+			if (settings.EnableAwsEc2Support)
+			{
+				services.AddHostedService(sp => sp.GetRequiredService<AwsInstanceLifecycleService>());
+			}
 
 			services.AddSingleton<GrpcService>();
 			services.AddSingleton<TelemetryService>();
diff --git a/Engine/Source/Programs/Horde/Horde.Agent/Services/AwsInstanceLifecycleService.cs b/Engine/Source/Programs/Horde/Horde.Agent/Services/AwsInstanceLifecycleService.cs
new file mode 100644
--- /dev/null
+++ b/Engine/Source/Programs/Horde/Horde.Agent/Services/AwsInstanceLifecycleService.cs
@@ -0,0 +1,133 @@
+// Copyright Epic Games, Inc. All Rights Reserved.
+
+using System;
+using System.Net;
+using System.Net.Http;
+using System.Threading;
+using System.Threading.Tasks;
+using Horde.Agent.Utility;
+using Microsoft.Extensions.Hosting;
+using Microsoft.Extensions.Logging;
+using Microsoft.Extensions.Options;
+
+namespace Horde.Agent.Services;
+
+/// <summary>
+/// Describes state of local EC2 instance the agent is running on
+/// </summary>
+public enum Ec2InstanceState
+{
+	/// <summary>
+	/// Normal state
+	/// </summary>
+	InService,
+	
+	/// <summary>
+	/// Termination is caused by a spot interruption.
+	/// </summary>
+	TerminatingSpot,
+	
+	/// <summary>
+	/// Termination is caused by the auto-scaling group
+	/// Can be caused by capacity re-balancing, scale-ins etc.
+	/// </summary>
+	TerminatingAsg
+}
+
+/// <summary>
+/// Monitors the local EC2 instance lifecycle state. In particular, auto-scaling group and spot instance events.
+/// </summary>
+class AwsInstanceLifecycleService : BackgroundService
+{
+	/// <summary>
+	/// Name of the HTTP client used for requests to IMDS
+	/// </summary>
+	public const string HttpClientName = "Horde.HttpAwsInstanceClient";
+
+	private const string BaseUri = "http://169.254.169.254/latest/meta-data";
+	private readonly ILogger<AwsInstanceLifecycleService> _logger;
+	private readonly HttpClient _httpClient;
+	private readonly TimeSpan _pollInterval = TimeSpan.FromSeconds(5);
+	internal Func<Ec2InstanceState, bool, CancellationToken, Task> _terminationCallback;
+
+	/// <summary>
+	/// Constructor
+	/// </summary>
+	public AwsInstanceLifecycleService(HttpClient httpClient, IOptions<AgentSettings> settings, ILogger<AwsInstanceLifecycleService> logger)
+	{
+		_httpClient = httpClient;
+		_httpClient.Timeout = TimeSpan.FromSeconds(2);
+		_terminationCallback = HandleTermination;
+		_logger = logger;
+	}
+
+	private async Task<Ec2InstanceState> GetStateAsync(CancellationToken cancellationToken)
+	{
+		HttpResponseMessage spotRes = await _httpClient.GetAsync(new Uri(BaseUri + "/spot/instance-action"), cancellationToken);
+		if (spotRes.StatusCode == HttpStatusCode.OK)
+		{
+			return Ec2InstanceState.TerminatingSpot;
+		}
+
+		HttpResponseMessage asgRes = await _httpClient.GetAsync(new Uri(BaseUri + "/autoscaling/target-lifecycle-state"), cancellationToken);
+		if (asgRes.StatusCode == HttpStatusCode.OK)
+		{
+			string state = await asgRes.Content.ReadAsStringAsync(cancellationToken);
+			if (state == "Terminated")
+			{
+				return Ec2InstanceState.TerminatingAsg;
+			}
+		}
+
+		return Ec2InstanceState.InService;
+	}
+	
+	private async Task<bool> IsSpotInstanceAsync(CancellationToken cancellationToken)
+	{
+		HttpResponseMessage res = await _httpClient.GetAsync(new Uri(BaseUri + "/instance-life-cycle"), cancellationToken);
+		if (res.StatusCode == HttpStatusCode.OK)
+		{
+			string state = await res.Content.ReadAsStringAsync(cancellationToken);
+			return state == "spot";
+		}
+
+		return false;
+	}
+	
+	internal async Task MonitorInstanceLifecycle(CancellationToken cancellationToken)
+	{
+		_logger.LogInformation("Monitoring EC2 instance lifecycle state...");
+		while (!cancellationToken.IsCancellationRequested)
+		{
+			try
+			{
+				Ec2InstanceState state = await GetStateAsync(cancellationToken);
+				if (state != Ec2InstanceState.InService)
+				{
+					bool isSpot = await IsSpotInstanceAsync(cancellationToken);
+					_logger.LogInformation("EC2 instance is terminating. IsSpot={IsSpot} Reason={InstanceState}", isSpot, state);
+					await _terminationCallback(state, isSpot, cancellationToken);
+					return;
+				}
+			}
+			catch (Exception e)
+			{
+				_logger.LogError(e, "Unhandled exception during EC2 instance monitoring");
+			}
+			
+			await Task.Delay(_pollInterval, cancellationToken);
+		}
+	}
+
+	private Task HandleTermination(Ec2InstanceState state, bool isSpot, CancellationToken cancellationToken)
+	{
+		_logger.LogInformation("Shutting down");
+		return Shutdown.ExecuteAsync(false, _logger, cancellationToken);
+	}
+	
+	protected override async Task ExecuteAsync(CancellationToken stoppingToken)
+	{
+		await MonitorInstanceLifecycle(stoppingToken);
+	}
+}
+
