Horde: Only shut down agent if it's a spot instance

diff --git a/Engine/Source/Programs/Horde/Horde.Agent/Services/AwsInstanceLifecycleService.cs b/Engine/Source/Programs/Horde/Horde.Agent/Services/AwsInstanceLifecycleService.cs
--- a/Engine/Source/Programs/Horde/Horde.Agent/Services/AwsInstanceLifecycleService.cs
+++ b/Engine/Source/Programs/Horde/Horde.Agent/Services/AwsInstanceLifecycleService.cs
@@ -121,10 +121,15 @@ class AwsInstanceLifecycleService : BackgroundService
 
 	private Task HandleTermination(Ec2InstanceState state, bool isSpot, CancellationToken cancellationToken)
 	{
-		_logger.LogInformation("Shutting down");
-		return Shutdown.ExecuteAsync(false, _logger, cancellationToken);
+		if (isSpot)
+		{
+			_logger.LogInformation("Shutting down");
+			return Shutdown.ExecuteAsync(false, _logger, cancellationToken);
+		}
+		
+		return Task.CompletedTask;
 	}
-	
+
 	protected override async Task ExecuteAsync(CancellationToken stoppingToken)
 	{
 		await MonitorInstanceLifecycle(stoppingToken);
