Horde: Only enable EC2 lifecycle monitoring if metadata server is available

diff --git a/Engine/Source/Programs/Horde/Horde.Agent.Tests/Services/AwsInstanceLifecycleServiceTests.cs b/Engine/Source/Programs/Horde/Horde.Agent.Tests/Services/AwsInstanceLifecycleServiceTests.cs
--- a/Engine/Source/Programs/Horde/Horde.Agent.Tests/Services/AwsInstanceLifecycleServiceTests.cs
+++ b/Engine/Source/Programs/Horde/Horde.Agent.Tests/Services/AwsInstanceLifecycleServiceTests.cs
@@ -52,17 +52,11 @@ public class FakeAwsImds
 	{
 		switch (request.RequestUri!.ToString())
 		{
-			case $"{BaseUrl}/autoscaling/target-lifecycle-state":
-				return Task.FromResult(CreateResponse(HttpStatusCode.OK, TargetLifecycleState));
-
-			case $"{BaseUrl}/spot/instance-action":
-				return Task.FromResult(CreateResponse(HttpStatusCode.OK, SpotInstanceAction));
-
-			case $"{BaseUrl}/instance-life-cycle":
-				return Task.FromResult(CreateResponse(HttpStatusCode.OK, InstanceLifeCycle));
-
-			default:
-				return Task.FromResult(CreateResponse(HttpStatusCode.InternalServerError, "Invalid state"));
+			case $"{BaseUrl}/": return Task.FromResult(CreateResponse(HttpStatusCode.OK, "(list of sub-paths)"));
+			case $"{BaseUrl}/autoscaling/target-lifecycle-state": return Task.FromResult(CreateResponse(HttpStatusCode.OK, TargetLifecycleState));
+			case $"{BaseUrl}/spot/instance-action": return Task.FromResult(CreateResponse(HttpStatusCode.OK, SpotInstanceAction));
+			case $"{BaseUrl}/instance-life-cycle": return Task.FromResult(CreateResponse(HttpStatusCode.OK, InstanceLifeCycle));
+			default: return Task.FromResult(CreateResponse(HttpStatusCode.InternalServerError, "Invalid state"));
 		}
 	}
 	
@@ -100,7 +94,7 @@ public sealed class AwsInstanceLifecycleServiceTests : System.IDisposable
 	public async Task Terminate_Asg_CallbackHasCorrectParameters()
 	{
 		_fakeImds.TargetLifecycleState = "Terminated";
-		await _service.MonitorInstanceLifecycle(CancellationToken.None);
+		await _service.MonitorInstanceLifecycleAsync(CancellationToken.None);
 		Assert.AreEqual(Ec2InstanceState.TerminatingAsg, _terminationState);
 		Assert.IsFalse(_terminationIsSpot);
 	}
@@ -110,11 +104,11 @@ public sealed class AwsInstanceLifecycleServiceTests : System.IDisposable
 	{
 		_fakeImds.SpotInstanceAction = FakeAwsImds.SpotInstanceData;
 		_fakeImds.InstanceLifeCycle = FakeAwsImds.Spot;
-		await _service.MonitorInstanceLifecycle(CancellationToken.None);
+		await _service.MonitorInstanceLifecycleAsync(CancellationToken.None);
 		Assert.AreEqual(Ec2InstanceState.TerminatingSpot, _terminationState);
 		Assert.IsTrue(_terminationIsSpot);
 	}
-
+	
 	public void Dispose()
 	{
 		_loggerFactory.Dispose();
diff --git a/Engine/Source/Programs/Horde/Horde.Agent/Services/AwsInstanceLifecycleService.cs b/Engine/Source/Programs/Horde/Horde.Agent/Services/AwsInstanceLifecycleService.cs
--- a/Engine/Source/Programs/Horde/Horde.Agent/Services/AwsInstanceLifecycleService.cs
+++ b/Engine/Source/Programs/Horde/Horde.Agent/Services/AwsInstanceLifecycleService.cs
@@ -94,8 +94,28 @@ class AwsInstanceLifecycleService : BackgroundService
 		return false;
 	}
 	
-	internal async Task MonitorInstanceLifecycle(CancellationToken cancellationToken)
+	private async Task<bool> IsImdsAvailableAsync(CancellationToken cancellationToken)
 	{
+		try
+		{
+			HttpResponseMessage res = await _httpClient.GetAsync(new Uri(BaseUri + "/"), cancellationToken);
+			return res.StatusCode == HttpStatusCode.OK;
+		}
+		catch (Exception)
+		{
+			// Timed out or other error. Can safely assume the metadata server is not available.
+			return false;
+		}
+	}
+	
+	internal async Task MonitorInstanceLifecycleAsync(CancellationToken cancellationToken)
+	{
+		if (!await IsImdsAvailableAsync(cancellationToken))
+		{
+			_logger.LogInformation("EC2 metadata server (IMDS) not available. Will not monitor EC2 lifecycle state");
+			return;
+		}
+		
 		_logger.LogInformation("Monitoring EC2 instance lifecycle state...");
 		while (!cancellationToken.IsCancellationRequested)
 		{
@@ -132,7 +152,7 @@ class AwsInstanceLifecycleService : BackgroundService
 
 	protected override async Task ExecuteAsync(CancellationToken stoppingToken)
 	{
-		await MonitorInstanceLifecycle(stoppingToken);
+		await MonitorInstanceLifecycleAsync(stoppingToken);
 	}
 }
 
