Horde: Fix duplicate "issue not resolved" notifications. Using the first failure time can cause inconsistencies when issues are merged, causing the base time for notifications to change. Always use the issue creation time instead.

diff --git a/Engine/Source/Programs/Horde/Horde.Server/Notifications/Sinks/SlackNotificationSink.cs b/Engine/Source/Programs/Horde/Horde.Server/Notifications/Sinks/SlackNotificationSink.cs
--- a/Engine/Source/Programs/Horde/Horde.Server/Notifications/Sinks/SlackNotificationSink.cs
+++ b/Engine/Source/Programs/Horde/Horde.Server/Notifications/Sinks/SlackNotificationSink.cs
@@ -1158,7 +1158,7 @@ namespace Horde.Server.Notifications.Sinks
 
 			if (workflow.EscalateAlias != null && workflow.EscalateTimes.Count > 0)
 			{
-				DateTime escalateTime = _clock.UtcNow + TimeSpan.FromMinutes(workflow.EscalateTimes[0]);
+				DateTime escalateTime = issue.CreatedAt.AddMinutes(workflow.EscalateTimes[0]);
 				if (await _redisService.GetDatabase().SortedSetAddAsync(_escalateIssues, issue.Id, (escalateTime - DateTime.UnixEpoch).TotalSeconds, StackExchange.Redis.When.NotExists))
 				{
 					_logger.LogInformation("First escalation time for issue {IssueId} is {Time}", issue.Id, escalateTime);
@@ -2636,7 +2636,7 @@ namespace Horde.Server.Notifications.Sinks
 				await _slackClient.PostMessageToThreadAsync(state.MessageId, $"{FormatUserOrGroupMention(workflow.EscalateAlias)} - Issue <{issueUrl}|{issue.Id}> has not been resolved after {openTimeStr}.");
 			}
 
-			DateTime nextEscalationTime = span.FirstFailure.StepTime;
+			DateTime nextEscalationTime = issue.CreatedAt;
 			for (int idx = 0;;idx++)
 			{
 				if (idx >= workflow.EscalateTimes.Count)
@@ -2649,7 +2649,7 @@ namespace Horde.Server.Notifications.Sinks
 				DateTime prevEscalationTime = nextEscalationTime;
 				nextEscalationTime += TimeSpan.FromMinutes(workflow.EscalateTimes[idx]);
 
-				if (nextEscalationTime > utcNow)
+				if (nextEscalationTime > utcNow + TimeSpan.FromMinutes(5.0))
 				{
 					_logger.LogInformation("Next escalation time for issue {IssueId} is {NextEscalationTime} ({PrevEscalationTime} + {Time}m [{Idx}/{Count}])", issue.Id, nextEscalationTime, prevEscalationTime, workflow.EscalateTimes[idx], idx + 1, workflow.EscalateTimes.Count);
 					break;
