Horde: Create a PerforceLogger when ManagedWorkspace is used in WorkspaceExecutor

diff --git a/Engine/Source/Programs/Horde/Horde.Agent/Execution/ManagedWorkspaceMaterializer.cs b/Engine/Source/Programs/Horde/Horde.Agent/Execution/ManagedWorkspaceMaterializer.cs
--- a/Engine/Source/Programs/Horde/Horde.Agent/Execution/ManagedWorkspaceMaterializer.cs
+++ b/Engine/Source/Programs/Horde/Horde.Agent/Execution/ManagedWorkspaceMaterializer.cs
@@ -122,6 +122,16 @@ public class ManagedWorkspaceMaterializer : IWorkspaceMaterializer
 
 		await _workspace.SyncAsync(changeNum, preflightChangeNum, cacheFile, cancellationToken);
 	}
+	
+	/// <summary>
+	/// Get info for Perforce workspace
+	/// </summary>
+	/// <returns>Workspace info</returns>
+	public WorkspaceInfo? GetWorkspaceInfo()
+	{
+		return _workspace;
+	}
+	
 
 	private IScope CreateTraceSpan(string operationName)
 	{
diff --git a/Engine/Source/Programs/Horde/Horde.Agent/Execution/PerforceExecutor.cs b/Engine/Source/Programs/Horde/Horde.Agent/Execution/PerforceExecutor.cs
--- a/Engine/Source/Programs/Horde/Horde.Agent/Execution/PerforceExecutor.cs
+++ b/Engine/Source/Programs/Horde/Horde.Agent/Execution/PerforceExecutor.cs
@@ -211,27 +211,27 @@ namespace Horde.Agent.Execution
 			}
 		}
 
-		PerforceLogger CreatePerforceLogger(ILogger logger)
+		public static PerforceLogger CreatePerforceLogger(ILogger logger, int changeNum, WorkspaceInfo workspace, WorkspaceInfo? autoSdkWorkspace)
 		{
 			PerforceLogger perforceLogger = new PerforceLogger(logger);
-			perforceLogger.AddClientView(_workspace.WorkspaceDir, _workspace.StreamView, _batch.Change);
-			if (_autoSdkWorkspace != null)
+			perforceLogger.AddClientView(workspace.WorkspaceDir, workspace.StreamView, changeNum);
+			if (autoSdkWorkspace != null)
 			{
-				perforceLogger.AddClientView(_autoSdkWorkspace.WorkspaceDir, _autoSdkWorkspace.StreamView, _batch.Change);
+				perforceLogger.AddClientView(autoSdkWorkspace.WorkspaceDir, autoSdkWorkspace.StreamView, changeNum);
 			}
 			return perforceLogger;
 		}
 
 		protected override async Task<bool> SetupAsync(BeginStepResponse step, ILogger logger, CancellationToken cancellationToken)
 		{
-			PerforceLogger perforceLogger = CreatePerforceLogger(logger);
+			PerforceLogger perforceLogger = CreatePerforceLogger(logger, _batch.Change, _workspace, _autoSdkWorkspace);
 			bool useP4 = WorkspaceInfo.ShouldUseHaveTable(_workspaceInfo.Method);
 			return await SetupAsync(step, _workspace.WorkspaceDir, _sharedStorageDir, useP4, perforceLogger, cancellationToken);
 		}
 
 		protected override async Task<bool> ExecuteAsync(BeginStepResponse step, ILogger logger, CancellationToken cancellationToken)
 		{
-			PerforceLogger perforceLogger = CreatePerforceLogger(logger);
+			PerforceLogger perforceLogger = CreatePerforceLogger(logger, _batch.Change, _workspace, _autoSdkWorkspace);
 			bool useP4 = WorkspaceInfo.ShouldUseHaveTable(_workspaceInfo.Method);
 			return await ExecuteAsync(step, _workspace.WorkspaceDir, _sharedStorageDir, useP4, perforceLogger, cancellationToken);
 		}
diff --git a/Engine/Source/Programs/Horde/Horde.Agent/Execution/WorkspaceExecutor.cs b/Engine/Source/Programs/Horde/Horde.Agent/Execution/WorkspaceExecutor.cs
--- a/Engine/Source/Programs/Horde/Horde.Agent/Execution/WorkspaceExecutor.cs
+++ b/Engine/Source/Programs/Horde/Horde.Agent/Execution/WorkspaceExecutor.cs
@@ -5,6 +5,7 @@ using System.Text.RegularExpressions;
 using System.Threading;
 using System.Threading.Tasks;
 using EpicGames.Core;
+using Horde.Agent.Utility;
 using HordeCommon.Rpc;
 using HordeCommon.Rpc.Messages;
 using Microsoft.Extensions.Logging;
@@ -102,7 +103,7 @@ namespace Horde.Agent.Execution
 			// Loop back to JobExecutor's SetupAsync again, but with workspace and shared storage dir set
 			WorkspaceMaterializerSettings settings = await _workspace.GetSettingsAsync(cancellationToken);
 			DirectoryReference workspaceDir = settings.DirectoryPath;
-			return await SetupAsync(step, workspaceDir, _sharedStorageDir, settings.IsPerforceWorkspace, logger, cancellationToken);
+			return await SetupAsync(step, workspaceDir, _sharedStorageDir, settings.IsPerforceWorkspace, GetLogger(settings, logger), cancellationToken);
 		}
 
 		/// <inheritdoc/>
@@ -111,7 +112,7 @@ namespace Horde.Agent.Execution
 			// Loop back to JobExecutor's ExecuteAsync again, but with workspace and shared storage dir set
 			WorkspaceMaterializerSettings settings = await _workspace.GetSettingsAsync(cancellationToken);
 			DirectoryReference workspaceDir = settings.DirectoryPath;
-			return await ExecuteAsync(step, workspaceDir, _sharedStorageDir, settings.IsPerforceWorkspace, logger, cancellationToken);
+			return await ExecuteAsync(step, workspaceDir, _sharedStorageDir, settings.IsPerforceWorkspace, GetLogger(settings, logger), cancellationToken);
 		}
 
 		/// <inheritdoc/>
@@ -128,6 +129,25 @@ namespace Horde.Agent.Execution
 			
 			await _workspace.FinalizeAsync(cancellationToken);
 		}
+
+		private ILogger GetLogger(WorkspaceMaterializerSettings settings, ILogger logger)
+		{
+			if (settings.IsPerforceWorkspace)
+			{
+				// Try resolve a PerforceLogger using assumptions about the materializer.
+				// This is to remain compatible with PerforceExecutor.
+				// These Perforce-specific references should ideally not exist in WorkspaceExecutor.
+				WorkspaceInfo? workspaceInfo = (_workspace as ManagedWorkspaceMaterializer)?.GetWorkspaceInfo();
+				WorkspaceInfo? autoSdkWorkspaceInfo = (_autoSdkWorkspace as ManagedWorkspaceMaterializer)?.GetWorkspaceInfo();
+				
+				if (workspaceInfo != null)
+				{
+					return PerforceExecutor.CreatePerforceLogger(logger, _batch.Change, workspaceInfo, autoSdkWorkspaceInfo);
+				}
+			}
+
+			return logger;
+		}
 	}
 
 	class WorkspaceExecutorFactory : IJobExecutorFactory
diff --git a/Engine/Source/Programs/Horde/Horde.Agent/Utility/WorkspaceInfo.cs b/Engine/Source/Programs/Horde/Horde.Agent/Utility/WorkspaceInfo.cs
--- a/Engine/Source/Programs/Horde/Horde.Agent/Utility/WorkspaceInfo.cs
+++ b/Engine/Source/Programs/Horde/Horde.Agent/Utility/WorkspaceInfo.cs
@@ -20,7 +20,7 @@ namespace Horde.Agent.Utility
 	/// <summary>
 	/// Stores information about a managed Perforce workspace
 	/// </summary>
-	class WorkspaceInfo
+	public class WorkspaceInfo
 	{
 		/// <summary>
 		/// The perforce connection
