Horde: Implement termination signal file for agent workloads

Once an impending termination of the agent machine is detected, this file is written to disk to let any workload get an early warning of what is about to happen.
A workload can detect and read this file supplied via an environment variable and take action.

diff --git a/Engine/Source/Programs/Horde/Horde.Agent.Tests/Services/AwsInstanceLifecycleServiceTests.cs b/Engine/Source/Programs/Horde/Horde.Agent.Tests/Services/AwsInstanceLifecycleServiceTests.cs
--- a/Engine/Source/Programs/Horde/Horde.Agent.Tests/Services/AwsInstanceLifecycleServiceTests.cs
+++ b/Engine/Source/Programs/Horde/Horde.Agent.Tests/Services/AwsInstanceLifecycleServiceTests.cs
@@ -1,11 +1,15 @@
 // Copyright Epic Games, Inc. All Rights Reserved.
 
+using System;
+using System.IO;
 using System.Net;
 using System.Net.Http;
 using System.Threading;
 using System.Threading.Tasks;
+using EpicGames.Core;
 using Horde.Agent.Services;
 using Microsoft.Extensions.Logging;
+using Microsoft.Extensions.Options;
 using Microsoft.VisualStudio.TestTools.UnitTesting;
 using Moq;
 using Moq.Protected;
@@ -75,13 +79,29 @@ public sealed class AwsInstanceLifecycleServiceTests : System.IDisposable
 	private readonly HttpClient _httpClient;
 	private readonly FakeAwsImds _fakeImds = new ();
 	private readonly AwsInstanceLifecycleService _service;
+	private readonly FileReference _terminationSignalFile;
+	private TimeSpan? _terminationTtl;
 	private Ec2InstanceState? _terminationState;
 	private bool? _terminationIsSpot;
 	
 	public AwsInstanceLifecycleServiceTests()
 	{
+		DirectoryInfo tempDir = Directory.CreateDirectory(Path.Combine(Path.GetTempPath(), "horde-agent-test-" + Path.GetRandomFileName()));
+		AgentSettings settings = new () { WorkingDir = tempDir.FullName };
+		_terminationSignalFile = settings.GetTerminationSignalFile();
+
 		_httpClient = _fakeImds.GetHttpClient();
-		_service = new AwsInstanceLifecycleService(_httpClient, null!, _loggerFactory.CreateLogger<AwsInstanceLifecycleService>());
+		_service = new AwsInstanceLifecycleService(_httpClient, new OptionsWrapper<AgentSettings>(settings), _loggerFactory.CreateLogger<AwsInstanceLifecycleService>());
+		_service._timeToLiveAsg = TimeSpan.FromMilliseconds(10);
+		_service._timeToLiveSpot = TimeSpan.FromMilliseconds(20);
+		_service._terminationBufferTime = TimeSpan.FromMilliseconds(2);
+
+		AwsInstanceLifecycleService.TerminationWarningDelegate origWarningCallback = _service._terminationWarningCallback;
+		_service._terminationWarningCallback = (state, isSpot, timeToLive, ct) =>
+		{
+			_terminationTtl = timeToLive;
+			return origWarningCallback(state, isSpot, timeToLive, ct);
+		};
 		_service._terminationCallback = (state, isSpot, _) =>
 		{
 			_terminationState = state;
@@ -97,6 +117,7 @@ public sealed class AwsInstanceLifecycleServiceTests : System.IDisposable
 		await _service.MonitorInstanceLifecycleAsync(CancellationToken.None);
 		Assert.AreEqual(Ec2InstanceState.TerminatingAsg, _terminationState);
 		Assert.IsFalse(_terminationIsSpot);
+		Assert.AreEqual(8, _terminationTtl!.Value.TotalMilliseconds); // 10 ms for ASG, minus 2 ms for termination buffer
 	}
 	
 	[TestMethod]
@@ -107,6 +128,20 @@ public sealed class AwsInstanceLifecycleServiceTests : System.IDisposable
 		await _service.MonitorInstanceLifecycleAsync(CancellationToken.None);
 		Assert.AreEqual(Ec2InstanceState.TerminatingSpot, _terminationState);
 		Assert.IsTrue(_terminationIsSpot);
+		Assert.AreEqual(18, _terminationTtl!.Value.TotalMilliseconds); // 20 ms for spot, minus 2 ms for termination buffer
+	}
+	
+	[TestMethod]
+	public async Task Terminate_Spot_WritesSignalFile()
+	{
+		_fakeImds.SpotInstanceAction = FakeAwsImds.SpotInstanceData;
+		_fakeImds.InstanceLifeCycle = FakeAwsImds.Spot;
+		Assert.IsFalse(File.Exists(_terminationSignalFile.FullName));
+		
+		await _service.MonitorInstanceLifecycleAsync(CancellationToken.None);
+
+		string data = await File.ReadAllTextAsync(_terminationSignalFile.FullName);
+		Assert.AreEqual("v1\t18", data); // 20 ms for spot, minus 2 ms for termination buffer
 	}
 	
 	public void Dispose()
diff --git a/Engine/Source/Programs/Horde/Horde.Agent/AgentSettings.cs b/Engine/Source/Programs/Horde/Horde.Agent/AgentSettings.cs
--- a/Engine/Source/Programs/Horde/Horde.Agent/AgentSettings.cs
+++ b/Engine/Source/Programs/Horde/Horde.Agent/AgentSettings.cs
@@ -3,6 +3,7 @@
 using System;
 using System.Collections.Generic;
 using System.ComponentModel.DataAnnotations;
+using System.IO;
 using System.Linq;
 using EpicGames.Core;
 using EpicGames.Horde.Storage;
@@ -336,6 +337,15 @@ namespace Horde.Agent
 			return processesToTerminate;
 		}
 
+		/// <summary>
+		/// Path to file used for signaling impending termination and shutdown of the agent
+		/// </summary>
+		/// <returns>Path to file which may or may not exist</returns>
+		public FileReference GetTerminationSignalFile()
+		{
+			return new FileReference(Path.Combine(WorkingDir ?? Path.GetTempPath(), ".horde-termination-signal"));
+		}
+
 		internal string GetAgentName()
 		{
 			return Name ?? Environment.MachineName;
diff --git a/Engine/Source/Programs/Horde/Horde.Agent/Leases/Handlers/ComputeHandler.cs b/Engine/Source/Programs/Horde/Horde.Agent/Leases/Handlers/ComputeHandler.cs
--- a/Engine/Source/Programs/Horde/Horde.Agent/Leases/Handlers/ComputeHandler.cs
+++ b/Engine/Source/Programs/Horde/Horde.Agent/Leases/Handlers/ComputeHandler.cs
@@ -4,6 +4,7 @@ using System;
 using System.Buffers;
 using System.Collections.Generic;
 using System.Diagnostics;
+using System.IO;
 using System.Net.Sockets;
 using System.Threading;
 using System.Threading.Tasks;
@@ -92,6 +93,7 @@ namespace Horde.Agent.Leases.Handlers
 			}
 
 			logger.LogInformation("Starting compute task (lease {LeaseId}). Waiting for connection with nonce {Nonce}...", leaseId, StringUtils.FormatHexString(computeTask.Nonce.Span));
+			ClearTerminationSignalFile();
 
 			TcpClient? tcpClient = null;
 			try
@@ -125,6 +127,7 @@ namespace Horde.Agent.Leases.Handlers
 
 								Dictionary<string, string?> newEnvVars = new Dictionary<string, string?>();
 								newEnvVars["UE_HORDE_SHARED_DIR"] = sharedDir.FullName;
+								newEnvVars["UE_HORDE_TERMINATION_SIGNAL_FILE"] = _settings.GetTerminationSignalFile().FullName;
 
 								AgentMessageHandler worker = new AgentMessageHandler(sandboxDir, _storageCache, newEnvVars, false, _settings.WineExecutablePath, logger);
 								await worker.RunAsync(socket, cts.Token);
@@ -155,6 +158,21 @@ namespace Horde.Agent.Leases.Handlers
 			}
 		}
 
+		private void ClearTerminationSignalFile()
+		{
+			string path = _settings.GetTerminationSignalFile().FullName;
+			try
+			{
+				File.Delete(path);
+			}
+			catch (Exception e)
+			{
+				// If this file is not removed and lingers on from previous executions,
+				// new compute tasks may pick it up and erroneously decide to terminate.
+				_logger.LogError(e, "Unable to delete termination signal file {Path}", path);
+			}
+		}
+
 		static async Task TickTimeoutAsync(TcpTransportWithTimeout transport, CancellationTokenSource cts, ILogger logger, CancellationToken cancellationToken)
 		{
 			while(!cancellationToken.IsCancellationRequested)
diff --git a/Engine/Source/Programs/Horde/Horde.Agent/Services/AwsInstanceLifecycleService.cs b/Engine/Source/Programs/Horde/Horde.Agent/Services/AwsInstanceLifecycleService.cs
--- a/Engine/Source/Programs/Horde/Horde.Agent/Services/AwsInstanceLifecycleService.cs
+++ b/Engine/Source/Programs/Horde/Horde.Agent/Services/AwsInstanceLifecycleService.cs
@@ -1,10 +1,12 @@
 // Copyright Epic Games, Inc. All Rights Reserved.
 
 using System;
+using System.IO;
 using System.Net;
 using System.Net.Http;
 using System.Threading;
 using System.Threading.Tasks;
+using EpicGames.Core;
 using Horde.Agent.Utility;
 using Microsoft.Extensions.Hosting;
 using Microsoft.Extensions.Logging;
@@ -48,7 +50,30 @@ class AwsInstanceLifecycleService : BackgroundService
 	private readonly ILogger<AwsInstanceLifecycleService> _logger;
 	private readonly HttpClient _httpClient;
 	private readonly TimeSpan _pollInterval = TimeSpan.FromSeconds(5);
-	internal Func<Ec2InstanceState, bool, CancellationToken, Task> _terminationCallback;
+	private readonly FileReference _terminationSignalFile;
+	
+	internal delegate Task TerminationWarningDelegate(Ec2InstanceState state, bool isSpot, TimeSpan timeToLive, CancellationToken cancellationToken);
+	internal delegate Task TerminationDelegate(Ec2InstanceState state, bool isSpot, CancellationToken cancellationToken);
+	internal TerminationWarningDelegate _terminationWarningCallback;
+	internal TerminationDelegate _terminationCallback;
+
+	/// <summary>
+	/// Time to live for EC2 instance once a termination is detected coming from the auto-scaling group (ASG)
+	/// In practice, this is dictated by the lifecycle hook set for the ASG.
+	/// Set to 120 sec to mimic the TTL for spot interruption, leading to similar handling of both for now.
+	/// </summary>
+	internal TimeSpan _timeToLiveAsg = TimeSpan.FromSeconds(120); 
+	
+	/// <summary>
+	/// Time to live for EC2 instance once a spot interruption is detected. Strictly defined by AWS EC2.
+	/// </summary>
+	internal TimeSpan _timeToLiveSpot = TimeSpan.FromSeconds(120); // Strictly defined by AWS EC2
+	
+	/// <summary>
+	/// Duration of the time-to-live to allocate towards shutting down the Horde agent and the machine itself.
+	/// Example: if TTL is 120 seconds, 90 seconds will be reported in the termination warning.
+	/// </summary>
+	internal TimeSpan _terminationBufferTime = TimeSpan.FromSeconds(30);
 
 	/// <summary>
 	/// Constructor
@@ -57,8 +82,10 @@ class AwsInstanceLifecycleService : BackgroundService
 	{
 		_httpClient = httpClient;
 		_httpClient.Timeout = TimeSpan.FromSeconds(2);
-		_terminationCallback = HandleTermination;
+		_terminationWarningCallback = OnTerminationWarningAsync;
+		_terminationCallback = OnTerminationAsync;
 		_logger = logger;
+		_terminationSignalFile = settings.Value.GetTerminationSignalFile();
 	}
 
 	private async Task<Ec2InstanceState> GetStateAsync(CancellationToken cancellationToken)
@@ -125,7 +152,14 @@ class AwsInstanceLifecycleService : BackgroundService
 				if (state != Ec2InstanceState.InService)
 				{
 					bool isSpot = await IsSpotInstanceAsync(cancellationToken);
-					_logger.LogInformation("EC2 instance is terminating. IsSpot={IsSpot} Reason={InstanceState}", isSpot, state);
+					TimeSpan ttl = GetTimeToLive(state);
+					_logger.LogInformation("EC2 instance is terminating. IsSpot={IsSpot} Reason={InstanceState} TimeToLive={Ttk} ms", isSpot, state, ttl.TotalMilliseconds);
+
+					ttl -= _terminationBufferTime;
+					ttl = ttl.Ticks >= 0 ? ttl : TimeSpan.Zero; 
+					
+					await _terminationWarningCallback(state, isSpot, ttl, cancellationToken);
+					await Task.Delay(ttl, cancellationToken);
 					await _terminationCallback(state, isSpot, cancellationToken);
 					return;
 				}
@@ -139,7 +173,30 @@ class AwsInstanceLifecycleService : BackgroundService
 		}
 	}
 
-	private Task HandleTermination(Ec2InstanceState state, bool isSpot, CancellationToken cancellationToken)
+	/// <summary>
+	/// Determine time to live for the current EC2 instance once a terminating state has been detected
+	/// </summary>
+	/// <param name="state">Current state</param>
+	/// <returns>Time to live</returns>
+	/// <exception cref="ArgumentException"></exception>
+	private TimeSpan GetTimeToLive(Ec2InstanceState state)
+	{
+		return state switch
+		{
+			Ec2InstanceState.TerminatingAsg => _timeToLiveAsg, 
+			Ec2InstanceState.TerminatingSpot => _timeToLiveSpot,
+			_ => throw new ArgumentException($"Invalid state {state}")
+		};
+	}
+
+	private async Task OnTerminationWarningAsync(Ec2InstanceState state, bool isSpot, TimeSpan timeToLive, CancellationToken cancellationToken)
+	{
+		// Create and write the termination signal file, containing the time-to-live for the EC2 instance.
+		// Workloads executed by the agent that support this protocol can pick this up and prepare/clean up prior to termination
+		await WriteTerminationSignalFileAsync(timeToLive, cancellationToken);
+	}
+	
+	private Task OnTerminationAsync(Ec2InstanceState state, bool isSpot, CancellationToken cancellationToken)
 	{
 		if (isSpot)
 		{
@@ -150,6 +207,13 @@ class AwsInstanceLifecycleService : BackgroundService
 		return Task.CompletedTask;
 	}
 
+	private Task WriteTerminationSignalFileAsync(TimeSpan timeToLive, CancellationToken cancellationToken)
+	{
+		string contents = $"v1\t{timeToLive.TotalMilliseconds}";
+		return File.WriteAllTextAsync(_terminationSignalFile.FullName, contents, cancellationToken);
+	}
+
+	/// <inheritdoc/>
 	protected override async Task ExecuteAsync(CancellationToken stoppingToken)
 	{
 		await MonitorInstanceLifecycleAsync(stoppingToken);
