Horde: Attempt to set better policies on HTTP clients used for storage.

- Maximum limit of 16 connections per server.
- Now has a global timeout of 240 seconds.
- Each individual request times out after 30 seconds.
- Added logging for transient http errors.

diff --git a/Engine/Source/Programs/Horde/Horde.Agent/Program.cs b/Engine/Source/Programs/Horde/Horde.Agent/Program.cs
--- a/Engine/Source/Programs/Horde/Horde.Agent/Program.cs
+++ b/Engine/Source/Programs/Horde/Horde.Agent/Program.cs
@@ -31,6 +31,7 @@ using Microsoft.Extensions.DependencyInjection;
 using Microsoft.Extensions.Logging;
 using OpenTracing.Util;
 using Polly;
+using Polly.Extensions.Http;
 
 namespace Horde.Agent
 {
@@ -180,11 +181,31 @@ namespace Horde.Agent
 				return builder.WaitAndRetryAsync(new[] { TimeSpan.FromSeconds(1), TimeSpan.FromSeconds(5), TimeSpan.FromSeconds(10) });
 			});
 
-			services.AddHttpClient(HttpStorageClient.HttpClientName)
-				.AddTransientHttpErrorPolicy(builder =>
+			services.AddHttpClient(HttpStorageClient.HttpClientName, builder => 
 				{
-					return builder.WaitAndRetryAsync(new[] { TimeSpan.FromSeconds(1), TimeSpan.FromSeconds(5), TimeSpan.FromSeconds(10) });
-				});
+					builder.Timeout = TimeSpan.FromSeconds(240); // Global timeout
+				})
+				.ConfigurePrimaryHttpMessageHandler(() => new SocketsHttpHandler
+					{
+						MaxConnectionsPerServer = 16,
+						PooledConnectionIdleTimeout = TimeSpan.FromMinutes(15),
+					})
+				.AddPolicyHandler((serviceProvider, request) => Policy.TimeoutAsync<HttpResponseMessage>(30,
+					(outcome, timespan, context) => { 
+						serviceProvider.GetRequiredService<ILogger<HttpServerStorageFactory>>().LogWarning("Http request timed out after {Time}s.", (int)timespan.TotalSeconds); 
+						return Task.CompletedTask; 
+					}))
+				.AddPolicyHandler((serviceProvider, request) => HttpPolicyExtensions.HandleTransientHttpError()
+					.WaitAndRetryAsync(new[]
+						{
+							TimeSpan.FromSeconds(1),
+							TimeSpan.FromSeconds(5),
+							TimeSpan.FromSeconds(10),
+							TimeSpan.FromSeconds(30),
+							TimeSpan.FromSeconds(30),
+						},
+						(outcome, timespan, retryAttempt, context) => serviceProvider.GetRequiredService<ILogger<HttpServerStorageFactory>>().LogWarning("Http request failed. Delaying for {DelayMs}ms (attempt #{RetryNum}).", timespan.TotalMilliseconds, retryAttempt)
+					));
 			services.AddHttpClient(AwsInstanceLifecycleService.HttpClientName);
 			services.AddSingleton<AwsInstanceLifecycleService>();
 			if (settings.EnableAwsEc2Support)
