Horde: Change format of termination signal file

diff --git a/Engine/Source/Programs/Horde/Horde.Agent.Tests/Services/AwsInstanceLifecycleServiceTests.cs b/Engine/Source/Programs/Horde/Horde.Agent.Tests/Services/AwsInstanceLifecycleServiceTests.cs
--- a/Engine/Source/Programs/Horde/Horde.Agent.Tests/Services/AwsInstanceLifecycleServiceTests.cs
+++ b/Engine/Source/Programs/Horde/Horde.Agent.Tests/Services/AwsInstanceLifecycleServiceTests.cs
@@ -80,9 +80,7 @@ public sealed class AwsInstanceLifecycleServiceTests : System.IDisposable
 	private readonly FakeAwsImds _fakeImds = new ();
 	private readonly AwsInstanceLifecycleService _service;
 	private readonly FileReference _terminationSignalFile;
-	private TimeSpan? _terminationTtl;
-	private Ec2InstanceState? _terminationState;
-	private bool? _terminationIsSpot;
+	private Ec2TerminationInfo? _info;
 	
 	public AwsInstanceLifecycleServiceTests()
 	{
@@ -97,15 +95,15 @@ public sealed class AwsInstanceLifecycleServiceTests : System.IDisposable
 		_service._terminationBufferTime = TimeSpan.FromMilliseconds(2);
 
 		AwsInstanceLifecycleService.TerminationWarningDelegate origWarningCallback = _service._terminationWarningCallback;
-		_service._terminationWarningCallback = (state, isSpot, timeToLive, ct) =>
+		_service._terminationWarningCallback = (info, ct) =>
 		{
-			_terminationTtl = timeToLive;
-			return origWarningCallback(state, isSpot, timeToLive, ct);
+			_info = info;
+			Ec2TerminationInfo newInfo = new (_info.State, _info.IsSpot, _info.TimeToLive, DateTime.UnixEpoch + TimeSpan.FromSeconds(2222), _info.Reason);
+			return origWarningCallback(newInfo, ct);
 		};
-		_service._terminationCallback = (state, isSpot, _) =>
+		_service._terminationCallback = (info, _) =>
 		{
-			_terminationState = state;
-			_terminationIsSpot = isSpot;
+			_info = info;
 			return Task.CompletedTask;
 		};
 	}
@@ -115,9 +113,9 @@ public sealed class AwsInstanceLifecycleServiceTests : System.IDisposable
 	{
 		_fakeImds.TargetLifecycleState = "Terminated";
 		await _service.MonitorInstanceLifecycleAsync(CancellationToken.None);
-		Assert.AreEqual(Ec2InstanceState.TerminatingAsg, _terminationState);
-		Assert.IsFalse(_terminationIsSpot);
-		Assert.AreEqual(8, _terminationTtl!.Value.TotalMilliseconds); // 10 ms for ASG, minus 2 ms for termination buffer
+		Assert.AreEqual(Ec2InstanceState.TerminatingAsg, _info!.State);
+		Assert.IsFalse(_info!.IsSpot);
+		Assert.AreEqual(8, _info!.TimeToLive.TotalMilliseconds); // 10 ms for ASG, minus 2 ms for termination buffer
 	}
 	
 	[TestMethod]
@@ -126,9 +124,9 @@ public sealed class AwsInstanceLifecycleServiceTests : System.IDisposable
 		_fakeImds.SpotInstanceAction = FakeAwsImds.SpotInstanceData;
 		_fakeImds.InstanceLifeCycle = FakeAwsImds.Spot;
 		await _service.MonitorInstanceLifecycleAsync(CancellationToken.None);
-		Assert.AreEqual(Ec2InstanceState.TerminatingSpot, _terminationState);
-		Assert.IsTrue(_terminationIsSpot);
-		Assert.AreEqual(18, _terminationTtl!.Value.TotalMilliseconds); // 20 ms for spot, minus 2 ms for termination buffer
+		Assert.AreEqual(Ec2InstanceState.TerminatingSpot, _info!.State);
+		Assert.IsTrue(_info!.IsSpot);
+		Assert.AreEqual(18, _info!.TimeToLive.TotalMilliseconds); // 20 ms for spot, minus 2 ms for termination buffer
 	}
 	
 	[TestMethod]
@@ -141,7 +139,7 @@ public sealed class AwsInstanceLifecycleServiceTests : System.IDisposable
 		await _service.MonitorInstanceLifecycleAsync(CancellationToken.None);
 
 		string data = await File.ReadAllTextAsync(_terminationSignalFile.FullName);
-		Assert.AreEqual("v1\t18", data); // 20 ms for spot, minus 2 ms for termination buffer
+		Assert.AreEqual("v1\n18\n2222000\nAWS EC2 Spot interruption\n", data); // 20 ms for spot, minus 2 ms for termination buffer
 	}
 	
 	public void Dispose()
diff --git a/Engine/Source/Programs/Horde/Horde.Agent/Services/AwsInstanceLifecycleService.cs b/Engine/Source/Programs/Horde/Horde.Agent/Services/AwsInstanceLifecycleService.cs
--- a/Engine/Source/Programs/Horde/Horde.Agent/Services/AwsInstanceLifecycleService.cs
+++ b/Engine/Source/Programs/Horde/Horde.Agent/Services/AwsInstanceLifecycleService.cs
@@ -4,6 +4,7 @@ using System;
 using System.IO;
 using System.Net;
 using System.Net.Http;
+using System.Text;
 using System.Threading;
 using System.Threading.Tasks;
 using EpicGames.Core;
@@ -36,6 +37,24 @@ public enum Ec2InstanceState
 	TerminatingAsg
 }
 
+internal class Ec2TerminationInfo
+{
+	public Ec2InstanceState State { get; init; }
+	public bool IsSpot { get; init; }
+	public TimeSpan TimeToLive { get; init; }
+	public DateTime TerminateAt { get; init; }
+	public string Reason { get; init; }
+
+	public Ec2TerminationInfo(Ec2InstanceState state, bool isSpot, TimeSpan timeToLive, DateTime terminateAt, string reason)
+	{
+		State = state;
+		IsSpot = isSpot;
+		TimeToLive = timeToLive;
+		TerminateAt = terminateAt;
+		Reason = reason;
+	}
+}
+
 /// <summary>
 /// Monitors the local EC2 instance lifecycle state. In particular, auto-scaling group and spot instance events.
 /// </summary>
@@ -52,8 +71,8 @@ class AwsInstanceLifecycleService : BackgroundService
 	private readonly TimeSpan _pollInterval = TimeSpan.FromSeconds(5);
 	private readonly FileReference _terminationSignalFile;
 	
-	internal delegate Task TerminationWarningDelegate(Ec2InstanceState state, bool isSpot, TimeSpan timeToLive, CancellationToken cancellationToken);
-	internal delegate Task TerminationDelegate(Ec2InstanceState state, bool isSpot, CancellationToken cancellationToken);
+	internal delegate Task TerminationWarningDelegate(Ec2TerminationInfo info, CancellationToken cancellationToken);
+	internal delegate Task TerminationDelegate(Ec2TerminationInfo info, CancellationToken cancellationToken);
 	internal TerminationWarningDelegate _terminationWarningCallback;
 	internal TerminationDelegate _terminationCallback;
 
@@ -156,11 +175,13 @@ class AwsInstanceLifecycleService : BackgroundService
 					_logger.LogInformation("EC2 instance is terminating. IsSpot={IsSpot} Reason={InstanceState} TimeToLive={Ttk} ms", isSpot, state, ttl.TotalMilliseconds);
 
 					ttl -= _terminationBufferTime;
-					ttl = ttl.Ticks >= 0 ? ttl : TimeSpan.Zero; 
-					
-					await _terminationWarningCallback(state, isSpot, ttl, cancellationToken);
+					ttl = ttl.Ticks >= 0 ? ttl : TimeSpan.Zero;
+					DateTime terminateAt = DateTime.UtcNow + ttl;
+					Ec2TerminationInfo info = new (state, isSpot, ttl, terminateAt, GetReason(state));
+
+					await _terminationWarningCallback(info, cancellationToken);
 					await Task.Delay(ttl, cancellationToken);
-					await _terminationCallback(state, isSpot, cancellationToken);
+					await _terminationCallback(info, cancellationToken);
 					return;
 				}
 			}
@@ -188,17 +209,27 @@ class AwsInstanceLifecycleService : BackgroundService
 			_ => throw new ArgumentException($"Invalid state {state}")
 		};
 	}
+	
+	private static string GetReason(Ec2InstanceState state)
+	{
+		return state switch
+		{
+			Ec2InstanceState.TerminatingAsg => "AWS EC2 ASG termination", 
+			Ec2InstanceState.TerminatingSpot => "AWS EC2 Spot interruption",
+			_ => throw new ArgumentException($"Invalid state {state}")
+		};
+	}
 
-	private async Task OnTerminationWarningAsync(Ec2InstanceState state, bool isSpot, TimeSpan timeToLive, CancellationToken cancellationToken)
+	private async Task OnTerminationWarningAsync(Ec2TerminationInfo info, CancellationToken cancellationToken)
 	{
 		// Create and write the termination signal file, containing the time-to-live for the EC2 instance.
 		// Workloads executed by the agent that support this protocol can pick this up and prepare/clean up prior to termination
-		await WriteTerminationSignalFileAsync(timeToLive, cancellationToken);
+		await WriteTerminationSignalFileAsync(info, cancellationToken);
 	}
 	
-	private Task OnTerminationAsync(Ec2InstanceState state, bool isSpot, CancellationToken cancellationToken)
+	private Task OnTerminationAsync(Ec2TerminationInfo info, CancellationToken cancellationToken)
 	{
-		if (isSpot)
+		if (info.IsSpot)
 		{
 			_logger.LogInformation("Shutting down");
 			return Shutdown.ExecuteAsync(false, _logger, cancellationToken);
@@ -207,10 +238,14 @@ class AwsInstanceLifecycleService : BackgroundService
 		return Task.CompletedTask;
 	}
 
-	private Task WriteTerminationSignalFileAsync(TimeSpan timeToLive, CancellationToken cancellationToken)
+	private Task WriteTerminationSignalFileAsync(Ec2TerminationInfo info, CancellationToken cancellationToken)
 	{
-		string contents = $"v1\t{timeToLive.TotalMilliseconds}";
-		return File.WriteAllTextAsync(_terminationSignalFile.FullName, contents, cancellationToken);
+		StringBuilder sb = new(100);
+		sb.Append("v1\n");
+		sb.Append($"{info.TimeToLive.TotalMilliseconds}\n");
+		sb.Append($"{new DateTimeOffset(info.TerminateAt).ToUnixTimeMilliseconds()}\n");
+		sb.Append($"{info.Reason}\n");
+		return File.WriteAllTextAsync(_terminationSignalFile.FullName, sb.ToString(), cancellationToken);
 	}
 
 	/// <inheritdoc/>
