Horde: Log information about each replicated CL snapshot.

diff --git a/Engine/Source/Programs/Horde/Horde.Server/Perforce/PerforceReplicator.cs b/Engine/Source/Programs/Horde/Horde.Server/Perforce/PerforceReplicator.cs
--- a/Engine/Source/Programs/Horde/Horde.Server/Perforce/PerforceReplicator.cs
+++ b/Engine/Source/Programs/Horde/Horde.Server/Perforce/PerforceReplicator.cs
@@ -543,7 +543,10 @@ namespace Horde.Server.Perforce
 			ChangeRecord changeRecord = await perforce.GetChangeAsync(GetChangeOptions.None, change, cancellationToken);
 			DirectoryNodeRef rootRef = new DirectoryNodeRef(root.Length, await writer.WriteNodeAsync(root, cancellationToken));
 			CommitNode commitNode = new CommitNode(change, parentRef, changeRecord.User ?? "Unknown", changeRecord.Description ?? String.Empty, changeRecord.Date, rootRef);
-			await store.WriteRefAsync(refName, commitNode, refOptions: options.RefOptions, cancellationToken: cancellationToken); 
+			BlobHandle commitHandle = await store.WriteRefAsync(refName, commitNode, refOptions: options.RefOptions, cancellationToken: cancellationToken);
+
+			// Log the snapshot info
+			_logger.LogInformation("Snapshot for {StreamId} CL {Change} is ref {RefName} (commit: {CommitHandle}, root: {RootHandle})", streamConfig.Id, change, refName, ((BundleNodeHandle)commitHandle).GetLocator(), ((BundleNodeHandle)rootRef.Handle).GetLocator());
 		}
 
 		static int GetFileOffset(Utf8String path)
