Horde: Allow extracting bundles from named node locators.

diff --git a/Engine/Source/Programs/Horde/Horde/Commands/Bundles/BundleExtract.cs b/Engine/Source/Programs/Horde/Horde/Commands/Bundles/BundleExtract.cs
--- a/Engine/Source/Programs/Horde/Horde/Commands/Bundles/BundleExtract.cs
+++ b/Engine/Source/Programs/Horde/Horde/Commands/Bundles/BundleExtract.cs
@@ -19,6 +19,9 @@ namespace Horde.Commands.Bundles
 		[CommandLine("-Ref=")]
 		public string? Ref { get; set; }
 
+		[CommandLine("-Node=")]
+		public string? Node { get; set; }
+
 		[CommandLine("-Stats")]
 		public bool Stats { get; set; }
 
@@ -44,6 +47,11 @@ namespace Horde.Commands.Bundles
 				store = await CreateStorageClientAsync(logger);
 				handle = await store.ReadRefTargetAsync(new RefName(Ref));
 			}
+			else if (Node != null)
+			{
+				store = await CreateStorageClientAsync(logger);
+				handle = ((BundleStorageClient)store).CreateNodeHandle(BundleNodeLocator.Parse(Node));
+			}
 			else
 			{
 				throw new CommandLineArgumentException("Either -File=... or -Ref=... must be specified");
