Horde: Fix upgrade attempt count not being reset after a machine starts a new session.

diff --git a/Engine/Source/Programs/Horde/Horde.Server/Agents/AgentCollection.cs b/Engine/Source/Programs/Horde/Horde.Server/Agents/AgentCollection.cs
--- a/Engine/Source/Programs/Horde/Horde.Server/Agents/AgentCollection.cs
+++ b/Engine/Source/Programs/Horde/Horde.Server/Agents/AgentCollection.cs
@@ -472,7 +472,7 @@ namespace Horde.Server.Agents
 			updates.Add(updateBuilder.Unset(x => x.RequestShutdown));
 			updates.Add(updateBuilder.Set(x => x.LastShutdownReason, "Unexpected"));
 
-			if (String.Equals(agent.Version, agent.LastUpgradeVersion, StringComparison.Ordinal))
+			if (String.Equals(version, agent.LastUpgradeVersion, StringComparison.Ordinal))
 			{
 				updates.Add(updateBuilder.Unset(x => x.UpgradeAttemptCount));
 			}
@@ -541,7 +541,7 @@ namespace Horde.Server.Agents
 				{
 					string newVersion = upgradeTask.SoftwareId;
 
-					int versionIdx = newVersion.IndexOf(':');
+					int versionIdx = newVersion.IndexOf(':', StringComparison.Ordinal);
 					if (versionIdx != -1)
 					{
 						newVersion = newVersion.Substring(versionIdx + 1);
