Horde: Prevent output of malformed quoted error lines when message text is empty.

diff --git a/Engine/Source/Programs/Horde/Horde.Server/Notifications/Sinks/SlackNotificationSink.cs b/Engine/Source/Programs/Horde/Horde.Server/Notifications/Sinks/SlackNotificationSink.cs
--- a/Engine/Source/Programs/Horde/Horde.Server/Notifications/Sinks/SlackNotificationSink.cs
+++ b/Engine/Source/Programs/Horde/Horde.Server/Notifications/Sinks/SlackNotificationSink.cs
@@ -637,7 +637,10 @@ namespace Horde.Server.Notifications.Sinks
 					attachment.AddSection(errorSummary);
 					foreach (ILogEventData error in errors.Take(MaxJobStepEvents))
 					{
-						attachment.AddSection(QuoteText(error.Message));
+						if (!String.IsNullOrWhiteSpace(error.Message))
+						{
+							attachment.AddSection(QuoteText(error.Message));
+						}
 					}
 				}
 				else if (warnings.Any())
@@ -646,7 +649,10 @@ namespace Horde.Server.Notifications.Sinks
 					eventStrings.Add(warningSummary);
 					foreach (ILogEventData warning in warnings.Take(MaxJobStepEvents))
 					{
-						attachment.AddSection(QuoteText(warning.Message));
+						if (!String.IsNullOrWhiteSpace(error.Message))
+						{
+							attachment.AddSection(QuoteText(warning.Message));
+						}
 					}
 				}
 
