Horde: Fix compilation error

diff --git a/Engine/Source/Programs/Horde/Horde.Server/Notifications/Sinks/SlackNotificationSink.cs b/Engine/Source/Programs/Horde/Horde.Server/Notifications/Sinks/SlackNotificationSink.cs
--- a/Engine/Source/Programs/Horde/Horde.Server/Notifications/Sinks/SlackNotificationSink.cs
+++ b/Engine/Source/Programs/Horde/Horde.Server/Notifications/Sinks/SlackNotificationSink.cs
@@ -649,7 +649,7 @@ namespace Horde.Server.Notifications.Sinks
 					eventStrings.Add(warningSummary);
 					foreach (ILogEventData warning in warnings.Take(MaxJobStepEvents))
 					{
-						if (!String.IsNullOrWhiteSpace(error.Message))
+						if (!String.IsNullOrWhiteSpace(warning.Message))
 						{
 							attachment.AddSection(QuoteText(warning.Message));
 						}
