Horde: Emit per-pool lease utilization as a AWS CloudWatch metric

diff --git a/Engine/Source/Programs/Horde/Horde.Server.Tests/Fleet/FakeAmazonCloudWatch.cs b/Engine/Source/Programs/Horde/Horde.Server.Tests/Fleet/FakeAmazonCloudWatch.cs
new file mode 100644
--- /dev/null
+++ b/Engine/Source/Programs/Horde/Horde.Server.Tests/Fleet/FakeAmazonCloudWatch.cs
@@ -0,0 +1,52 @@
+// Copyright Epic Games, Inc. All Rights Reserved.
+
+using System.Collections.Generic;
+using System.Net;
+using System.Threading;
+using System.Threading.Tasks;
+using Amazon.CloudWatch;
+using Amazon.CloudWatch.Model;
+using Moq;
+
+namespace Horde.Server.Tests.Fleet;
+
+/// <summary>
+/// Fake implementation of the IAmazonCloudWatch interface
+/// Simplifies testing of external CloudWatch API calls.
+/// Uses a mock to avoid implementing the entire interface. Even if shimmed, it becomes many lines of no-op code.
+/// </summary>
+public class FakeAmazonCloudWatch
+{
+	private readonly Mock<IAmazonCloudWatch> _mock;
+	private readonly Dictionary<string, List<MetricDatum>> _namespaceToMetricData = new();
+
+	public FakeAmazonCloudWatch()
+	{
+		_mock = new Mock<IAmazonCloudWatch>(MockBehavior.Strict);
+		_mock
+			.Setup(x => x.PutMetricDataAsync(It.IsAny<PutMetricDataRequest>(), It.IsAny<CancellationToken>()))
+			.Returns(PutMetricDataAsync);
+	}
+
+	public IAmazonCloudWatch Get()
+	{
+		return _mock.Object;
+	}
+
+	public List<MetricDatum> GetMetricData(string ns)
+	{
+		return _namespaceToMetricData[ns];
+	}
+
+	private Task<PutMetricDataResponse> PutMetricDataAsync(PutMetricDataRequest request, CancellationToken cancellationToken = default)
+	{
+		if (!_namespaceToMetricData.TryGetValue(request.Namespace, out List<MetricDatum>? datums))
+		{
+			datums = new List<MetricDatum>();
+			_namespaceToMetricData[request.Namespace] = datums;
+		}
+
+		datums.AddRange(request.MetricData);
+		return Task.FromResult(new PutMetricDataResponse { HttpStatusCode = HttpStatusCode.OK });
+	}
+}
\ No newline at end of file
diff --git a/Engine/Source/Programs/Horde/Horde.Server.Tests/Fleet/FleetServiceTest.cs b/Engine/Source/Programs/Horde/Horde.Server.Tests/Fleet/FleetServiceTest.cs
--- a/Engine/Source/Programs/Horde/Horde.Server.Tests/Fleet/FleetServiceTest.cs
+++ b/Engine/Source/Programs/Horde/Horde.Server.Tests/Fleet/FleetServiceTest.cs
@@ -232,7 +232,7 @@ namespace Horde.Server.Tests.Fleet
 
 			FleetService service = new(
 				AgentCollection, GraphCollection, JobCollection, LeaseCollection, PoolCollection, new DowntimeServiceStub(isDowntimeActive), StreamCollection, Meter,
-				new StubFleetManagerFactory(fleetManager), Clock, Cache, serverSettingsOpt, GlobalConfig, Tracer, loggerFactory.CreateLogger<FleetService>());
+				new StubFleetManagerFactory(fleetManager), Clock, Cache, serverSettingsOpt, GlobalConfig, ServiceProvider, Tracer, loggerFactory.CreateLogger<FleetService>());
 				
 			return service;
 		}
@@ -283,6 +283,16 @@ namespace Horde.Server.Tests.Fleet
 			Assert.AreEqual(30, ((LeaseUtilizationStrategy)s).Settings.NumSamplesForResult);
 		}
 		
+		[TestMethod]
+		public async Task CreateLeaseUtilizationAwsMetricStrategy()
+		{
+			string config = "{\"SamplePeriodSec\": 123, \"CloudWatchNamespace\": \"myNs\"}";
+			IPoolSizeStrategy s = await CreateStrategy(new PoolSizeStrategyInfo(PoolSizeStrategy.LeaseUtilizationAwsMetric, null, config));
+			Assert.AreEqual(typeof(LeaseUtilizationAwsMetricStrategy), s.GetType());
+			Assert.AreEqual(123, ((LeaseUtilizationAwsMetricStrategy)s).Settings.SamplePeriodSec);
+			Assert.AreEqual("myNs", ((LeaseUtilizationAwsMetricStrategy)s).Settings.CloudWatchNamespace);
+		}
+		
 		[TestMethod]
 		public async Task CreateStrategyWithExtraAgentCount()
 		{
diff --git a/Engine/Source/Programs/Horde/Horde.Server.Tests/Fleet/LeaseUtilizationAwsMetricStrategyTests.cs b/Engine/Source/Programs/Horde/Horde.Server.Tests/Fleet/LeaseUtilizationAwsMetricStrategyTests.cs
new file mode 100644
--- /dev/null
+++ b/Engine/Source/Programs/Horde/Horde.Server.Tests/Fleet/LeaseUtilizationAwsMetricStrategyTests.cs
@@ -0,0 +1,186 @@
+// Copyright Epic Games, Inc. All Rights Reserved.
+
+using System;
+using System.Collections.Generic;
+using System.Threading.Tasks;
+using Amazon.CloudWatch.Model;
+using Google.Protobuf;
+using Google.Protobuf.WellKnownTypes;
+using Horde.Server.Agents;
+using Horde.Server.Agents.Fleet;
+using Horde.Server.Agents.Leases;
+using Horde.Server.Streams;
+using Horde.Server.Agents.Pools;
+using Horde.Server.Utilities;
+using HordeCommon;
+using HordeCommon.Rpc.Tasks;
+using Microsoft.VisualStudio.TestTools.UnitTesting;
+using EpicGames.Horde.Api;
+using Microsoft.Extensions.Logging.Abstractions;
+
+namespace Horde.Server.Tests.Fleet
+{
+	[TestClass]
+	public class LeaseUtilizationAwsMetricStrategyTest : TestSetup
+	{
+		private readonly IPool _pool;
+		private readonly LeaseUtilizationAwsMetricSettings _settings = new (60, "myNamespace");
+		private readonly FakeAmazonCloudWatch _cloudWatch = new();
+		private readonly LeaseUtilizationAwsMetricStrategy _strategy;
+		private readonly NullLogger<LeaseUtilizationAwsMetricStrategy> _logger = NullLogger<LeaseUtilizationAwsMetricStrategy>.Instance;
+
+		public LeaseUtilizationAwsMetricStrategyTest()
+		{
+			_pool = PoolService.CreatePoolAsync("my-pool", new() { EnableAutoscaling = true, MinAgents = 0, NumReserveAgents = 0 }).Result;
+			_strategy = new (LeaseCollection, _cloudWatch.Get(), _settings, Clock, _logger);
+		}
+
+		private async Task<List<IAgent>> CreateAgentsAsync(IPool pool, int numAgents)
+		{
+			List<IAgent> agents = new (numAgents);
+			for (int i = 0; i < numAgents; i++)
+			{
+				agents.Add(await CreateAgentAsync(pool));
+			}
+
+			return agents;
+		}
+		
+		[TestMethod]
+		public async Task Metadata()
+		{
+			// Arrange
+			List<IAgent> agents = await CreateAgentsAsync(_pool, 4);
+
+			// Act
+			await _strategy.CalculatePoolSizeAsync(_pool, agents);
+
+			// Assert
+			List<MetricDatum> datums = _cloudWatch.GetMetricData("myNamespace");
+			Assert.AreEqual(1, datums.Count);
+			Assert.AreEqual("LeaseUtilizationPercentage", datums[0].MetricName);
+			Assert.AreEqual(1, datums[0].Dimensions.Count);
+			Assert.AreEqual("Pool", datums[0].Dimensions[0].Name);
+			Assert.AreEqual("my-pool", datums[0].Dimensions[0].Value);
+		}
+		
+		[TestMethod]
+		public async Task UtilizationZero()
+		{
+			// Arrange
+			List<IAgent> agents = await CreateAgentsAsync(_pool, 4);
+
+			// Act
+			await _strategy.CalculatePoolSizeAsync(_pool, agents);
+
+			// Assert
+			Assert.AreEqual(0.0, _cloudWatch.GetMetricData("myNamespace")[0].Value, 0.0001);
+		}
+		
+		[TestMethod]
+		public async Task UtilizationHalf()
+		{
+			// Arrange
+			List<IAgent> agents = await CreateAgentsAsync(_pool, 4);
+			await AddPlaceholderLease(agents[0], _pool, Clock.UtcNow - TimeSpan.FromMinutes(120), TimeSpan.FromMinutes(120));
+			await AddPlaceholderLease(agents[1], _pool, Clock.UtcNow - TimeSpan.FromMinutes(120), TimeSpan.FromMinutes(120));
+
+			// Act
+			await _strategy.CalculatePoolSizeAsync(_pool, agents);
+
+			// Assert
+			Assert.AreEqual(0.5, _cloudWatch.GetMetricData("myNamespace")[0].Value, 0.0001);
+		}
+		
+		[TestMethod]
+		public async Task UtilizationFull()
+		{
+			// Arrange
+			List<IAgent> agents = await CreateAgentsAsync(_pool, 4);
+			await AddPlaceholderLease(agents[0], _pool, Clock.UtcNow - TimeSpan.FromMinutes(120), TimeSpan.FromMinutes(120));
+			await AddPlaceholderLease(agents[1], _pool, Clock.UtcNow - TimeSpan.FromMinutes(120), TimeSpan.FromMinutes(120));
+			await AddPlaceholderLease(agents[2], _pool, Clock.UtcNow - TimeSpan.FromMinutes(120), TimeSpan.FromMinutes(120));
+			await AddPlaceholderLease(agents[3], _pool, Clock.UtcNow - TimeSpan.FromMinutes(120), TimeSpan.FromMinutes(120));
+
+			// Act
+			await _strategy.CalculatePoolSizeAsync(_pool, agents);
+
+			// Assert
+			Assert.AreEqual(1.0, _cloudWatch.GetMetricData("myNamespace")[0].Value, 0.0001);
+		}
+		
+		[TestMethod]
+		public async Task MoreLeasesThanAgents()
+		{
+			// Arrange
+			List<IAgent> agents = await CreateAgentsAsync(_pool, 4);
+			await AddPlaceholderLease(agents[0], _pool, OneMinuteAgo(), TimeSpan.FromSeconds(5));
+			await AddPlaceholderLease(agents[0], _pool, OneMinuteAgo(), TimeSpan.FromSeconds(10));
+			await AddPlaceholderLease(agents[0], _pool, OneMinuteAgo(), TimeSpan.FromSeconds(15));
+			await AddPlaceholderLease(agents[1], _pool, OneMinuteAgo(), TimeSpan.FromSeconds(5));
+			await AddPlaceholderLease(agents[1], _pool, OneMinuteAgo(), TimeSpan.FromSeconds(15));
+			await AddPlaceholderLease(agents[2], _pool, OneMinuteAgo(), TimeSpan.FromSeconds(15));
+			await AddPlaceholderLease(agents[3], _pool, OneMinuteAgo(), TimeSpan.FromSeconds(15));
+
+			// Act
+			await _strategy.CalculatePoolSizeAsync(_pool, agents);
+
+			// Assert
+			Assert.AreEqual(1.0, _cloudWatch.GetMetricData("myNamespace")[0].Value, 0.0001);
+		}
+		
+		[TestMethod]
+		public async Task LeasesInProgress()
+		{
+			// Arrange
+			List<IAgent> agents = await CreateAgentsAsync(_pool, 4);
+			await AddPlaceholderLease(agents[0], _pool, OneMinuteAgo(), null);
+			await AddPlaceholderLease(agents[1], _pool, OneMinuteAgo(), null);
+			await AddPlaceholderLease(agents[2], _pool, OneMinuteAgo(), TimeSpan.FromSeconds(15));
+
+			// Act
+			await _strategy.CalculatePoolSizeAsync(_pool, agents);
+
+			// Assert
+			Assert.AreEqual(0.75, _cloudWatch.GetMetricData("myNamespace")[0].Value, 0.0001);
+		}
+
+		private DateTime OneMinuteAgo()
+		{
+			return Clock.UtcNow - TimeSpan.FromSeconds(60);
+		}
+
+		enum LeaseType
+		{
+			ExecuteJob,
+			Compute
+		}
+		
+		private async Task<ILease> AddPlaceholderLease(IAgent agent, IPool pool, DateTime startTime, TimeSpan? duration, LeaseType leaseType = LeaseType.ExecuteJob)
+		{
+			Assert.IsNotNull(agent.SessionId);
+
+			byte[] payload;
+			ExecuteJobTask executeJobTask = new();
+			executeJobTask.JobName = "placeholderJobName";
+			payload = Any.Pack(executeJobTask).ToByteArray();
+			
+			if (leaseType == LeaseType.Compute)
+			{
+				ComputeTask computeTask = new();
+				computeTask.Nonce = ByteString.CopyFromUtf8("test-nonce");
+				computeTask.Key = ByteString.CopyFromUtf8("test-key");
+				payload = Any.Pack(computeTask).ToByteArray();
+			}
+
+			ILease lease = await LeaseCollection.AddAsync(new LeaseId(BinaryIdUtils.CreateNew()), null, "placeholderLease", agent.Id, agent.SessionId!.Value, new StreamId("placeholderStream"), pool.Id, null, startTime, payload);
+			if (duration != null)
+			{
+				bool wasModified = await LeaseCollection.TrySetOutcomeAsync(lease.Id, startTime + duration.Value, LeaseOutcome.Success, null);
+				Assert.IsTrue(wasModified);	
+			}
+			
+			return lease;
+		}
+	}
+}
\ No newline at end of file
diff --git a/Engine/Source/Programs/Horde/Horde.Server.Tests/TestSetup.cs b/Engine/Source/Programs/Horde/Horde.Server.Tests/TestSetup.cs
--- a/Engine/Source/Programs/Horde/Horde.Server.Tests/TestSetup.cs
+++ b/Engine/Source/Programs/Horde/Horde.Server.Tests/TestSetup.cs
@@ -8,6 +8,7 @@ using System.Security.Claims;
 using System.Text.Json;
 using System.Threading.Tasks;
 using Amazon.AutoScaling;
+using Amazon.CloudWatch;
 using Amazon.EC2;
 using Datadog.Trace;
 using EpicGames.Core;
@@ -237,6 +238,7 @@ namespace Horde.Server.Tests
 			// Empty mocked object to satisfy basic test runs
 			services.AddSingleton<IAmazonEC2>(sp => new Mock<IAmazonEC2>().Object);
 			services.AddSingleton<IAmazonAutoScaling>(sp => new Mock<IAmazonAutoScaling>().Object);
+			services.AddSingleton<IAmazonCloudWatch>(sp => new Mock<IAmazonCloudWatch>().Object);
 			services.AddSingleton<IFleetManagerFactory, FleetManagerFactory>();
 
 			services.AddSingleton<AclService>();
diff --git a/Engine/Source/Programs/Horde/Horde.Server/Agents/Fleet/ComputeQueueAwsMetricStrategy.cs b/Engine/Source/Programs/Horde/Horde.Server/Agents/Fleet/ComputeQueueAwsMetricStrategy.cs
--- a/Engine/Source/Programs/Horde/Horde.Server/Agents/Fleet/ComputeQueueAwsMetricStrategy.cs
+++ b/Engine/Source/Programs/Horde/Horde.Server/Agents/Fleet/ComputeQueueAwsMetricStrategy.cs
@@ -4,6 +4,7 @@ using System;
 using System.Collections.Generic;
 using System.Linq;
 using System.Net;
+using System.Text.Json.Serialization;
 using System.Threading.Tasks;
 using Amazon.CloudWatch;
 using Amazon.CloudWatch.Model;
@@ -24,19 +25,27 @@ namespace Horde.Server.Agents.Fleet
 		/// <summary>
 		/// Compute cluster ID to observe
 		/// </summary>
-		public string ComputeClusterId { get; set; }
-		
+		public string ComputeClusterId { get; set; } = "default";
+
 		/// <summary>
 		/// AWS CloudWatch namespace to write metrics in
 		/// </summary>
-		public string Namespace { get; set; }
+		public string Namespace { get; set; } = "Horde";
+		
+		/// <summary>
+		/// Constructor used for JSON serialization
+		/// </summary>
+		[JsonConstructor]
+		public ComputeQueueAwsMetricSettings()
+		{
+		}
 		
 		/// <summary>
 		/// Constructor
 		/// </summary>
 		/// <param name="computeClusterId"></param>
 		/// <param name="cloudWatchNamespace"></param>
-		public ComputeQueueAwsMetricSettings(string computeClusterId, string cloudWatchNamespace = "HordeBuild")
+		public ComputeQueueAwsMetricSettings(string computeClusterId, string cloudWatchNamespace = "Horde")
 		{
 			ComputeClusterId = computeClusterId;
 			Namespace = cloudWatchNamespace;
diff --git a/Engine/Source/Programs/Horde/Horde.Server/Agents/Fleet/FleetService.cs b/Engine/Source/Programs/Horde/Horde.Server/Agents/Fleet/FleetService.cs
--- a/Engine/Source/Programs/Horde/Horde.Server/Agents/Fleet/FleetService.cs
+++ b/Engine/Source/Programs/Horde/Horde.Server/Agents/Fleet/FleetService.cs
@@ -1,7 +1,6 @@
 // Copyright Epic Games, Inc. All Rights Reserved.
 
 using System;
-using System.Collections.Concurrent;
 using System.Collections.Generic;
 using System.Diagnostics;
 using System.Diagnostics.Metrics;
@@ -9,6 +8,7 @@ using System.Linq;
 using System.Text.Json;
 using System.Threading;
 using System.Threading.Tasks;
+using Amazon.Runtime;
 using Horde.Server.Agents.Leases;
 using Horde.Server.Agents.Pools;
 using Horde.Server.Jobs;
@@ -18,6 +18,7 @@ using Horde.Server.Streams;
 using Horde.Server.Utilities;
 using HordeCommon;
 using Microsoft.Extensions.Caching.Memory;
+using Microsoft.Extensions.DependencyInjection;
 using Microsoft.Extensions.Hosting;
 using Microsoft.Extensions.Logging;
 using Microsoft.Extensions.Options;
@@ -79,6 +80,7 @@ namespace Horde.Server.Agents.Fleet
 		private readonly TimeSpan _defaultScaleInCooldown;
 		private readonly IOptions<ServerSettings> _settings;
 		private readonly IOptionsMonitor<GlobalConfig> _globalConfig;
+		private readonly IServiceProvider _provider;
 		private readonly Tracer _tracer;
 		private readonly ILogger<FleetService> _logger;
 		
@@ -99,6 +101,7 @@ namespace Horde.Server.Agents.Fleet
 			IMemoryCache cache,
 			IOptions<ServerSettings> settings,
 			IOptionsMonitor<GlobalConfig> globalConfig,
+			IServiceProvider provider,
 			Tracer tracer,
 			ILogger<FleetService> logger)
 		{
@@ -119,6 +122,7 @@ namespace Horde.Server.Agents.Fleet
 			_ticker = clock.AddSharedTicker<FleetService>(TimeSpan.FromSeconds(30), TickLeaderAsync, _logger);
 			_tickerHighFrequency = clock.AddSharedTicker("FleetService.TickHighFrequency", TimeSpan.FromSeconds(30), TickHighFrequencyAsync, _logger);
 			_settings = settings;
+			_provider = provider;
 			_defaultScaleOutCooldown = TimeSpan.FromSeconds(settings.Value.AgentPoolScaleOutCooldownSeconds);
 			_defaultScaleInCooldown = TimeSpan.FromSeconds(settings.Value.AgentPoolScaleInCooldownSeconds);
 		}
@@ -380,11 +384,14 @@ namespace Horde.Server.Agents.Fleet
 								LeaseUtilizationSettings luSettings = DeserializeConfig<LeaseUtilizationSettings>(info.Config);
 								LeaseUtilizationStrategy luStrategy = new (_agentCollection, _poolCollection, _leaseCollection, _clock, _cache, luSettings);
 								return info.ExtraAgentCount != 0 ? new ExtraAgentCountStrategy(luStrategy, info.ExtraAgentCount) : luStrategy;
+
+							case PoolSizeStrategy.ComputeQueueAwsMetric:
+								ComputeQueueAwsMetricSettings cqamSettings = DeserializeConfig<ComputeQueueAwsMetricSettings>(info.Config);
+								return ActivatorUtilities.CreateInstance<ComputeQueueAwsMetricStrategy>(_provider, cqamSettings);
 							
-							// Disabled until moved to a separate factory class as FleetService should not contain AWS-specific classes
-							// case PoolSizeStrategy.ComputeQueueAwsMetric:
-							// 	ComputeQueueAwsMetricSettings cqamSettings = DeserializeConfig<ComputeQueueAwsMetricSettings>(info.Config);
-							// 	return new ComputeQueueAwsMetricStrategy(_awsCloudWatch, _computeService, cqamSettings);
+							case PoolSizeStrategy.LeaseUtilizationAwsMetric:
+								LeaseUtilizationAwsMetricSettings luamSettings = DeserializeConfig<LeaseUtilizationAwsMetricSettings>(info.Config);
+								return ActivatorUtilities.CreateInstance<LeaseUtilizationAwsMetricStrategy>(_provider, luamSettings);
 							
 							case PoolSizeStrategy.NoOp:
 								NoOpPoolSizeStrategy noStrategy = new ();
@@ -420,10 +427,6 @@ namespace Horde.Server.Agents.Fleet
 			}
 		}
 
-		/// <summary>
-		/// </summary>
-		
-		
 		/// <summary>
 		/// Cancel a number of pending agent shutdowns for a pool
 		/// 
diff --git a/Engine/Source/Programs/Horde/Horde.Server/Agents/Fleet/IPoolSizeStrategy.cs b/Engine/Source/Programs/Horde/Horde.Server/Agents/Fleet/IPoolSizeStrategy.cs
--- a/Engine/Source/Programs/Horde/Horde.Server/Agents/Fleet/IPoolSizeStrategy.cs
+++ b/Engine/Source/Programs/Horde/Horde.Server/Agents/Fleet/IPoolSizeStrategy.cs
@@ -33,7 +33,13 @@ namespace Horde.Server.Agents.Fleet
 		/// A no-op strategy that reports metrics to let an external AWS auto-scaling policy scale the fleet
 		/// <see cref="ComputeQueueAwsMetric"/> 
 		/// </summary>
-		ComputeQueueAwsMetric
+		ComputeQueueAwsMetric,
+		
+		/// <summary>
+		/// A no-op strategy that reports metrics to let an external AWS auto-scaling policy scale the fleet
+		/// <see cref="LeaseUtilizationAwsMetric"/> 
+		/// </summary>
+		LeaseUtilizationAwsMetric
 	}
 	
 	/// <summary>
diff --git a/Engine/Source/Programs/Horde/Horde.Server/Agents/Fleet/LeaseUtilizationAwsMetricStrategy.cs b/Engine/Source/Programs/Horde/Horde.Server/Agents/Fleet/LeaseUtilizationAwsMetricStrategy.cs
new file mode 100644
--- /dev/null
+++ b/Engine/Source/Programs/Horde/Horde.Server/Agents/Fleet/LeaseUtilizationAwsMetricStrategy.cs
@@ -0,0 +1,143 @@
+// Copyright Epic Games, Inc. All Rights Reserved.
+
+using System;
+using System.Collections.Generic;
+using System.Linq;
+using System.Net;
+using System.Text.Json.Serialization;
+using System.Threading.Tasks;
+using Amazon.CloudWatch;
+using Amazon.CloudWatch.Model;
+using Horde.Server.Agents.Leases;
+using Horde.Server.Agents.Pools;
+using Horde.Server.Utilities;
+using HordeCommon;
+using Microsoft.Extensions.Logging;
+using OpenTelemetry.Trace;
+
+namespace Horde.Server.Agents.Fleet;
+
+/// <summary>
+/// Settings for <see cref="LeaseUtilizationAwsMetricStrategy" />
+/// </summary>
+public class LeaseUtilizationAwsMetricSettings
+{
+	/// <summary>
+	/// Time window backwards in time to sample for active leases
+	/// A value of 60 sec, will query for active leases the last minute. 
+	/// </summary>
+	public int SamplePeriodSec { get; set; } = 60;
+
+	/// <summary>
+	/// AWS CloudWatch namespace to write metrics in
+	/// </summary>
+	public string CloudWatchNamespace { get; set; } = "Horde";
+	
+	/// <summary>
+	/// Constructor used for JSON serialization
+	/// </summary>
+	[JsonConstructor]
+	public LeaseUtilizationAwsMetricSettings()
+	{
+	}
+	
+	/// <summary>
+	/// Constructor
+	/// </summary>
+	/// <param name="samplePeriodSec"></param>
+	/// <param name="cloudWatchNamespace"></param>
+	public LeaseUtilizationAwsMetricSettings(int samplePeriodSec, string cloudWatchNamespace)
+	{
+		SamplePeriodSec = samplePeriodSec;
+		CloudWatchNamespace = cloudWatchNamespace;
+	}
+}
+
+/// <summary>
+/// A no-op strategy that reports lease utilization for a given pool as AWS CloudWatch metrics
+/// A metric that later can be used as a source for AWS-controlled auto-scaling policies. 
+/// </summary>
+public class LeaseUtilizationAwsMetricStrategy : IPoolSizeStrategy
+{
+	internal LeaseUtilizationAwsMetricSettings Settings { get; }
+	
+	private readonly ILeaseCollection _leaseCollection;
+	private readonly IAmazonCloudWatch _cloudWatch;
+	private readonly IClock _clock;
+	private readonly ILogger<LeaseUtilizationAwsMetricStrategy> _logger;
+
+	/// <summary>
+	/// Constructor
+	/// </summary>
+	/// <param name="leaseCollection"></param>
+	/// <param name="cloudWatch"></param>
+	/// <param name="settings"></param>
+	/// <param name="clock"></param>
+	/// <param name="logger"></param>
+	public LeaseUtilizationAwsMetricStrategy(ILeaseCollection leaseCollection, IAmazonCloudWatch cloudWatch, LeaseUtilizationAwsMetricSettings settings, IClock clock, ILogger<LeaseUtilizationAwsMetricStrategy> logger)
+	{
+		_leaseCollection = leaseCollection;
+		_cloudWatch = cloudWatch;
+		Settings = settings;
+		_clock = clock;
+		_logger = logger;
+	}
+
+	/// <inheritdoc/>
+	public string Name { get; } = "LeaseUtilizationAwsMetric";
+
+	/// <inheritdoc/>
+	public async Task<PoolSizeResult> CalculatePoolSizeAsync(IPool pool, List<IAgent> agents)
+	{
+		using TelemetrySpan span = OpenTelemetryTracers.Horde.StartActiveSpan($"{nameof(LeaseUtilizationAwsMetricStrategy)}.{nameof(CalculatePoolSizeAsync)}");
+		span.SetAttribute(OpenTelemetryTracers.DatadogResourceAttribute, pool.Id.ToString());
+		span.SetAttribute("currentAgentCount", agents.Count);
+		span.SetAttribute("cloudWatchNs", Settings.CloudWatchNamespace);
+		span.SetAttribute("samplePeriodSec", Settings.SamplePeriodSec);
+
+		List<ILease> leases = await _leaseCollection.FindLeasesAsync(minTime: _clock.UtcNow - TimeSpan.FromSeconds(Settings.SamplePeriodSec));
+		leases = leases.Where(x => x.PoolId == pool.Id).ToList();
+
+		DateTime now = _clock.UtcNow;
+		double leaseUtilization = leases.Count / (double)agents.Count;
+
+		// Clamp utilization as agents cannot be more utilized than 100%
+		leaseUtilization = Math.Min(leaseUtilization, 1.0);
+
+		List<Dimension> dimensions = new() { new() { Name = "Pool", Value = pool.Id.ToString() } };
+		Dictionary<string, List<MetricDatum>> metricsPerCloudWatchNamespace = new();
+		
+		if (!metricsPerCloudWatchNamespace.TryGetValue(Settings.CloudWatchNamespace, out List<MetricDatum>? metricDatums))
+		{
+			metricDatums = new List<MetricDatum>();
+			metricsPerCloudWatchNamespace[Settings.CloudWatchNamespace] = metricDatums;
+		}
+		
+		metricDatums.Add(new()
+		{
+			MetricName = "LeaseUtilizationPercentage",
+			Dimensions = dimensions,
+			Unit = StandardUnit.Percent,
+			Value = leaseUtilization,
+			TimestampUtc = now
+		});
+		
+		foreach ((string ns, List<MetricDatum> metricDatumsNs) in metricsPerCloudWatchNamespace)
+		{
+			using TelemetrySpan cwSpan = OpenTelemetryTracers.Horde.StartActiveSpan($"{nameof(ComputeQueueAwsMetricStrategy)}.{nameof(CalculatePoolSizeAsync)}.PutCloudWatchMetrics");
+			cwSpan.SetAttribute("namespace", ns);
+		
+			PutMetricDataRequest request = new() { Namespace = ns, MetricData = metricDatumsNs };
+			PutMetricDataResponse response = await _cloudWatch.PutMetricDataAsync(request);
+			cwSpan.SetAttribute("res.statusCode", (int)response.HttpStatusCode);
+			if (response.HttpStatusCode != HttpStatusCode.OK)
+			{
+				_logger.LogError("Unable to put CloudWatch metrics");
+			}
+		}
+
+		// Return the input pool size data as-is since we are only observing and not modifying
+		return new PoolSizeResult(agents.Count, agents.Count);
+	}
+}
+
