Horde: Better matching for UAT/UBT compile errors at startup.

- Messages saying "XXX failed to compile" from RunUAT.bat/RunUBT.bat are now grouped with regular compile erorrs, and hidden if more contextual messages are available.
- Summary saying "N warning(s)" and "N error(s)" are hidden by default.

diff --git a/Engine/Source/Programs/Horde/Horde.Agent.Tests/LogParserTests.cs b/Engine/Source/Programs/Horde/Horde.Agent.Tests/LogParserTests.cs
--- a/Engine/Source/Programs/Horde/Horde.Agent.Tests/LogParserTests.cs
+++ b/Engine/Source/Programs/Horde/Horde.Agent.Tests/LogParserTests.cs
@@ -302,6 +302,41 @@ namespace Horde.Agent.Tests
 			}
 		}
 
+		[TestMethod]
+		public void CsCompileEventMatcher4()
+		{
+			// C# compile error from UBT startup
+			{
+				string[] lines =
+				{
+					@"Building AutomationTool...",
+					@"Microsoft (R) Build Engine version 17.2.0+41abc5629 for .NET",
+					@"Copyright (C) Microsoft Corporation. All rights reserved.",
+					@"Engine/Source/Programs/Shared/EpicGames.Horde/Storage/Workspace.cs(405,11): error CS1501: No overload for method 'ReadList' takes 2 arguments [Engine/Source/Programs/Shared/EpicGames.Horde/EpicGames.Horde.csproj]",
+					@"Engine/Source/Programs/Shared/EpicGames.Horde/Storage/Workspace.cs(423,12): error CS1501: No overload for method 'ReadList' takes 2 arguments [Engine/Source/Programs/Shared/EpicGames.Horde/EpicGames.Horde.csproj]",
+					@"Build FAILED.",
+					@"Engine/Source/Programs/Shared/EpicGames.Horde/Storage/Workspace.cs(405,11): error CS1501: No overload for method 'ReadList' takes 2 arguments [Engine/Source/Programs/Shared/EpicGames.Horde/EpicGames.Horde.csproj]",
+					@"Engine/Source/Programs/Shared/EpicGames.Horde/Storage/Workspace.cs(423,12): error CS1501: No overload for method 'ReadList' takes 2 arguments [Engine/Source/Programs/Shared/EpicGames.Horde/EpicGames.Horde.csproj]",
+					@"    0 Warning(s)",
+					@"    2 Error(s)",
+					@"Time Elapsed 00:00:06.84",
+					@"RunUBT ERROR: UnrealBuildTool failed to compile.",
+					@"RunUAT.bat ERROR: AutomationTool failed to compile.",
+					@"BUILD FAILED",
+				};
+
+				List<LogEvent> logEvents = Parse(String.Join("\n", lines));
+				Assert.AreEqual(6, logEvents.Count);
+
+				CheckEventGroup(logEvents.Slice(0, 1), 3, 1, LogLevel.Error, KnownLogEvents.Compiler);
+				CheckEventGroup(logEvents.Slice(1, 1), 4, 1, LogLevel.Error, KnownLogEvents.Compiler);
+				CheckEventGroup(logEvents.Slice(2, 1), 6, 1, LogLevel.Error, KnownLogEvents.Compiler);
+				CheckEventGroup(logEvents.Slice(3, 1), 7, 1, LogLevel.Error, KnownLogEvents.Compiler);
+				CheckEventGroup(logEvents.Slice(4, 1), 11, 1, LogLevel.Error, KnownLogEvents.Compiler_Summary);
+				CheckEventGroup(logEvents.Slice(5, 1), 12, 1, LogLevel.Error, KnownLogEvents.Compiler_Summary);
+			}
+		}
+
 		[TestMethod]
 		public void HttpEventMatcher()
 		{
diff --git a/Engine/Source/Programs/Horde/Horde.Server/Issues/Handlers/CompileIssueHandler.cs b/Engine/Source/Programs/Horde/Horde.Server/Issues/Handlers/CompileIssueHandler.cs
--- a/Engine/Source/Programs/Horde/Horde.Server/Issues/Handlers/CompileIssueHandler.cs
+++ b/Engine/Source/Programs/Horde/Horde.Server/Issues/Handlers/CompileIssueHandler.cs
@@ -42,7 +42,7 @@ namespace Horde.Server.Issues.Handlers
 			return eventId == KnownLogEvents.Compiler || eventId == KnownLogEvents.UHT || eventId == KnownLogEvents.AutomationTool_SourceFileLine || eventId == KnownLogEvents.MSBuild;
 		}
 
-		static bool IsMaskedEventId(EventId id) => id == KnownLogEvents.ExitCode || id == KnownLogEvents.Systemic_Xge_BuildFailed;
+		static bool IsMaskedEventId(EventId id) => id == KnownLogEvents.ExitCode || id == KnownLogEvents.Systemic_Xge_BuildFailed || id == KnownLogEvents.Compiler_Summary;
 
 		/// <inheritdoc/>
 		public override void TagEvents(IJob job, INode node, IReadOnlyNodeAnnotations annotations, IReadOnlyList<IssueEvent> stepEvents)
diff --git a/Engine/Source/Programs/Shared/EpicGames.Core/KnownLogEvents.cs b/Engine/Source/Programs/Shared/EpicGames.Core/KnownLogEvents.cs
--- a/Engine/Source/Programs/Shared/EpicGames.Core/KnownLogEvents.cs
+++ b/Engine/Source/Programs/Shared/EpicGames.Core/KnownLogEvents.cs
@@ -35,6 +35,11 @@ namespace EpicGames.Core
 		/// </summary>
 		public static EventId Compiler { get; } = new EventId(100);
 
+		/// <summary>
+		/// Compiler summary line (eg. 'XYZ failed to compile')
+		/// </summary>
+		public static EventId Compiler_Summary { get; } = new EventId(101);
+
 		/// <summary>
 		/// UHT compiler message
 		/// </summary>
diff --git a/Engine/Source/Programs/UnrealBuildTool/Matchers/CompileEventMatcher.cs b/Engine/Source/Programs/UnrealBuildTool/Matchers/CompileEventMatcher.cs
--- a/Engine/Source/Programs/UnrealBuildTool/Matchers/CompileEventMatcher.cs
+++ b/Engine/Source/Programs/UnrealBuildTool/Matchers/CompileEventMatcher.cs
@@ -63,7 +63,9 @@ namespace UnrealBuildTool.Matchers
 		static readonly Regex s_clangNotePattern = new Regex($"^\\s*{FilePattern}\\s*{ClangLocationPattern}:\\s*note:");
 		static readonly Regex s_clangMarkerPattern = new Regex(@"^(\s*)[\^~][\s\^~]*$");
 		static readonly Regex s_xcodeIDEWatchExtensionPattern = new Regex(@"xcodebuild.*Requested but did not find extension point with identifier.*for extension.*\.watchOS of plug-in com\.apple\.dt\.IDEWatchSupportCore");
-
+		static readonly Regex s_scriptCompilePattern = new Regex(@"^\s*[A-Za-z0-9_\.]+ ERROR:.* [A-Za-z_]+ failed to compile\.");
+		static readonly Regex s_cscSummaryPattern = new Regex(@"^\s+\d+ (?:Warning|Error)\(s\)");
+	
 		static readonly string[] s_invalidExtensions =
 		{
 			".obj",
@@ -179,6 +181,11 @@ namespace UnrealBuildTool.Matchers
 				LogEventBuilder builder = new LogEventBuilder(input);
 				return builder.ToMatch(LogEventPriority.Normal, LogLevel.Information, KnownLogEvents.Systemic_XCode);
 			}
+			else if (input.TryMatch(s_scriptCompilePattern, out match))
+			{
+				LogEventBuilder builder = new LogEventBuilder(input);
+				return builder.ToMatch(LogEventPriority.High, LogLevel.Error, KnownLogEvents.Compiler_Summary);
+			}
 			return null;
 		}
 
@@ -260,7 +267,7 @@ namespace UnrealBuildTool.Matchers
 
 			for (; ; )
 			{
-				while (builder.Current.StartsWith(1, nextIndent))
+				while (builder.Current.StartsWith(1, nextIndent) && !builder.Current.IsMatch(1, s_cscSummaryPattern))
 				{
 					builder.MoveNext();
 				}
