Horde: Improve trace spans for lease utilization metric reporting

diff --git a/Engine/Source/Programs/Horde/Horde.Server/Agents/Fleet/LeaseUtilizationAwsMetricStrategy.cs b/Engine/Source/Programs/Horde/Horde.Server/Agents/Fleet/LeaseUtilizationAwsMetricStrategy.cs
--- a/Engine/Source/Programs/Horde/Horde.Server/Agents/Fleet/LeaseUtilizationAwsMetricStrategy.cs
+++ b/Engine/Source/Programs/Horde/Horde.Server/Agents/Fleet/LeaseUtilizationAwsMetricStrategy.cs
@@ -91,9 +91,6 @@ public class LeaseUtilizationAwsMetricStrategy : IPoolSizeStrategy
 	{
 		using TelemetrySpan span = OpenTelemetryTracers.Horde.StartActiveSpan($"{nameof(LeaseUtilizationAwsMetricStrategy)}.{nameof(CalculatePoolSizeAsync)}");
 		span.SetAttribute(OpenTelemetryTracers.DatadogResourceAttribute, pool.Id.ToString());
-		span.SetAttribute("currentAgentCount", agents.Count);
-		span.SetAttribute("cloudWatchNs", Settings.CloudWatchNamespace);
-		span.SetAttribute("samplePeriodSec", Settings.SamplePeriodSec);
 
 		List<ILease> leases = await _leaseCollection.FindLeasesAsync(minTime: _clock.UtcNow - TimeSpan.FromSeconds(Settings.SamplePeriodSec));
 		leases = leases.Where(x => x.PoolId == pool.Id).ToList();
@@ -103,37 +100,32 @@ public class LeaseUtilizationAwsMetricStrategy : IPoolSizeStrategy
 
 		// Clamp utilization as agents cannot be more utilized than 100%
 		leaseUtilization = Math.Min(leaseUtilization, 1.0);
+		
+		span.SetAttribute("cloudWatchNs", Settings.CloudWatchNamespace);
+		span.SetAttribute("samplePeriodSec", Settings.SamplePeriodSec);
+		span.SetAttribute("agentCount", agents.Count);
+		span.SetAttribute("leaseCount", leases.Count);
+		span.SetAttribute("leaseUtilizationPct", leaseUtilization);
 
 		List<Dimension> dimensions = new() { new() { Name = "Pool", Value = pool.Id.ToString() } };
-		Dictionary<string, List<MetricDatum>> metricsPerCloudWatchNamespace = new();
-		
-		if (!metricsPerCloudWatchNamespace.TryGetValue(Settings.CloudWatchNamespace, out List<MetricDatum>? metricDatums))
+		List<MetricDatum> metricDatums = new List<MetricDatum>()
 		{
-			metricDatums = new List<MetricDatum>();
-			metricsPerCloudWatchNamespace[Settings.CloudWatchNamespace] = metricDatums;
-		}
-		
-		metricDatums.Add(new()
-		{
-			MetricName = "LeaseUtilizationPercentage",
-			Dimensions = dimensions,
-			Unit = StandardUnit.Percent,
-			Value = leaseUtilization,
-			TimestampUtc = now
-		});
-		
-		foreach ((string ns, List<MetricDatum> metricDatumsNs) in metricsPerCloudWatchNamespace)
-		{
-			using TelemetrySpan cwSpan = OpenTelemetryTracers.Horde.StartActiveSpan($"{nameof(ComputeQueueAwsMetricStrategy)}.{nameof(CalculatePoolSizeAsync)}.PutCloudWatchMetrics");
-			cwSpan.SetAttribute("namespace", ns);
-		
-			PutMetricDataRequest request = new() { Namespace = ns, MetricData = metricDatumsNs };
-			PutMetricDataResponse response = await _cloudWatch.PutMetricDataAsync(request);
-			cwSpan.SetAttribute("res.statusCode", (int)response.HttpStatusCode);
-			if (response.HttpStatusCode != HttpStatusCode.OK)
+			new ()
 			{
-				_logger.LogError("Unable to put CloudWatch metrics");
+				MetricName = "LeaseUtilizationPercentage",
+				Dimensions = dimensions,
+				Unit = StandardUnit.Percent,
+				Value = leaseUtilization,
+				TimestampUtc = now
 			}
+		};
+		
+		PutMetricDataRequest request = new() { Namespace = Settings.CloudWatchNamespace, MetricData = metricDatums };
+		PutMetricDataResponse response = await _cloudWatch.PutMetricDataAsync(request);
+
+		if (response.HttpStatusCode != HttpStatusCode.OK)
+		{
+			_logger.LogError("Unable to put CloudWatch metrics");
 		}
 
 		// Return the input pool size data as-is since we are only observing and not modifying
